Stub of AIModule added and plugged in as default dependency of the Engine #UE4

diff --git a/Engine/Source/Runtime/AIModule/AIModule.Build.cs b/Engine/Source/Runtime/AIModule/AIModule.Build.cs
new file mode 100644
--- /dev/null
+++ b/Engine/Source/Runtime/AIModule/AIModule.Build.cs
@@ -0,0 +1,53 @@
+// Copyright 1998-2014 Epic Games, Inc. All Rights Reserved.
+
+namespace UnrealBuildTool.Rules
+{
+	public class AIModule : ModuleRules
+	{
+        public AIModule(TargetInfo Target)
+		{
+			PublicIncludePaths.AddRange(
+				new string[] {
+                    "Runtime/AIModule/Public",
+					// ... add public include paths required here ...
+				}
+				);
+
+			PrivateIncludePaths.AddRange(
+				new string[] {
+					"Runtime/AIModule/Private",
+					// ... add other private include paths required here ...
+				}
+				);
+
+			PublicDependencyModuleNames.AddRange(
+				new string[]
+				{
+					"Core",
+					"CoreUObject",
+					"Engine"
+					// ... add other public dependencies that you statically link with here ...
+				}
+				);
+
+			PrivateDependencyModuleNames.AddRange(
+				new string[]
+				{
+					// ... add private dependencies that you statically link with here ...
+				}
+				);
+
+			DynamicallyLoadedModuleNames.AddRange(
+				new string[]
+				{
+					// ... add any modules that your module loads dynamically here ...
+				}
+				);
+
+			if (UEBuildConfiguration.bBuildEditor == true)
+			{
+				PrivateDependencyModuleNames.Add("UnrealEd");
+			}
+		}
+	}
+}
\ No newline at end of file
diff --git a/Engine/Source/Runtime/AIModule/Private/AIModule.cpp b/Engine/Source/Runtime/AIModule/Private/AIModule.cpp
new file mode 100644
--- /dev/null
+++ b/Engine/Source/Runtime/AIModule/Private/AIModule.cpp
@@ -0,0 +1,24 @@
+// Copyright 1998-2014 Epic Games, Inc. All Rights Reserved.
+
+#include "AIModulePrivatePCH.h"
+
+class FAIModule : public IAIModule
+{
+	// Begin IModuleInterface
+	virtual void StartupModule() OVERRIDE;
+	virtual void ShutdownModule() OVERRIDE;
+	// End IModuleInterface
+};
+
+IMPLEMENT_MODULE(FAIModule, AIModule)
+
+void FAIModule::StartupModule()
+{ 
+	// This code will execute after your module is loaded into memory (but after global variables are initialized, of course.)
+}
+
+void FAIModule::ShutdownModule()
+{
+	// This function may be called during shutdown to clean up your module.  For modules that support dynamic reloading,
+	// we call this function before unloading the module.
+}
diff --git a/Engine/Source/Runtime/AIModule/Private/AIModulePrivatePCH.h b/Engine/Source/Runtime/AIModule/Private/AIModulePrivatePCH.h
new file mode 100644
--- /dev/null
+++ b/Engine/Source/Runtime/AIModule/Private/AIModulePrivatePCH.h
@@ -0,0 +1,20 @@
+// Copyright 1998-2014 Epic Games, Inc. All Rights Reserved.
+
+#pragma once
+
+
+
+#if WITH_EDITOR
+#include "UnrealEd.h"
+#endif
+
+
+// You should place include statements to your module's private header files here.  You only need to
+// add includes for headers that are used in most of your module's source files though.
+
+#include "Core.h"
+#include "Engine.h"
+//#include "AIModuleClasses.h"
+
+#include "AIModule.h"
+
diff --git a/Engine/Source/Runtime/AIModule/Public/AIModule.h b/Engine/Source/Runtime/AIModule/Public/AIModule.h
new file mode 100644
--- /dev/null
+++ b/Engine/Source/Runtime/AIModule/Public/AIModule.h
@@ -0,0 +1,36 @@
+// Copyright 1998-2014 Epic Games, Inc. All Rights Reserved.
+
+#pragma once
+
+#include "ModuleManager.h"
+
+/**
+ * The public interface to this module
+ */
+class IAIModule : public IModuleInterface
+{
+
+public:
+
+	/**
+	 * Singleton-like access to this module's interface.  This is just for convenience!
+	 * Beware of calling this during the shutdown phase, though.  Your module might have been unloaded already.
+	 *
+	 * @return Returns singleton instance, loading the module on demand if needed
+	 */
+	static inline IAIModule& Get()
+	{
+		return FModuleManager::LoadModuleChecked< IAIModule >( "AIModule" );
+	}
+
+	/**
+	 * Checks to see if this module is loaded and ready.  It is only valid to call Get() if IsAvailable() returns true.
+	 *
+	 * @return True if the module is loaded and ready to use
+	 */
+	static inline bool IsAvailable()
+	{
+		return FModuleManager::Get().IsModuleLoaded( "AIModule" );
+	}
+};
+
diff --git a/Engine/Source/Runtime/Engine/Engine.Build.cs b/Engine/Source/Runtime/Engine/Engine.Build.cs
--- a/Engine/Source/Runtime/Engine/Engine.Build.cs
+++ b/Engine/Source/Runtime/Engine/Engine.Build.cs
@@ -61,6 +61,7 @@ public class Engine : ModuleRules
 				"SlateCore",
 				"SlateReflector",
 				"VectorVM",
+                "AIModule",
 			}
 		);
 
@@ -89,6 +90,7 @@ public class Engine : ModuleRules
 			PrivateIncludePathModuleNames.Add("MeshUtilities");
 			DynamicallyLoadedModuleNames.Add("MeshUtilities");
 
+
 			if (Target.Configuration != UnrealTargetConfiguration.Shipping && Target.Configuration != UnrealTargetConfiguration.Test)
 			{
 				PrivateDependencyModuleNames.Add("CollisionAnalyzer");
