Extracting AI code from the engine into a separate module, pass 1 #UE4

diff --git a/Engine/Source/Runtime/AIModule/Classes/BehaviorTree/BehaviorTreeComponent.h b/Engine/Source/Runtime/AIModule/Classes/BehaviorTree/BehaviorTreeComponent.h
new file mode 100644
--- /dev/null
+++ b/Engine/Source/Runtime/AIModule/Classes/BehaviorTree/BehaviorTreeComponent.h
@@ -0,0 +1,278 @@
+// Copyright 1998-2014 Epic Games, Inc. All Rights Reserved.
+
+#pragma once
+
+#include "AITypes.h"
+#include "BrainComponent.h"
+#include "BehaviorTreeTypes.h"
+#include "BehaviorTreeComponent.generated.h"
+
+struct FBTNodeExecutionInfo
+{
+	/** index of first task allowed to be executed */
+	struct FBTNodeIndex SearchStart;
+
+	/** node to be executed */
+	class UBTCompositeNode* ExecuteNode;
+
+	/** subtree index */
+	uint16 ExecuteInstanceIdx;
+
+	/** result used for resuming execution */
+	TEnumAsByte<EBTNodeResult::Type> ContinueWithResult;
+
+	/** if set, tree will try to execute next child of composite instead of forcing branch containing SearchStart */
+	uint8 bTryNextChild : 1;
+
+	FBTNodeExecutionInfo() : ExecuteNode(NULL) { }
+};
+
+UCLASS()
+class AIMODULE_API UBehaviorTreeComponent : public UBrainComponent
+{
+	GENERATED_UCLASS_BODY()
+
+	// Begin UBrainComponent overrides
+	virtual void RestartLogic() OVERRIDE;
+protected:
+	virtual void StopLogic(const FString& Reason) OVERRIDE;
+	virtual void PauseLogic(const FString& Reason) OVERRIDE;
+	virtual EAILogicResuming::Type ResumeLogic(const FString& Reason) OVERRIDE;
+
+	/** indicates instance has been initialized to work with specific BT asset */
+	bool TreeHasBeenStarted() const;
+
+public:
+	virtual bool IsRunning() const OVERRIDE;
+	virtual bool IsPaused() const OVERRIDE;
+	// End UBrainComponent overrides
+
+	/** starts execution from root */
+	bool StartTree(class UBehaviorTree* Asset, EBTExecutionMode::Type ExecuteMode = EBTExecutionMode::Looped);
+
+	/** stops execution */
+	void StopTree();
+
+	/** restarts execution from root */
+	void RestartTree();
+
+	/** request execution change */
+	void RequestExecution(class UBTCompositeNode* RequestedOn, int32 InstanceIdx, 
+		const class UBTNode* RequestedBy, int32 RequestedByChildIndex,
+		EBTNodeResult::Type ContinueWithResult, bool bStoreForDebugger = true);
+
+	/** request execution change: helpers for decorator nodes */
+	void RequestExecution(const class UBTDecorator* RequestedBy);
+
+	/** request execution change: helpers for task nodes */
+	void RequestExecution(EBTNodeResult::Type ContinueWithResult);
+
+	/** finish latent execution or abort */
+	void OnTaskFinished(const class UBTTaskNode* TaskNode, EBTNodeResult::Type TaskResult);
+
+	/** setup message observer for given task */
+	void RegisterMessageObserver(const class UBTTaskNode* TaskNode, FName MessageType);
+	void RegisterMessageObserver(const class UBTTaskNode* TaskNode, FName MessageType, FAIRequestID MessageID);
+	
+	/** remove message observers registered with task */
+	void UnregisterMessageObserversFrom(const class UBTTaskNode* TaskNode);
+	void UnregisterMessageObserversFrom(const FBTNodeIndex& TaskIdx);
+
+	/** add active parallel task */
+	void RegisterParallelTask(const class UBTTaskNode* TaskNode);
+
+	/** remove parallel task */
+	void UnregisterParallelTask(const class UBTTaskNode* TaskNode, uint16 InstanceIdx);
+
+	/** unregister all aux nodes less important than given index */
+	void UnregisterAuxNodesUpTo(const struct FBTNodeIndex& Index);
+
+	/** BEGIN UActorComponent overrides */
+	virtual void TickComponent(float DeltaTime, enum ELevelTick TickType, FActorComponentTickFunction *ThisTickFunction) OVERRIDE;
+	/** END UActorComponent overrides */
+
+	/** process execution flow */
+	void ProcessExecutionRequest();
+
+	/** schedule execution flow update in next tick */
+	void ScheduleExecutionUpdate();
+
+	/** tries to find behavior tree instance in context */
+	int32 FindInstanceContainingNode(const class UBTNode* Node) const;
+
+	/** tries to find template node for given instanced node */
+	class UBTNode* FindTemplateNode(const class UBTNode* Node) const;
+
+	/** @return current tree */
+	class UBehaviorTree* GetCurrentTree() const;
+
+	/** @return tree from top of instance stack */
+	class UBehaviorTree* GetRootTree() const;
+
+	/** @return active node */
+	const class UBTNode* GetActiveNode() const;
+	
+	/** get index of active instance on stack */
+	uint16 GetActiveInstanceIdx() const;
+
+	/** @return node memory */
+	uint8* GetNodeMemory(class UBTNode* Node, int32 InstanceIdx) const;
+
+	/** @return true if ExecutionRequest is switching to higher priority node */
+	bool IsRestartPending() const;
+
+	/** @return true if active node is one of child nodes of given one */
+	bool IsExecutingBranch(const class UBTNode* Node, int32 ChildIndex = -1) const;
+
+	/** @return true if aux node is currently active */
+	bool IsAuxNodeActive(const class UBTAuxiliaryNode* AuxNode) const;
+
+	/** @return status of speficied task */
+	EBTTaskStatus::Type GetTaskStatus(const class UBTTaskNode* TaskNode) const;
+
+	virtual FString GetDebugInfoString() const OVERRIDE;
+	virtual FString DescribeActiveTasks() const;
+	virtual FString DescribeActiveTrees() const;
+
+#if ENABLE_VISUAL_LOG
+	virtual void DescribeSelfToVisLog(struct FVisLogEntry* Snapshot) const OVERRIDE;
+#endif
+
+protected:
+	/** stack of behavior tree instances */
+	TArray<struct FBehaviorTreeInstance> InstanceStack;
+
+	/** list of known subtree instances */
+	TArray<struct FBehaviorTreeInstanceId> KnownInstances;
+
+	/** instanced nodes */
+	UPROPERTY(transient)
+	TArray<class UBTNode*> NodeInstances;
+
+	/** search data being currently used */
+	struct FBehaviorTreeSearchData SearchData;
+
+	/** execution request, will be applied when current task finish execution/aborting */
+	FBTNodeExecutionInfo ExecutionRequest;
+
+	/** message observers mapped by instance & execution index */
+	TMultiMap<FBTNodeIndex,FAIMessageObserverHandle> TaskMessageObservers;
+
+#if USE_BEHAVIORTREE_DEBUGGER
+	/** search flow for debugger */
+	mutable TArray<TArray<FBehaviorTreeDebuggerInstance::FNodeFlowData> > CurrentSearchFlow;
+	mutable TArray<TArray<FBehaviorTreeDebuggerInstance::FNodeFlowData> > CurrentRestarts;
+	mutable TArray<FString> SearchStartBlackboard;
+	mutable TArray<FBehaviorTreeDebuggerInstance> RemovedInstances;
+
+	/** debugger's recorded data */
+	mutable TArray<FBehaviorTreeExecutionStep> DebuggerSteps;
+#endif
+
+	/** index of last active instance on stack */
+	uint16 ActiveInstanceIdx;
+
+	/** loops tree execution */
+	uint8 bLoopExecution : 1;
+
+	/** set when execution is waiting for tasks to abort (current or parallel's main) */
+	uint8 bWaitingForAbortingTasks : 1;
+
+	/** set when execution update is scheduled for next tick */
+	uint8 bRequestedFlowUpdate : 1;
+
+	/** if set, tree execution is allowed */
+	uint8 bIsRunning : 1;
+
+	/** if set, execution requests will be postponed */
+	uint8 bIsPaused : 1;
+
+	/** push behavior tree instance on execution stack
+	 *	@NOTE: should never be called out-side of BT execution, meaning only BT tasks can push another BT instance! */
+	bool PushInstance(class UBehaviorTree* TreeAsset);
+
+	/** add unique Id of newly created subtree to KnownInstances list and return its index */
+	uint8 UpdateInstanceId(class UBehaviorTree* TreeAsset, const class UBTNode* OriginNode, int32 OriginInstanceIdx);
+
+	/** find next task to execute */
+	UBTTaskNode* FindNextTask(class UBTCompositeNode* ParentNode, uint16 ParentInstanceIdx, EBTNodeResult::Type LastResult);
+
+	/** called when tree runs out of nodes to execute */
+	void OnTreeFinished();
+
+	/** apply pending node updates from SearchData */
+	void ApplySearchData(class UBTNode* NewActiveNode, int32 UpToIdx = -1);
+
+	/** apply updates from specific list */
+	void ApplySearchUpdates(const TArray<struct FBehaviorTreeSearchUpdate>& UpdateList, int32 UpToIdx, int32 NewNodeExecutionIndex, bool bPostUpdate = false);
+
+	/** abort currently executed task */
+	void AbortCurrentTask();
+
+	/** execute new task */
+	void ExecuteTask(class UBTTaskNode* TaskNode);
+
+	/** deactivate all nodes up to requested one */
+	bool DeactivateUpTo(class UBTCompositeNode* Node, uint16 NodeInstanceIdx, EBTNodeResult::Type& NodeResult);
+
+	/** update state of aborting tasks */
+	void UpdateAbortingTasks();
+
+	/** make a snapshot for debugger */
+	void StoreDebuggerExecutionStep(EBTExecutionSnap::Type SnapType);
+
+	/** make a snapshot for debugger from given subtree instance */
+	void StoreDebuggerInstance(FBehaviorTreeDebuggerInstance& InstanceInfo, uint16 InstanceIdx, EBTExecutionSnap::Type SnapType) const;
+	void StoreDebuggerRemovedInstance(uint16 InstanceIdx) const;
+
+	/** store search step for debugger */
+	void StoreDebuggerSearchStep(const class UBTNode* Node, uint16 InstanceIdx, EBTNodeResult::Type NodeResult) const;
+	void StoreDebuggerSearchStep(const class UBTNode* Node, uint16 InstanceIdx, bool bPassed) const;
+
+	/** store restarting node for debugger */
+	void StoreDebuggerRestart(const class UBTNode* Node, uint16 InstanceIdx, bool bAllowed);
+
+	/** describe blackboard's key values */
+	void StoreDebuggerBlackboard(TArray<FString>& BlackboardValueDesc) const;
+
+	/** gather nodes runtime descriptions */
+	void StoreDebuggerRuntimeValues(TArray<FString>& RuntimeDescriptions, class UBTNode* RootNode, uint16 InstanceIdx) const;
+
+	/** update runtime description of given task node in latest debugger's snapshot */
+	void UpdateDebuggerAfterExecution(const class UBTTaskNode* TaskNode, uint16 InstanceIdx) const;
+
+	friend class UBTNode;
+	friend class UBTCompositeNode;
+	friend class UBTTaskNode;
+	friend class UBTTask_RunBehavior;
+	friend class FBehaviorTreeDebugger;
+	friend struct FBehaviorTreeInstance;
+};
+
+//////////////////////////////////////////////////////////////////////////
+// Inlines
+
+FORCEINLINE class UBehaviorTree* UBehaviorTreeComponent::GetCurrentTree() const
+{
+	return InstanceStack.Num() ? KnownInstances[InstanceStack[ActiveInstanceIdx].InstanceIdIndex].TreeAsset : NULL;
+}
+
+FORCEINLINE class UBehaviorTree* UBehaviorTreeComponent::GetRootTree() const
+{
+	return InstanceStack.Num() ? KnownInstances[InstanceStack[0].InstanceIdIndex].TreeAsset : NULL;
+}
+
+FORCEINLINE const class UBTNode* UBehaviorTreeComponent::GetActiveNode() const
+{
+	return InstanceStack.Num() ? InstanceStack[ActiveInstanceIdx].ActiveNode : NULL;
+}
+
+FORCEINLINE uint16 UBehaviorTreeComponent::GetActiveInstanceIdx() const
+{
+	return ActiveInstanceIdx;
+}
+
+FORCEINLINE bool UBehaviorTreeComponent::IsRestartPending() const
+{
+	return ExecutionRequest.ExecuteNode && !ExecutionRequest.bTryNextChild;
+}
diff --git a/Engine/Source/Runtime/AIModule/Classes/BehaviorTree/BehaviorTreeManager.h b/Engine/Source/Runtime/AIModule/Classes/BehaviorTree/BehaviorTreeManager.h
new file mode 100644
--- /dev/null
+++ b/Engine/Source/Runtime/AIModule/Classes/BehaviorTree/BehaviorTreeManager.h
@@ -0,0 +1,52 @@
+// Copyright 1998-2014 Epic Games, Inc. All Rights Reserved.
+
+#pragma once
+#include "BehaviorTreeManager.generated.h"
+
+USTRUCT()
+struct FBehaviorTreeTemplateInfo
+{
+	GENERATED_USTRUCT_BODY()
+
+	/** behavior tree asset */
+	UPROPERTY()
+	class UBehaviorTree* Asset;
+
+	/** initialized template */
+	UPROPERTY(transient)
+	class UBTCompositeNode* Template;
+
+	/** size required for instance memory */
+	uint16 InstanceMemorySize;
+};
+
+UCLASS(config=Engine)
+class AIMODULE_API UBehaviorTreeManager : public UObject
+{
+	GENERATED_UCLASS_BODY()
+
+	/** limit for recording execution steps for debugger */
+	UPROPERTY(config)
+	int32 MaxDebuggerSteps;
+
+	/** get behavior tree template for given blueprint */
+	bool LoadTree(class UBehaviorTree* Asset, class UBTCompositeNode*& Root, uint16& InstanceMemorySize);
+
+	/** check if BT system is to be used. Here only for BT development time */
+	static bool IsBehaviorTreeUsageEnabled();
+
+	/** get aligned memory size */
+	static int32 GetAlignedDataSize(int32 Size);
+
+	/** helper function for sorting and aligning node memory */
+	static void InitializeMemoryHelper(const TArray<class UBTDecorator*>& Nodes, TArray<uint16>& MemoryOffsets, int32& MemorySize);
+
+	/** cleanup hooks for map loading */
+	virtual void FinishDestroy() OVERRIDE;
+
+protected:
+
+	/** initialized tree templates */
+	UPROPERTY()
+	TArray<FBehaviorTreeTemplateInfo> LoadedTemplates;
+};
diff --git a/Engine/Source/Runtime/AIModule/Classes/BehaviorTree/BehaviorTreeTypes.h b/Engine/Source/Runtime/AIModule/Classes/BehaviorTree/BehaviorTreeTypes.h
new file mode 100644
--- /dev/null
+++ b/Engine/Source/Runtime/AIModule/Classes/BehaviorTree/BehaviorTreeTypes.h
@@ -0,0 +1,420 @@
+// Copyright 1998-2014 Epic Games, Inc. All Rights Reserved.
+
+#pragma once
+#include "BehaviorTreeTypes.generated.h"
+
+// Visual logging helper
+#define BT_VLOG(Context, Verbosity, Format, ...) UE_VLOG(Context->OwnerComp.IsValid() ? Context->OwnerComp->GetOwner() : NULL, LogBehaviorTree, Verbosity, Format, ##__VA_ARGS__)
+#define BT_SEARCHLOG(SearchData, Verbosity, Format, ...) UE_VLOG(SearchData.OwnerComp->GetOwner(), LogBehaviorTree, Verbosity, Format, ##__VA_ARGS__)
+
+// Behavior tree debugger in editor
+#define USE_BEHAVIORTREE_DEBUGGER	(1 && WITH_EDITORONLY_DATA)
+
+DECLARE_STATS_GROUP(TEXT("Behavior Tree"), STATGROUP_AIBehaviorTree, STATCAT_Advanced);
+
+DECLARE_CYCLE_STAT_EXTERN(TEXT("Tick"),STAT_AI_BehaviorTree_Tick,STATGROUP_AIBehaviorTree, );
+DECLARE_CYCLE_STAT_EXTERN(TEXT("Load Time"),STAT_AI_BehaviorTree_LoadTime,STATGROUP_AIBehaviorTree, );
+DECLARE_CYCLE_STAT_EXTERN(TEXT("Search Time"),STAT_AI_BehaviorTree_SearchTime,STATGROUP_AIBehaviorTree, );
+DECLARE_CYCLE_STAT_EXTERN(TEXT("Execution Time"),STAT_AI_BehaviorTree_ExecutionTime,STATGROUP_AIBehaviorTree, );
+DECLARE_CYCLE_STAT_EXTERN(TEXT("Auxiliary Update Time"),STAT_AI_BehaviorTree_AuxUpdateTime,STATGROUP_AIBehaviorTree, );
+DECLARE_DWORD_COUNTER_STAT_EXTERN(TEXT("Num Templates"),STAT_AI_BehaviorTree_NumTemplates,STATGROUP_AIBehaviorTree, );
+DECLARE_DWORD_COUNTER_STAT_EXTERN(TEXT("Num Instances"),STAT_AI_BehaviorTree_NumInstances,STATGROUP_AIBehaviorTree, );
+DECLARE_MEMORY_STAT_EXTERN(TEXT("Instance memory"),STAT_AI_BehaviorTree_InstanceMemory,STATGROUP_AIBehaviorTree, AIMODULE_API);
+
+// delegate defines
+DECLARE_DELEGATE_TwoParams(FOnBlackboardChange, const class UBlackboardComponent*, uint8 /*key ID*/);
+
+namespace BTSpecialChild
+{
+	const int32 NotInitialized = -1;	// special value for child indices: needs to be initialized
+	const int32 ReturnToParent = -2;	// special value for child indices: return to parent node
+}
+
+UENUM(BlueprintType)
+namespace EBTNodeResult
+{
+	// keep in sync with DescribeNodeResult()
+	enum Type
+	{
+		Succeeded,		// finished as success
+		Failed,			// finished as failure
+		Optional,		// finished, parent should continue
+		Aborted,		// finished aborting = failure
+		InProgress,		// not finished yet
+	};
+}
+
+namespace EBTExecutionMode
+{
+	enum Type
+	{
+		SingleRun,
+		Looped,
+	};
+}
+
+UENUM()
+namespace EBTFlowAbortMode
+{
+	// keep in sync with DescribeFlowAbortMode()
+	enum Type
+	{
+		None				UMETA(DisplayName="Nothing"),
+		LowerPriority		UMETA(DisplayName="Lower Priority"),
+		Self				UMETA(DisplayName="Self"),
+		Both				UMETA(DisplayName="Both"),
+	};
+}
+
+namespace EBTActiveNode
+{
+	// keep in sync with DescribeActiveNode()
+	enum Type
+	{
+		Composite,
+		ActiveTask,
+		AbortingTask,
+		InactiveTask,
+	};
+}
+
+namespace EBTTaskStatus
+{
+	// keep in sync with DescribeTaskStatus()
+	enum Type
+	{
+		Active,
+		Aborting,
+		Inactive,
+	};
+}
+
+namespace EBTNodeUpdateMode
+{
+	// keep in sync with DescribeNodeUpdateMode()
+	enum Type
+	{
+		Add,				// add node
+		AddForLowerPri,		// add node only when new task has lower priority
+		Remove,				// remove node
+	};
+}
+
+/** wrapper struct for holding a parallel task node and its status */
+struct FBehaviorTreeParallelTask
+{
+	/** worker object */
+	const class UBTTaskNode* TaskNode;
+
+	/** additional mode data used for context switching */
+	TEnumAsByte<EBTTaskStatus::Type> Status;
+
+	FBehaviorTreeParallelTask() : TaskNode(NULL) {}
+	FBehaviorTreeParallelTask(const UBTTaskNode* InTaskNode, EBTTaskStatus::Type InStatus) : TaskNode(InTaskNode), Status(InStatus) {}
+
+	bool operator==(const FBehaviorTreeParallelTask& Other) const { return TaskNode == Other.TaskNode; }
+	bool operator==(const UBTTaskNode* OtherTask) const { return TaskNode == OtherTask; }
+};
+
+namespace EBTExecutionSnap
+{
+	enum Type
+	{
+		Regular,
+		OutOfNodes,
+	};
+}
+
+namespace EBTDescriptionVerbosity
+{
+	enum Type
+	{
+		Basic,
+		Detailed,
+	};
+}
+
+/** @TODO might wanna move this to a dedicated header */
+namespace FBlackboard
+{
+	const FName KeySelf = TEXT("SelfActor");
+
+	/** @todo FBlackboard::FKey needs to be used instead of uint8 for keys */
+	typedef uint8 FKey;
+
+	const FKey InvalidKey = FKey(-1);
+}
+
+/** debugger data about subtree instance */
+struct FBehaviorTreeDebuggerInstance
+{
+	struct FNodeFlowData
+	{
+		uint16 ExecutionIndex;
+		uint16 bPassed : 1;
+		uint16 bOptional : 1;
+		uint16 bTrigger : 1;
+		uint16 bDiscardedTrigger : 1;
+
+		FNodeFlowData() : ExecutionIndex(INDEX_NONE), bPassed(0), bOptional(0), bTrigger(0), bDiscardedTrigger(0) {}
+	};
+
+	FBehaviorTreeDebuggerInstance() : TreeAsset(NULL), RootNode(NULL) {}
+
+	/** behavior tree asset */
+	class UBehaviorTree* TreeAsset;
+
+	/** root node in template */
+	class UBTCompositeNode* RootNode;
+
+	/** execution indices of active nodes */
+	TArray<uint16> ActivePath;
+
+	/** execution indices of active nodes */
+	TArray<uint16> AdditionalActiveNodes;
+
+	/** search flow from previous state */
+	TArray<FNodeFlowData> PathFromPrevious;
+
+	/** runtime descriptions for each execution index */
+	TArray<FString> RuntimeDesc;
+
+	FORCEINLINE bool IsValid() const { return ActivePath.Num() != 0; }
+};
+
+/** debugger data about current execution step */
+struct FBehaviorTreeExecutionStep
+{
+	FBehaviorTreeExecutionStep() : TimeStamp(0.f), StepIndex(INDEX_NONE) {}
+
+	/** subtree instance stack */
+	TArray<FBehaviorTreeDebuggerInstance> InstanceStack;
+
+	/** blackboard snapshot: value descriptions */
+	TArray<FString> BlackboardValues;
+
+	/** Game world's time stamp of this step */
+	float TimeStamp;
+
+	/** index of execution step */
+	int32 StepIndex;
+};
+
+/** identifier of subtree instance */
+struct FBehaviorTreeInstanceId
+{
+	/** behavior tree asset */
+	class UBehaviorTree* TreeAsset;
+
+	/** execution index path from root */
+	TArray<uint16> Path;
+
+	/** index of first node instance (BehaviorTreeComponent.NodeInstances) */
+	int32 FirstNodeInstance;
+
+	FBehaviorTreeInstanceId() :	TreeAsset(0), FirstNodeInstance(-1) {}
+
+	bool operator==(const FBehaviorTreeInstanceId& Other) const
+	{
+		return (TreeAsset == Other.TreeAsset) && (Path == Other.Path);
+	}
+};
+
+/** data required for instance of single subtree */
+struct FBehaviorTreeInstance
+{
+	/** root node in template */
+	class UBTCompositeNode* RootNode;
+
+	/** active node in template */
+	class UBTNode* ActiveNode;
+
+	/** active auxiliary nodes */
+	TArray<class UBTAuxiliaryNode*> ActiveAuxNodes;
+
+	/** active parallel tasks */
+	TArray<FBehaviorTreeParallelTask> ParallelTasks;
+
+	/** memory: instance */
+	TArray<uint8> InstanceMemory;
+
+	/** index of identifier (BehaviorTreeComponent.KnownInstances) */
+	uint8 InstanceIdIndex;
+
+	/** active node type */
+	TEnumAsByte<EBTActiveNode::Type> ActiveNodeType;
+
+	FBehaviorTreeInstance() { IncMemoryStats(); }
+	FBehaviorTreeInstance(const FBehaviorTreeInstance& Other) { *this = Other; IncMemoryStats(); }
+	FBehaviorTreeInstance(int32 MemorySize) { InstanceMemory.AddZeroed(MemorySize); IncMemoryStats(); }
+	~FBehaviorTreeInstance() { DecMemoryStats(); }
+
+#if STATS
+	FORCEINLINE void IncMemoryStats() { INC_MEMORY_STAT_BY(STAT_AI_BehaviorTree_InstanceMemory, GetAllocatedSize()); }
+	FORCEINLINE void DecMemoryStats() { DEC_MEMORY_STAT_BY(STAT_AI_BehaviorTree_InstanceMemory, GetAllocatedSize()); }
+	FORCEINLINE uint32 GetAllocatedSize() const 
+	{
+		return sizeof(*this) + ActiveAuxNodes.GetAllocatedSize() + ParallelTasks.GetAllocatedSize() + InstanceMemory.GetAllocatedSize(); 
+	}
+#else
+	FORCEINLINE uint32 GetAllocatedSize() const { return 0; }
+	FORCEINLINE void IncMemoryStats() {}
+	FORCEINLINE void DecMemoryStats() {}
+#endif // STATS
+
+	/** initialize memory and create node instances */
+	void Initialize(class UBehaviorTreeComponent* OwnerComp, UBTCompositeNode* Node, int32& InstancedIndex);
+
+	/** update injected nodes */
+	void InjectNodes(class UBehaviorTreeComponent* OwnerComp, UBTCompositeNode* Node, int32& InstancedIndex);
+
+	/** cleanup node instances */
+	void Cleanup(class UBehaviorTreeComponent* OwnerComp);
+};
+
+struct FBTNodeIndex
+{
+	/** index of instance of stack */
+	uint16 InstanceIndex;
+
+	/** execution index within instance */
+	uint16 ExecutionIndex;
+
+	FBTNodeIndex() : InstanceIndex(MAX_uint16), ExecutionIndex(MAX_uint16) {}
+	FBTNodeIndex(uint16 InInstanceIndex, uint16 InExecutionIndex) : InstanceIndex(InInstanceIndex), ExecutionIndex(InExecutionIndex) {}
+
+	bool TakesPriorityOver(const FBTNodeIndex& Other) const;
+	bool IsSet() const { return InstanceIndex < MAX_uint16; }
+
+	FORCEINLINE bool operator==(const FBTNodeIndex& Other) const { return Other.ExecutionIndex == ExecutionIndex && Other.InstanceIndex == InstanceIndex; }
+	FORCEINLINE friend uint32 GetTypeHash(const FBTNodeIndex& Other) { return Other.ExecutionIndex ^ Other.InstanceIndex; }
+
+	FORCEINLINE FString Describe() const { return FString::Printf(TEXT("[%d:%d]"), InstanceIndex, ExecutionIndex); }
+};
+
+/** node update data */
+struct FBehaviorTreeSearchUpdate
+{
+	class UBTAuxiliaryNode* AuxNode;
+	class UBTTaskNode* TaskNode;
+
+	uint16 InstanceIndex;
+
+	TEnumAsByte<EBTNodeUpdateMode::Type> Mode;
+
+	/** if set, this entry will be applied AFTER other are processed */
+	uint8 bPostUpdate : 1;
+
+	FBehaviorTreeSearchUpdate() : AuxNode(0), TaskNode(0), InstanceIndex(0) {}
+	FBehaviorTreeSearchUpdate(const class UBTAuxiliaryNode* InAuxNode, uint16 InInstanceIndex, EBTNodeUpdateMode::Type InMode) :
+		AuxNode((class UBTAuxiliaryNode*)InAuxNode), TaskNode(0), InstanceIndex(InInstanceIndex), Mode(InMode) 
+	{}
+	FBehaviorTreeSearchUpdate(const class UBTTaskNode* InTaskNode, uint16 InInstanceIndex, EBTNodeUpdateMode::Type InMode) :
+		AuxNode(0), TaskNode((class UBTTaskNode*)InTaskNode), InstanceIndex(InInstanceIndex), Mode(InMode) 
+	{}
+};
+
+/** node search data */
+struct FBehaviorTreeSearchData
+{
+	/** BT component */
+	class UBehaviorTreeComponent* OwnerComp;
+
+	/** requested updates of additional nodes (preconditions, services, parallels)
+	 *  buffered during search to prevent instant add & remove pairs */
+	TArray<struct FBehaviorTreeSearchUpdate> PendingUpdates;
+
+	/** first node allowed in search */
+	struct FBTNodeIndex SearchStart;
+
+	/** adds update info to PendingUpdates array, removing all previous updates for this node */
+	void AddUniqueUpdate(const FBehaviorTreeSearchUpdate& UpdateInfo);
+};
+
+/** property block in blueprint defined nodes */
+struct FBehaviorTreePropertyMemory
+{
+	uint16 Offset;
+	uint16 BlockSize;
+	
+	FBehaviorTreePropertyMemory() {}
+	FBehaviorTreePropertyMemory(int32 Value) : Offset((uint32)Value >> 16), BlockSize((uint32)Value & 0xFFFF) {}
+
+	int32 Pack() const { return (int32)(((uint32)Offset << 16) | BlockSize); }
+};
+
+/** helper struct for defining types of allowed blackboard entries
+ *  (e.g. only entries holding points and objects derived form actor class) */
+USTRUCT(BlueprintType)
+struct AIMODULE_API FBlackboardKeySelector
+{
+	GENERATED_USTRUCT_BODY()
+
+	FBlackboardKeySelector() : SelectedKeyID(FBlackboard::InvalidKey)
+	{}
+
+	/** array of allowed types with additional properties (e.g. ubject's base class) 
+	  * EditDefaults is required for FBlackboardSelectorDetails::CacheBlackboardData() */
+	UPROPERTY(transient, EditDefaultsOnly, Category=Blackboard)
+	TArray<class UBlackboardKeyType*> AllowedTypes;
+
+	/** name of selected key */
+	UPROPERTY(EditInstanceOnly, Category=Blackboard)
+	FName SelectedKeyName;
+
+	/** class of selected key  */
+	UPROPERTY(transient, EditInstanceOnly, Category=Blackboard)
+	TSubclassOf<class UBlackboardKeyType> SelectedKeyType;
+
+protected:
+	/** ID of selected key */
+	UPROPERTY(transient, EditInstanceOnly, Category=Blackboard)
+	uint8 SelectedKeyID;
+
+	UPROPERTY(transient, EditDefaultsOnly, Category=Blackboard)
+	uint32 bNoneIsAllowedValue:1;
+
+public:
+	/** cache ID and class of selected key */
+	void CacheSelectedKey(class UBlackboardData* BlackboardAsset);
+
+	/** find initial selection */
+	void InitSelectedKey(class UBlackboardData* BlackboardAsset);
+
+	void AllowNoneAsValue(bool bNewVal ) { bNoneIsAllowedValue = bNewVal; }
+
+	FORCEINLINE uint8 GetSelectedKeyID() const { return SelectedKeyID; }
+
+	/** helper functions for setting basic filters */
+	void AddObjectFilter(UObject* Owner, TSubclassOf<UObject> AllowedClass);
+	void AddClassFilter(UObject* Owner, TSubclassOf<UClass> AllowedClass);
+	void AddEnumFilter(UObject* Owner, UEnum* AllowedEnum);
+	void AddNativeEnumFilter(UObject* Owner, const FString& AllowedEnumName);
+	void AddIntFilter(UObject* Owner);
+	void AddFloatFilter(UObject* Owner);
+	void AddBoolFilter(UObject* Owner);
+	void AddVectorFilter(UObject* Owner);
+	void AddStringFilter(UObject* Owner);
+	void AddNameFilter(UObject* Owner);
+
+	FORCEINLINE bool IsNone() const { return bNoneIsAllowedValue && SelectedKeyID == FBlackboard::InvalidKey; }
+
+	friend class FBlackboardDecoratorDetails;
+};
+
+UCLASS(Abstract)
+class AIMODULE_API UBehaviorTreeTypes : public UObject
+{
+	GENERATED_UCLASS_BODY()
+
+	static FString DescribeNodeHelper(const UBTNode* Node);
+
+	static FString DescribeNodeResult(EBTNodeResult::Type NodeResult);
+	static FString DescribeFlowAbortMode(EBTFlowAbortMode::Type FlowAbortMode);
+	static FString DescribeActiveNode(EBTActiveNode::Type ActiveNodeType);
+	static FString DescribeTaskStatus(EBTTaskStatus::Type TaskStatus);
+	static FString DescribeNodeUpdateMode(EBTNodeUpdateMode::Type UpdateMode);
+
+	/** returns short name of object's class (BTTaskNode_Wait -> Wait) */
+	static FString GetShortTypeName(const UObject* Ob);
+};
diff --git a/Engine/Source/Runtime/AIModule/Classes/BehaviorTree/Blackboard/BlackboardKeyType.h b/Engine/Source/Runtime/AIModule/Classes/BehaviorTree/Blackboard/BlackboardKeyType.h
new file mode 100644
--- /dev/null
+++ b/Engine/Source/Runtime/AIModule/Classes/BehaviorTree/Blackboard/BlackboardKeyType.h
@@ -0,0 +1,142 @@
+// Copyright 1998-2014 Epic Games, Inc. All Rights Reserved.
+#pragma once
+#include "BlackboardKeyType.generated.h"
+
+namespace EBlackboardCompare
+{
+	enum Type
+	{
+		Less = -1, 	
+		Equal = 0, 
+		Greater = 1,
+
+		NotEqual = 1,	// important, do not change
+	};
+}
+
+namespace EBlackboardKeyOperation
+{
+	enum Type
+	{
+		Basic,
+		Arithmetic,
+		Text,
+	};
+}
+
+UENUM()
+namespace EBasicKeyOperation
+{
+	enum Type
+	{
+		Set				UMETA(DisplayName="Is Set"),
+		NotSet			UMETA(DisplayName="Is Not Set"),
+	};
+}
+
+UENUM()
+namespace EArithmeticKeyOperation
+{
+	enum Type
+	{
+		Equal			UMETA(DisplayName="Is Equal To"),
+		NotEqual		UMETA(DisplayName="Is Not Equal To"),
+		Less			UMETA(DisplayName="Is Less Than"),
+		LessOrEqual		UMETA(DisplayName="Is Less Than Or Equal To"),
+		Greater			UMETA(DisplayName="Is Greater Than"),
+		GreaterOrEqual	UMETA(DisplayName="Is Greater Than Or Equal To"),
+	};
+}
+
+UENUM()
+namespace ETextKeyOperation
+{
+	enum Type
+	{
+		Equal			UMETA(DisplayName="Is Equal To"),
+		NotEqual		UMETA(DisplayName="Is Not Equal To"),
+		Contain			UMETA(DisplayName="Contains"),
+		NotContain		UMETA(DisplayName="Not Contains"),
+	};
+}
+
+UCLASS(EditInlineNew, Abstract, CollapseCategories, AutoExpandCategories=(Blackboard))
+class AIMODULE_API UBlackboardKeyType : public UObject
+{
+	GENERATED_UCLASS_BODY()
+		
+	/** get ValueSize */
+	FORCEINLINE uint16 GetValueSize() const { return ValueSize; }
+
+	/** get test supported by this type */
+	FORCEINLINE EBlackboardKeyOperation::Type GetTestOperation() const { return SupportedOp; }
+
+	/** initialize memory */
+	virtual void Initialize(uint8* MemoryBlock) const;
+
+	/** does it match settings in filter? */
+	virtual bool IsAllowedByFilter(UBlackboardKeyType* FilterOb) const;
+
+	/** extract location from entry */
+	virtual bool GetLocation(const uint8* MemoryBlock, FVector& Location) const;
+
+	/** extract rotation from entry */
+	virtual bool GetRotation(const uint8* MemoryBlock, FRotator& Rotation) const;
+
+	/** interprets given two memory blocks as "my" type and compares them */
+	virtual EBlackboardCompare::Type Compare(const uint8* MemoryBlockA, const uint8* MemoryBlockB) const;
+
+	/** various value testing, used by decorators */
+	virtual bool TestBasicOperation(const uint8* MemoryBlock, EBasicKeyOperation::Type Op) const;
+	virtual bool TestArithmeticOperation(const uint8* MemoryBlock, EArithmeticKeyOperation::Type Op, int32 OtherIntValue, float OtherFloatValue) const;
+	virtual bool TestTextOperation(const uint8* MemoryBlock, ETextKeyOperation::Type Op, const FString& OtherString) const;
+
+	/** describe params of arithmetic test */
+	virtual FString DescribeArithmeticParam(int32 IntValue, float FloatValue) const;
+
+	/** convert value to text */
+	virtual FString DescribeValue(const uint8* MemoryBlock) const;
+
+	/** description of params for property view */
+	virtual FString DescribeSelf() const;
+
+protected:
+
+	/** size of value for this type */
+	uint16 ValueSize;
+
+	/** decorator operation supported with this type */
+	TEnumAsByte<EBlackboardKeyOperation::Type> SupportedOp;
+
+	/** helper function for reading typed data from memory block */
+	template<typename T>
+	static T GetValueFromMemory(const uint8* MemoryBlock)
+	{
+		return *((T*)MemoryBlock);
+	}
+
+	/** helper function for writing typed data to memory block, returns true if value has changed */
+	template<typename T>
+	static bool SetValueInMemory(uint8* MemoryBlock, const T& Value)
+	{
+		const bool bChanged = *((T*)MemoryBlock) != Value;
+		*((T*)MemoryBlock) = Value;
+		
+		return bChanged;
+	}
+
+	/** helper function for witting weak object data to memory block, returns true if value has changed */
+	template<typename T>
+	static bool SetWeakObjectInMemory(uint8* MemoryBlock, const TWeakObjectPtr<T>& Value)
+	{
+		TWeakObjectPtr<T>* PrevValue = (TWeakObjectPtr<T>*)MemoryBlock;
+		const bool bChanged =
+			(Value.IsValid(false, true) != PrevValue->IsValid(false, true)) ||
+			(Value.IsStale(false, true) != PrevValue->IsStale(false, true)) ||
+			(*PrevValue) != Value;
+
+		*((TWeakObjectPtr<T>*)MemoryBlock) = Value;
+
+		return bChanged;
+	}
+};
diff --git a/Engine/Source/Runtime/AIModule/Classes/BehaviorTree/Blackboard/BlackboardKeyType_Bool.h b/Engine/Source/Runtime/AIModule/Classes/BehaviorTree/Blackboard/BlackboardKeyType_Bool.h
new file mode 100644
--- /dev/null
+++ b/Engine/Source/Runtime/AIModule/Classes/BehaviorTree/Blackboard/BlackboardKeyType_Bool.h
@@ -0,0 +1,17 @@
+// Copyright 1998-2014 Epic Games, Inc. All Rights Reserved.
+#pragma once
+#include "BehaviorTree/Blackboard/BlackboardKeyType.h"
+#include "BlackboardKeyType_Bool.generated.h"
+
+UCLASS(EditInlineNew, MinimalAPI)
+class UBlackboardKeyType_Bool : public UBlackboardKeyType
+{
+	GENERATED_UCLASS_BODY()
+
+	static bool GetValue(const uint8* RawData);
+	static bool SetValue(uint8* RawData, bool bValue);
+
+	virtual FString DescribeValue(const uint8* RawData) const OVERRIDE;
+	virtual EBlackboardCompare::Type Compare(const uint8* MemoryBlockA, const uint8* MemoryBlockB) const OVERRIDE;
+	virtual bool TestBasicOperation(const uint8* MemoryBlock, EBasicKeyOperation::Type Op) const OVERRIDE;
+};
diff --git a/Engine/Source/Runtime/AIModule/Classes/BehaviorTree/Blackboard/BlackboardKeyType_Class.h b/Engine/Source/Runtime/AIModule/Classes/BehaviorTree/Blackboard/BlackboardKeyType_Class.h
new file mode 100644
--- /dev/null
+++ b/Engine/Source/Runtime/AIModule/Classes/BehaviorTree/Blackboard/BlackboardKeyType_Class.h
@@ -0,0 +1,22 @@
+// Copyright 1998-2014 Epic Games, Inc. All Rights Reserved.
+#pragma once
+#include "BehaviorTree/Blackboard/BlackboardKeyType.h"
+#include "BlackboardKeyType_Class.generated.h"
+
+UCLASS(EditInlineNew, MinimalAPI)
+class UBlackboardKeyType_Class : public UBlackboardKeyType
+{
+	GENERATED_UCLASS_BODY()
+
+	UPROPERTY(Category=Blackboard, EditDefaultsOnly, meta=(AllowAbstract="1"))
+	UClass* BaseClass;
+	
+	static UClass* GetValue(const uint8* RawData);
+	static bool SetValue(uint8* RawData, UClass* Value);
+
+	virtual FString DescribeValue(const uint8* RawData) const OVERRIDE;
+	virtual FString DescribeSelf() const OVERRIDE;
+	virtual bool IsAllowedByFilter(UBlackboardKeyType* FilterOb) const OVERRIDE;
+	virtual EBlackboardCompare::Type Compare(const uint8* MemoryBlockA, const uint8* MemoryBlockB) const OVERRIDE;
+	virtual bool TestBasicOperation(const uint8* MemoryBlock, EBasicKeyOperation::Type Op) const OVERRIDE;
+};
diff --git a/Engine/Source/Runtime/AIModule/Classes/BehaviorTree/Blackboard/BlackboardKeyType_Enum.h b/Engine/Source/Runtime/AIModule/Classes/BehaviorTree/Blackboard/BlackboardKeyType_Enum.h
new file mode 100644
--- /dev/null
+++ b/Engine/Source/Runtime/AIModule/Classes/BehaviorTree/Blackboard/BlackboardKeyType_Enum.h
@@ -0,0 +1,23 @@
+// Copyright 1998-2014 Epic Games, Inc. All Rights Reserved.
+#pragma once
+#include "BehaviorTree/Blackboard/BlackboardKeyType.h"
+#include "BlackboardKeyType_Enum.generated.h"
+
+UCLASS(EditInlineNew, MinimalAPI)
+class UBlackboardKeyType_Enum : public UBlackboardKeyType
+{
+	GENERATED_UCLASS_BODY()
+
+	UPROPERTY(Category=Blackboard, EditDefaultsOnly)
+	UEnum* EnumType;
+
+	static uint8 GetValue(const uint8* RawData);
+	static bool SetValue(uint8* RawData, uint8 Value);
+
+	virtual FString DescribeValue(const uint8* RawData) const OVERRIDE;
+	virtual FString DescribeSelf() const OVERRIDE;
+	virtual bool IsAllowedByFilter(UBlackboardKeyType* FilterOb) const OVERRIDE;
+	virtual EBlackboardCompare::Type Compare(const uint8* MemoryBlockA, const uint8* MemoryBlockB) const OVERRIDE;
+	virtual bool TestArithmeticOperation(const uint8* MemoryBlock, EArithmeticKeyOperation::Type Op, int32 OtherIntValue, float OtherFloatValue) const OVERRIDE;
+	virtual FString DescribeArithmeticParam(int32 IntValue, float FloatValue) const OVERRIDE;
+};
diff --git a/Engine/Source/Runtime/AIModule/Classes/BehaviorTree/Blackboard/BlackboardKeyType_Float.h b/Engine/Source/Runtime/AIModule/Classes/BehaviorTree/Blackboard/BlackboardKeyType_Float.h
new file mode 100644
--- /dev/null
+++ b/Engine/Source/Runtime/AIModule/Classes/BehaviorTree/Blackboard/BlackboardKeyType_Float.h
@@ -0,0 +1,18 @@
+// Copyright 1998-2014 Epic Games, Inc. All Rights Reserved.
+#pragma once
+#include "BehaviorTree/Blackboard/BlackboardKeyType.h"
+#include "BlackboardKeyType_Float.generated.h"
+
+UCLASS(EditInlineNew, MinimalAPI)
+class UBlackboardKeyType_Float : public UBlackboardKeyType
+{
+	GENERATED_UCLASS_BODY()
+
+	static float GetValue(const uint8* RawData);
+	static bool SetValue(uint8* RawData, float Value);
+
+	virtual FString DescribeValue(const uint8* RawData) const OVERRIDE;
+	virtual EBlackboardCompare::Type Compare(const uint8* MemoryBlockA, const uint8* MemoryBlockB) const OVERRIDE;
+	virtual bool TestArithmeticOperation(const uint8* MemoryBlock, EArithmeticKeyOperation::Type Op, int32 OtherIntValue, float OtherFloatValue) const OVERRIDE;
+	virtual FString DescribeArithmeticParam(int32 IntValue, float FloatValue) const OVERRIDE;
+};
diff --git a/Engine/Source/Runtime/AIModule/Classes/BehaviorTree/Blackboard/BlackboardKeyType_Int.h b/Engine/Source/Runtime/AIModule/Classes/BehaviorTree/Blackboard/BlackboardKeyType_Int.h
new file mode 100644
--- /dev/null
+++ b/Engine/Source/Runtime/AIModule/Classes/BehaviorTree/Blackboard/BlackboardKeyType_Int.h
@@ -0,0 +1,18 @@
+// Copyright 1998-2014 Epic Games, Inc. All Rights Reserved.
+#pragma once
+#include "BehaviorTree/Blackboard/BlackboardKeyType.h"
+#include "BlackboardKeyType_Int.generated.h"
+
+UCLASS(EditInlineNew, MinimalAPI)
+class UBlackboardKeyType_Int : public UBlackboardKeyType
+{
+	GENERATED_UCLASS_BODY()
+
+	static int32 GetValue(const uint8* RawData);
+	static bool SetValue(uint8* RawData, int32 Value);
+
+	virtual FString DescribeValue(const uint8* RawData) const OVERRIDE;
+	virtual EBlackboardCompare::Type Compare(const uint8* MemoryBlockA, const uint8* MemoryBlockB) const OVERRIDE;
+	virtual bool TestArithmeticOperation(const uint8* MemoryBlock, EArithmeticKeyOperation::Type Op, int32 OtherIntValue, float OtherFloatValue) const OVERRIDE;
+	virtual FString DescribeArithmeticParam(int32 IntValue, float FloatValue) const OVERRIDE;
+};
diff --git a/Engine/Source/Runtime/AIModule/Classes/BehaviorTree/Blackboard/BlackboardKeyType_Name.h b/Engine/Source/Runtime/AIModule/Classes/BehaviorTree/Blackboard/BlackboardKeyType_Name.h
new file mode 100644
--- /dev/null
+++ b/Engine/Source/Runtime/AIModule/Classes/BehaviorTree/Blackboard/BlackboardKeyType_Name.h
@@ -0,0 +1,17 @@
+// Copyright 1998-2014 Epic Games, Inc. All Rights Reserved.
+#pragma once
+#include "BehaviorTree/Blackboard/BlackboardKeyType.h"
+#include "BlackboardKeyType_Name.generated.h"
+
+UCLASS(EditInlineNew, MinimalAPI)
+class UBlackboardKeyType_Name : public UBlackboardKeyType
+{
+	GENERATED_UCLASS_BODY()
+
+	static FName GetValue(const uint8* RawData);
+	static bool SetValue(uint8* RawData, const FName& Value);
+
+	virtual FString DescribeValue(const uint8* RawData) const OVERRIDE;
+	virtual EBlackboardCompare::Type Compare(const uint8* MemoryBlockA, const uint8* MemoryBlockB) const OVERRIDE;
+	virtual bool TestTextOperation(const uint8* MemoryBlock, ETextKeyOperation::Type Op, const FString& OtherString) const OVERRIDE;
+};
diff --git a/Engine/Source/Runtime/AIModule/Classes/BehaviorTree/Blackboard/BlackboardKeyType_NativeEnum.h b/Engine/Source/Runtime/AIModule/Classes/BehaviorTree/Blackboard/BlackboardKeyType_NativeEnum.h
new file mode 100644
--- /dev/null
+++ b/Engine/Source/Runtime/AIModule/Classes/BehaviorTree/Blackboard/BlackboardKeyType_NativeEnum.h
@@ -0,0 +1,34 @@
+// Copyright 1998-2014 Epic Games, Inc. All Rights Reserved.
+
+#pragma once
+#include "BehaviorTree/Blackboard/BlackboardKeyType.h"
+#include "BlackboardKeyType_NativeEnum.generated.h"
+
+UCLASS(EditInlineNew, MinimalAPI)
+class UBlackboardKeyType_NativeEnum : public UBlackboardKeyType
+{
+	GENERATED_UCLASS_BODY()
+
+	UPROPERTY(Category=Blackboard, EditDefaultsOnly)
+	FString EnumName;
+
+	UPROPERTY(Category=Blackboard, VisibleDefaultsOnly)
+	bool bIsEnumNameValid;
+
+	UPROPERTY()
+	UEnum* EnumType;
+
+	static uint8 GetValue(const uint8* RawData);
+	static bool SetValue(uint8* RawData, uint8 Value);
+
+	virtual FString DescribeValue(const uint8* RawData) const OVERRIDE;
+	virtual FString DescribeSelf() const OVERRIDE;
+	virtual bool IsAllowedByFilter(UBlackboardKeyType* FilterOb) const OVERRIDE;
+	virtual EBlackboardCompare::Type Compare(const uint8* MemoryBlockA, const uint8* MemoryBlockB) const OVERRIDE;
+	virtual bool TestArithmeticOperation(const uint8* MemoryBlock, EArithmeticKeyOperation::Type Op, int32 OtherIntValue, float OtherFloatValue) const OVERRIDE;
+	virtual FString DescribeArithmeticParam(int32 IntValue, float FloatValue) const OVERRIDE;
+
+#if WITH_EDITOR
+	virtual void PostEditChangeProperty(struct FPropertyChangedEvent& PropertyChangedEvent) OVERRIDE;
+#endif
+};
diff --git a/Engine/Source/Runtime/AIModule/Classes/BehaviorTree/Blackboard/BlackboardKeyType_Object.h b/Engine/Source/Runtime/AIModule/Classes/BehaviorTree/Blackboard/BlackboardKeyType_Object.h
new file mode 100644
--- /dev/null
+++ b/Engine/Source/Runtime/AIModule/Classes/BehaviorTree/Blackboard/BlackboardKeyType_Object.h
@@ -0,0 +1,24 @@
+// Copyright 1998-2014 Epic Games, Inc. All Rights Reserved.
+#pragma once
+#include "BehaviorTree/Blackboard/BlackboardKeyType.h"
+#include "BlackboardKeyType_Object.generated.h"
+
+UCLASS(EditInlineNew, MinimalAPI)
+class UBlackboardKeyType_Object : public UBlackboardKeyType
+{
+	GENERATED_UCLASS_BODY()
+
+	UPROPERTY(Category=Blackboard, EditDefaultsOnly, meta=(AllowAbstract="1"))
+	UClass* BaseClass;
+
+	static UObject* GetValue(const uint8* RawData);
+	static bool SetValue(uint8* RawData, UObject* Value);
+
+	virtual FString DescribeValue(const uint8* RawData) const OVERRIDE;
+	virtual FString DescribeSelf() const OVERRIDE;
+	virtual bool IsAllowedByFilter(UBlackboardKeyType* FilterOb) const OVERRIDE;
+	virtual bool GetLocation(const uint8* RawData, FVector& Location) const OVERRIDE;
+	virtual bool GetRotation(const uint8* MemoryBlock, FRotator& Rotation) const OVERRIDE;
+	virtual EBlackboardCompare::Type Compare(const uint8* MemoryBlockA, const uint8* MemoryBlockB) const OVERRIDE;
+	virtual bool TestBasicOperation(const uint8* MemoryBlock, EBasicKeyOperation::Type Op) const OVERRIDE;
+};
diff --git a/Engine/Source/Runtime/AIModule/Classes/BehaviorTree/Blackboard/BlackboardKeyType_String.h b/Engine/Source/Runtime/AIModule/Classes/BehaviorTree/Blackboard/BlackboardKeyType_String.h
new file mode 100644
--- /dev/null
+++ b/Engine/Source/Runtime/AIModule/Classes/BehaviorTree/Blackboard/BlackboardKeyType_String.h
@@ -0,0 +1,17 @@
+// Copyright 1998-2014 Epic Games, Inc. All Rights Reserved.
+#pragma once
+#include "BehaviorTree/Blackboard/BlackboardKeyType.h"
+#include "BlackboardKeyType_String.generated.h"
+
+UCLASS(EditInlineNew, MinimalAPI)
+class UBlackboardKeyType_String : public UBlackboardKeyType
+{
+	GENERATED_UCLASS_BODY()
+
+	static FString GetValue(const uint8* RawData);
+	static bool SetValue(uint8* RawData, const FString& Value);
+
+	virtual FString DescribeValue(const uint8* RawData) const OVERRIDE;
+	virtual EBlackboardCompare::Type Compare(const uint8* MemoryBlockA, const uint8* MemoryBlockB) const OVERRIDE;
+	virtual bool TestTextOperation(const uint8* MemoryBlock, ETextKeyOperation::Type Op, const FString& OtherString) const OVERRIDE;
+};
diff --git a/Engine/Source/Runtime/AIModule/Classes/BehaviorTree/Blackboard/BlackboardKeyType_Vector.h b/Engine/Source/Runtime/AIModule/Classes/BehaviorTree/Blackboard/BlackboardKeyType_Vector.h
new file mode 100644
--- /dev/null
+++ b/Engine/Source/Runtime/AIModule/Classes/BehaviorTree/Blackboard/BlackboardKeyType_Vector.h
@@ -0,0 +1,19 @@
+// Copyright 1998-2014 Epic Games, Inc. All Rights Reserved.
+#pragma once
+#include "BehaviorTree/Blackboard/BlackboardKeyType.h"
+#include "BlackboardKeyType_Vector.generated.h"
+
+UCLASS(EditInlineNew, MinimalAPI)
+class UBlackboardKeyType_Vector : public UBlackboardKeyType
+{
+	GENERATED_UCLASS_BODY()
+
+	static FVector GetValue(const uint8* RawData);
+	static bool SetValue(uint8* RawData, const FVector& Value);
+
+	virtual void Initialize(uint8* RawData) const OVERRIDE;
+	virtual FString DescribeValue(const uint8* RawData) const OVERRIDE;
+	virtual bool GetLocation(const uint8* RawData, FVector& Location) const OVERRIDE;
+	virtual EBlackboardCompare::Type Compare(const uint8* MemoryBlockA, const uint8* MemoryBlockB) const OVERRIDE;
+	virtual bool TestBasicOperation(const uint8* MemoryBlock, EBasicKeyOperation::Type Op) const OVERRIDE;
+};
diff --git a/Engine/Source/Runtime/AIModule/Classes/BehaviorTree/BlackboardComponent.h b/Engine/Source/Runtime/AIModule/Classes/BehaviorTree/BlackboardComponent.h
new file mode 100644
--- /dev/null
+++ b/Engine/Source/Runtime/AIModule/Classes/BehaviorTree/BlackboardComponent.h
@@ -0,0 +1,225 @@
+// Copyright 1998-2014 Epic Games, Inc. All Rights Reserved.
+
+/** 
+ *  Blackboard - holds AI's world knowledge, easily accessible for behavior trees
+ *
+ *  Properties are stored in byte array, and should be accessed only though
+ *  GetValue* / SetValue* functions. They will handle broadcasting change events
+ *  for registered observers.
+ *
+ *  Keys are defined by BlackboardData data asset.
+ */
+
+#pragma once
+#include "Components/ActorComponent.h"
+#include "BehaviorTreeTypes.h"
+#include "BehaviorTree/Blackboard/BlackboardKeyType.h"
+#include "BlackboardData.h"
+#include "BlackboardComponent.generated.h"
+
+namespace EBlackboardDescription
+{
+	enum Type
+	{
+		OnlyValue,
+		KeyWithValue,
+		DetailedKeyWithValue,
+		Full,
+	};
+}
+
+UCLASS()
+class AIMODULE_API UBlackboardComponent : public UActorComponent
+{
+	GENERATED_UCLASS_BODY()
+
+	/** BEGIN UActorComponent overrides */
+	virtual void InitializeComponent() OVERRIDE;
+	/** END UActorComponent overrides */
+
+	/** @return name of key */
+	FName GetKeyName(uint8 KeyID) const;
+
+	/** @return key ID from name */
+	uint8 GetKeyID(const FName& KeyName) const;
+
+	/** @return class of value for given key */
+	TSubclassOf<UBlackboardKeyType> GetKeyType(uint8 KeyID) const;
+
+	/** @return number of entries in data asset */
+	int32 GetNumKeys() const;
+
+	/** @return true if blackboard have valid data asset */
+	bool HasValidAsset() const;
+
+	/** register observer for blackboard key */
+	void RegisterObserver(uint8 KeyID, FOnBlackboardChange ObserverDelegate);
+
+	/** unregister observer from blackboard key */
+	void UnregisterObserver(uint8 KeyID, FOnBlackboardChange ObserverDelegate);
+
+	/** pause change notifies and add them to queue */
+	void PauseUpdates();
+
+	/** resume change notifies and process queued list */
+	void ResumeUpdates();
+
+	/** @return associated behavior tree component */
+	class UBrainComponent* GetBrainComponent() const;
+
+	/** @return blackboard data asset */
+	class UBlackboardData* GetBlackboardAsset() const;
+
+	/** caches UBrainComponent pointer to be used in communication */
+	void CacheBrainComponent(class UBrainComponent* BrainComponent);
+
+	/** setup component for using given blackboard asset */
+	void InitializeBlackboard(class UBlackboardData* NewAsset);
+	
+	/** @return true if component can be used with specified blackboard asset */
+	bool IsCompatibileWith(class UBlackboardData* TestAsset) const;
+
+	UFUNCTION(BlueprintCallable, Category="AI|Components|Blackboard")
+	UObject* GetValueAsObject(const FName& KeyName) const;
+	UObject* GetValueAsObject(uint8 KeyID) const;
+
+	UFUNCTION(BlueprintCallable, Category="AI|Components|Blackboard")
+	UClass* GetValueAsClass(const FName& KeyName) const;
+	UClass* GetValueAsClass(uint8 KeyID) const;
+
+	UFUNCTION(BlueprintCallable, Category="AI|Components|Blackboard")
+	uint8 GetValueAsEnum(const FName& KeyName) const;
+	uint8 GetValueAsEnum(uint8 KeyID) const;
+
+	UFUNCTION(BlueprintCallable, Category="AI|Components|Blackboard")
+	int32 GetValueAsInt(const FName& KeyName) const;
+	int32 GetValueAsInt(uint8 KeyID) const;
+
+	UFUNCTION(BlueprintCallable, Category="AI|Components|Blackboard")
+	float GetValueAsFloat(const FName& KeyName) const;
+	float GetValueAsFloat(uint8 KeyID) const;
+
+	UFUNCTION(BlueprintCallable, Category="AI|Components|Blackboard")
+	bool GetValueAsBool(const FName& KeyName) const;
+	bool GetValueAsBool(uint8 KeyID) const;
+
+	UFUNCTION(BlueprintCallable, Category="AI|Components|Blackboard")
+	FString GetValueAsString(const FName& KeyName) const;
+	FString GetValueAsString(uint8 KeyID) const;
+	
+	UFUNCTION(BlueprintCallable, Category="AI|Components|Blackboard")
+	FName GetValueAsName(const FName& KeyName) const;
+	FName GetValueAsName(uint8 KeyID) const;
+
+	UFUNCTION(BlueprintCallable, Category="AI|Components|Blackboard")
+	FVector GetValueAsVector(const FName& KeyName) const;
+	FVector GetValueAsVector(uint8 KeyID) const;
+
+	UFUNCTION(BlueprintCallable, Category="AI|Components|Blackboard")
+	void SetValueAsObject(const FName& KeyName, UObject* ObjectValue);
+	void SetValueAsObject(uint8 KeyID, UObject* ObjectValue);
+
+	UFUNCTION(BlueprintCallable, Category="AI|Components|Blackboard")
+	void SetValueAsClass(const FName& KeyName, UClass* ClassValue);
+	void SetValueAsClass(uint8 KeyID, UClass* ClassValue);
+
+	UFUNCTION(BlueprintCallable, Category="AI|Components|Blackboard")
+	void SetValueAsEnum(const FName& KeyName, uint8 EnumValue);
+	void SetValueAsEnum(uint8 KeyID, uint8 EnumValue);
+
+	UFUNCTION(BlueprintCallable, Category="AI|Components|Blackboard")
+	void SetValueAsInt(const FName& KeyName, int32 IntValue);
+	void SetValueAsInt(uint8 KeyID, int32 IntValue);
+
+	UFUNCTION(BlueprintCallable, Category="AI|Components|Blackboard")
+	void SetValueAsFloat(const FName& KeyName, float FloatValue);
+	void SetValueAsFloat(uint8 KeyID, float FloatValue);
+
+	UFUNCTION(BlueprintCallable, Category="AI|Components|Blackboard")
+	void SetValueAsBool(const FName& KeyName, bool BoolValue);
+	void SetValueAsBool(uint8 KeyID, bool BoolValue);
+
+	UFUNCTION(BlueprintCallable, Category="AI|Components|Blackboard")
+	void SetValueAsString(const FName& KeyName, FString StringValue);
+	void SetValueAsString(uint8 KeyID, const FString& StringValue);
+
+	UFUNCTION(BlueprintCallable, Category="AI|Components|Blackboard")
+	void SetValueAsName(const FName& KeyName, FName NameValue);
+	void SetValueAsName(uint8 KeyID, const FName& NameValue);
+
+	UFUNCTION(BlueprintCallable, Category="AI|Components|Blackboard")
+	void SetValueAsVector(const FName& KeyName, FVector VectorValue);
+	void SetValueAsVector(uint8 KeyID, const FVector& VectorValue);
+
+	UFUNCTION(BlueprintCallable, Category="Blackboard")
+	void ClearValueAsVector(const FName& KeyName);
+	void ClearValueAsVector(uint8 KeyID);
+
+	/** return false if call failed (most probably no such entry in BB) */
+	UFUNCTION(BlueprintCallable, Category="AI|Components|Blackboard")
+	bool GetLocationFromEntry(const FName& KeyName, FVector& ResultLocation) const;
+	bool GetLocationFromEntry(uint8 KeyID, FVector& ResultLocation) const;
+
+	/** return false if call failed (most probably no such entry in BB) */
+	UFUNCTION(BlueprintCallable, Category="Blackboard")
+	bool GetRotationFromEntry(const FName& KeyName, FRotator& ResultRotation) const;
+	bool GetRotationFromEntry(uint8 KeyID, FRotator& ResultRotation) const;
+
+	/** get pointer to raw data for given key */
+	FORCEINLINE uint8* GetKeyRawData(const FName& KeyName) { return GetKeyRawData(GetKeyID(KeyName)); }
+	FORCEINLINE uint8* GetKeyRawData(uint8 KeyID) { return ValueMemory.Num() && ValueOffsets.IsValidIndex(KeyID) ? (ValueMemory.GetTypedData() + ValueOffsets[KeyID]) : NULL; }
+
+	FORCEINLINE const uint8* GetKeyRawData(const FName& KeyName) const { return GetKeyRawData(GetKeyID(KeyName)); }
+	FORCEINLINE const uint8* GetKeyRawData(uint8 KeyID) const { return ValueMemory.Num() && ValueOffsets.IsValidIndex(KeyID) ? (ValueMemory.GetTypedData() + ValueOffsets[KeyID]) : NULL; }
+
+	FORCEINLINE bool IsValidKey(uint8 KeyID) const { check(BlackboardAsset); return KeyID != UBlackboardData::InvalidKeyID && BlackboardAsset->Keys.IsValidIndex(KeyID); }
+
+	/** compares blackboard's values under specified keys */
+	EBlackboardCompare::Type CompareKeyValues(TSubclassOf<class UBlackboardKeyType> KeyType, uint8 KeyA, uint8 KeyB) const;
+
+	FString GetDebugInfoString(EBlackboardDescription::Type Mode) const;
+
+	/** get description of value under given key */
+	FString DescribeKeyValue(uint8 KeyID, EBlackboardDescription::Type Mode) const;
+
+#if ENABLE_VISUAL_LOG
+	/** prepare blackboard snapshot for logs */
+	virtual void DescribeSelfToVisLog(struct FVisLogEntry* Snapshot) const;
+#endif
+	
+protected:
+
+	/** cached behavior tree component */
+	UPROPERTY(transient)
+	class UBrainComponent* BrainComp;
+
+	/** data asset defining entries */
+	UPROPERTY(transient)
+	class UBlackboardData* BlackboardAsset;
+
+	/** memory block holding all values */
+	TArray<uint8> ValueMemory;
+
+	/** offsets in ValueMemory for each key */
+	TArray<uint16> ValueOffsets;
+
+	/** observers registered for blackboard keys */
+	TMultiMap<uint8, FOnBlackboardChange> Observers;
+
+	/** queued key change notification, will be processed on ResumeUpdates call */
+	mutable TArray<uint8> QueuedUpdates;
+
+	/** set when notifies are paused and shouldn't be passed to observers */
+	uint32 bPausedNotifies : 1;
+
+	/** notifies behavior tree decorators about change in blackboard */
+	void NotifyObservers(uint8 KeyID) const;
+};
+
+//////////////////////////////////////////////////////////////////////////
+// Inlines
+
+FORCEINLINE bool UBlackboardComponent::HasValidAsset() const
+{
+	return BlackboardAsset && BlackboardAsset->IsValid();
+}
diff --git a/Engine/Source/Runtime/AIModule/Classes/BehaviorTree/BlackboardData.h b/Engine/Source/Runtime/AIModule/Classes/BehaviorTree/BlackboardData.h
new file mode 100644
--- /dev/null
+++ b/Engine/Source/Runtime/AIModule/Classes/BehaviorTree/BlackboardData.h
@@ -0,0 +1,119 @@
+// Copyright 1998-2014 Epic Games, Inc. All Rights Reserved.
+#pragma once
+#include "BlackboardData.generated.h"
+
+/** blackboard entry definition */
+USTRUCT()
+struct FBlackboardEntry
+{
+	GENERATED_USTRUCT_BODY()
+
+	UPROPERTY(EditAnywhere, Category=Blackboard)
+	FName EntryName;
+
+	/** key type and additional properties */
+	UPROPERTY(EditAnywhere, EditInline, Category=Blackboard)
+	UBlackboardKeyType* KeyType;
+};
+
+UCLASS()
+class AIMODULE_API UBlackboardData : public UDataAsset
+{
+	GENERATED_UCLASS_BODY()
+	DECLARE_MULTICAST_DELEGATE_OneParam(FKeyUpdate, UBlackboardData* /*asset*/);
+
+	enum 
+	{
+		InvalidKeyID = 0xFF
+	};
+
+	/** parent blackboard (keys can be overridden) */
+	UPROPERTY(EditAnywhere, Category=Parent)
+	class UBlackboardData* Parent;
+
+#if WITH_EDITORONLY_DATA
+	/** all keys inherited from parent chain */
+	UPROPERTY(VisibleDefaultsOnly, Transient, Category=Parent)
+	TArray<FBlackboardEntry> ParentKeys;
+#endif
+
+	/** blackboard keys */
+	UPROPERTY(EditAnywhere, EditInLine, Category=Blackboard)
+	TArray<FBlackboardEntry> Keys;
+
+	/** @return key ID from name */
+	uint8 GetKeyID(const FName& KeyName) const;
+
+	/** @return name of key */
+	FName GetKeyName(uint8 KeyID) const;
+
+	/** @return class of value for given key */
+	TSubclassOf<UBlackboardKeyType> GetKeyType(uint8 KeyID) const;
+
+	/** @return number of defined keys, including parent chain */
+	int32 GetNumKeys() const;
+
+	FORCEINLINE int32 GetFirstKeyID() const { return FirstKeyID; }
+
+	/** @return key data */
+	const FBlackboardEntry* GetKey(uint8 KeyID) const;
+
+	virtual void PostLoad() OVERRIDE;
+#if WITH_EDITOR
+	virtual void PostEditChangeProperty(struct FPropertyChangedEvent& PropertyChangedEvent) OVERRIDE;
+#endif
+
+	/** @return true if blackboard keys are not conflicting with parent key chain */
+	bool IsValid() const;
+
+	/** updates persistent key with given name, depending on currently defined entries and parent chain
+	 *  @return key type of newly created entry for further setup
+	 */
+	template<class T>
+	T* UpdatePersistentKey(const FName& KeyName)
+	{
+		T* CreatedKeyType = NULL;
+
+		const uint8 KeyID = InternalGetKeyID(KeyName, DontCheckParentKeys);
+		if (KeyID == InvalidKeyID && Parent == NULL)
+		{
+			FBlackboardEntry Entry;
+			Entry.EntryName = KeyName;
+
+			CreatedKeyType = ConstructObject<T>(T::StaticClass(), this);
+			Entry.KeyType = CreatedKeyType;		
+
+			Keys.Add(Entry);
+			MarkPackageDirty();
+		}
+		else if (KeyID != InvalidKeyID && Parent != NULL)
+		{
+			const int32 KeyIndex = KeyID - FirstKeyID;
+			Keys.RemoveAt(KeyIndex);
+			MarkPackageDirty();
+		}
+
+		return CreatedKeyType;
+	}
+
+	/** delegate called for every loaded blackboard asset
+	 *  meant for adding game specific persistent keys */
+	static FKeyUpdate OnUpdateKeys;
+
+protected:
+
+	enum EKeyLookupMode
+	{
+		CheckParentKeys,
+		DontCheckParentKeys,
+	};
+
+	/** @return first ID for keys of this asset (parent keys goes first) */
+	int32 FirstKeyID;
+
+	/** @return key ID from name */
+	uint8 InternalGetKeyID(const FName& KeyName, EKeyLookupMode LookupMode) const;
+
+	/** updates parent key cache for editor */
+	void UpdateParentKeys();
+};
diff --git a/Engine/Source/Runtime/AIModule/Classes/BehaviorTree/Composites/BTComposite_Selector.h b/Engine/Source/Runtime/AIModule/Classes/BehaviorTree/Composites/BTComposite_Selector.h
new file mode 100644
--- /dev/null
+++ b/Engine/Source/Runtime/AIModule/Classes/BehaviorTree/Composites/BTComposite_Selector.h
@@ -0,0 +1,13 @@
+// Copyright 1998-2014 Epic Games, Inc. All Rights Reserved.
+#pragma once
+
+#include "BehaviorTree/BTCompositeNode.h"
+#include "BTComposite_Selector.generated.h"
+
+UCLASS()
+class UBTComposite_Selector : public UBTCompositeNode
+{
+	GENERATED_UCLASS_BODY()
+
+	int32 GetNextChildHandler(struct FBehaviorTreeSearchData& SearchData, int32 PrevChild, EBTNodeResult::Type LastResult) const;
+};
diff --git a/Engine/Source/Runtime/AIModule/Classes/BehaviorTree/Composites/BTComposite_Sequence.h b/Engine/Source/Runtime/AIModule/Classes/BehaviorTree/Composites/BTComposite_Sequence.h
new file mode 100644
--- /dev/null
+++ b/Engine/Source/Runtime/AIModule/Classes/BehaviorTree/Composites/BTComposite_Sequence.h
@@ -0,0 +1,17 @@
+// Copyright 1998-2014 Epic Games, Inc. All Rights Reserved.
+#pragma once
+
+#include "BehaviorTree/BTCompositeNode.h"
+#include "BTComposite_Sequence.generated.h"
+
+UCLASS()
+class UBTComposite_Sequence : public UBTCompositeNode
+{
+	GENERATED_UCLASS_BODY()
+
+	int32 GetNextChildHandler(struct FBehaviorTreeSearchData& SearchData, int32 PrevChild, EBTNodeResult::Type LastResult) const;
+
+#if WITH_EDITOR
+	virtual bool CanAbortLowerPriority() const OVERRIDE;
+#endif
+};
diff --git a/Engine/Source/Runtime/AIModule/Classes/BehaviorTree/Composites/BTComposite_SimpleParallel.h b/Engine/Source/Runtime/AIModule/Classes/BehaviorTree/Composites/BTComposite_SimpleParallel.h
new file mode 100644
--- /dev/null
+++ b/Engine/Source/Runtime/AIModule/Classes/BehaviorTree/Composites/BTComposite_SimpleParallel.h
@@ -0,0 +1,66 @@
+// Copyright 1998-2014 Epic Games, Inc. All Rights Reserved.
+
+#pragma once
+#include "BehaviorTree/BTCompositeNode.h"
+#include "BTComposite_SimpleParallel.generated.h"
+
+namespace EBTParallelChild
+{
+	enum Type
+	{
+		MainTask,
+		BackgroundTree,
+		SearchIndicator,
+	};
+}
+
+UENUM()
+namespace EBTParallelMode
+{
+	// keep in sync with DescribeFinishMode
+	enum Type
+	{
+		AbortBackground UMETA(DisplayName="Immediate" , ToolTip="When main task finishes, immediately abort background tree."),
+		WaitForBackground UMETA(DisplayName="Delayed" , ToolTip="When main task finishes, wait for background tree to finish."),
+	};
+}
+
+struct FBTParallelMemory : public FBTCompositeMemory
+{
+	/** finish result of main task */
+	TEnumAsByte<EBTNodeResult::Type> MainTaskResult;
+
+	/** set when main task is running */
+	uint8 bMainTaskIsActive : 1;
+
+	/** try running background tree task even if main task has finished */
+	uint8 bForceBackgroundTree : 1;
+};
+
+UCLASS(MinimalAPI)
+class UBTComposite_SimpleParallel : public UBTCompositeNode
+{
+	GENERATED_UCLASS_BODY()
+
+	/** how background tree should be handled when main task finishes execution */
+	UPROPERTY(EditInstanceOnly, Category=Parallel)
+	TEnumAsByte<EBTParallelMode::Type> FinishMode;
+
+	/** handle child updates */
+	int32 GetNextChildHandler(struct FBehaviorTreeSearchData& SearchData, int32 PrevChild, EBTNodeResult::Type LastResult) const;
+
+	virtual void NotifyChildExecution(class UBehaviorTreeComponent* OwnerComp, uint8* NodeMemory, int32 ChildIdx, EBTNodeResult::Type& NodeResult) const OVERRIDE;
+	virtual void NotifyNodeDeactivation(struct FBehaviorTreeSearchData& SearchData, EBTNodeResult::Type& NodeResult) const OVERRIDE;
+	virtual bool CanPushSubtree(class UBehaviorTreeComponent* OwnerComp, uint8* NodeMemory, int32 ChildIdx) const OVERRIDE;
+	virtual uint16 GetInstanceMemorySize() const OVERRIDE;
+	virtual FString GetStaticDescription() const OVERRIDE;
+	virtual void DescribeRuntimeValues(const class UBehaviorTreeComponent* OwnerComp, uint8* NodeMemory, EBTDescriptionVerbosity::Type Verbosity, TArray<FString>& Values) const OVERRIDE;
+
+	/** helper for showing values of EBTParallelMode enum */
+	static FString DescribeFinishMode(EBTParallelMode::Type Mode);
+
+#if WITH_EDITOR
+	virtual bool CanAbortLowerPriority() const OVERRIDE;
+	virtual bool CanAbortSelf() const OVERRIDE;
+#endif // WITH_EDITOR
+};
diff --git a/Engine/Source/Runtime/AIModule/Classes/BehaviorTree/Decorators/BTDecorator_Blackboard.h b/Engine/Source/Runtime/AIModule/Classes/BehaviorTree/Decorators/BTDecorator_Blackboard.h
new file mode 100644
--- /dev/null
+++ b/Engine/Source/Runtime/AIModule/Classes/BehaviorTree/Decorators/BTDecorator_Blackboard.h
@@ -0,0 +1,73 @@
+// Copyright 1998-2014 Epic Games, Inc. All Rights Reserved.
+
+#pragma once
+#include "BehaviorTree/Blackboard/BlackboardKeyType.h"
+#include "BehaviorTree/BTDecorator.h"
+#include "BTDecorator_Blackboard.generated.h"
+
+/**
+ *  Decorator for accessing blackboard values
+ */
+
+UCLASS(HideCategories=(Condition))
+class AIMODULE_API UBTDecorator_Blackboard : public UBTDecorator_BlackboardBase
+{
+	GENERATED_UCLASS_BODY()
+
+	virtual bool CalculateRawConditionValue(class UBehaviorTreeComponent* OwnerComp, uint8* NodeMemory) const OVERRIDE;
+	virtual void OnBlackboardChange(const class UBlackboardComponent* Blackboard, uint8 ChangedKeyID) OVERRIDE;
+	virtual void DescribeRuntimeValues(const class UBehaviorTreeComponent* OwnerComp, uint8* NodeMemory, EBTDescriptionVerbosity::Type Verbosity, TArray<FString>& Values) const OVERRIDE;
+	virtual FString GetStaticDescription() const OVERRIDE;
+
+protected:
+
+	/** value for arithmetic operations */
+	UPROPERTY(Category=Blackboard, EditAnywhere, meta=(DisplayName="Key Value"))
+	int32 IntValue;
+
+	/** value for arithmetic operations */
+	UPROPERTY(Category=Blackboard, EditAnywhere, meta=(DisplayName="Key Value"))
+	float FloatValue;
+
+	/** value for string operations */
+	UPROPERTY(Category=Blackboard, EditAnywhere, meta=(DisplayName="Key Value"))
+	FString StringValue;
+
+	/** cached description */
+	UPROPERTY()
+	FString CachedDescription;
+
+	/** operation type */
+	UPROPERTY()
+	uint8 OperationType;
+
+#if WITH_EDITORONLY_DATA
+
+	UPROPERTY(Category=Blackboard, EditAnywhere, meta=(DisplayName="Key Query"))
+	TEnumAsByte<EBasicKeyOperation::Type> BasicOperation;
+
+	UPROPERTY(Category=Blackboard, EditAnywhere, meta=(DisplayName="Key Query"))
+	TEnumAsByte<EArithmeticKeyOperation::Type> ArithmeticOperation;
+
+	UPROPERTY(Category=Blackboard, EditAnywhere, meta=(DisplayName="Key Query"))
+	TEnumAsByte<ETextKeyOperation::Type> TextOperation;
+
+#endif
+
+#if WITH_EDITOR
+
+	/** describe decorator and cache it */
+	virtual void BuildDescription();
+
+	virtual void PostEditChangeProperty(FPropertyChangedEvent& PropertyChangedEvent) OVERRIDE;
+	virtual void InitializeFromAsset(class UBehaviorTree* Asset) OVERRIDE;
+
+#endif
+
+	/** take blackboard value and evaluate decorator's condition */
+	bool EvaluateOnBlackboard(const class UBlackboardComponent* BlackboardComp) const;
+
+#if CPP
+	friend class FBlackboardDecoratorDetails;
+#endif
+};
diff --git a/Engine/Source/Runtime/AIModule/Classes/BehaviorTree/Decorators/BTDecorator_BlackboardBase.h b/Engine/Source/Runtime/AIModule/Classes/BehaviorTree/Decorators/BTDecorator_BlackboardBase.h
new file mode 100644
--- /dev/null
+++ b/Engine/Source/Runtime/AIModule/Classes/BehaviorTree/Decorators/BTDecorator_BlackboardBase.h
@@ -0,0 +1,43 @@
+// Copyright 1998-2014 Epic Games, Inc. All Rights Reserved.
+#pragma once
+#include "BehaviorTree/BTDecorator.h"
+#include "BTDecorator_BlackboardBase.generated.h"
+
+UCLASS(Abstract)
+class AIMODULE_API UBTDecorator_BlackboardBase : public UBTDecorator
+{
+	GENERATED_UCLASS_BODY()
+
+	/** initialize any asset related data */
+	virtual void InitializeFromAsset(class UBehaviorTree* Asset) OVERRIDE;
+
+	/** notify about change in blackboard keys */
+	virtual void OnBlackboardChange(const class UBlackboardComponent* Blackboard, uint8 ChangedKeyID);
+
+	/** get name of selected blackboard key */
+	FName GetSelectedBlackboardKey() const;
+
+protected:
+
+	/** blackboard key selector */
+	UPROPERTY(EditAnywhere, Category=Blackboard)
+	struct FBlackboardKeySelector BlackboardKey;
+
+	FOnBlackboardChange BBKeyObserver;
+
+	virtual void PostInitProperties() OVERRIDE;
+
+	/** called when execution flow controller becomes active */
+	virtual void OnBecomeRelevant(UBehaviorTreeComponent* OwnerComp, uint8* NodeMemory) OVERRIDE;
+
+	/** called when execution flow controller becomes inactive */
+	virtual void OnCeaseRelevant(UBehaviorTreeComponent* OwnerComp, uint8* NodeMemory) OVERRIDE;
+};
+
+//////////////////////////////////////////////////////////////////////////
+// Inlines
+
+FORCEINLINE FName UBTDecorator_BlackboardBase::GetSelectedBlackboardKey() const
+{
+	return BlackboardKey.SelectedKeyName;
+}
diff --git a/Engine/Source/Runtime/AIModule/Classes/BehaviorTree/Decorators/BTDecorator_BlueprintBase.h b/Engine/Source/Runtime/AIModule/Classes/BehaviorTree/Decorators/BTDecorator_BlueprintBase.h
new file mode 100644
--- /dev/null
+++ b/Engine/Source/Runtime/AIModule/Classes/BehaviorTree/Decorators/BTDecorator_BlueprintBase.h
@@ -0,0 +1,117 @@
+// Copyright 1998-2014 Epic Games, Inc. All Rights Reserved.
+
+#pragma once
+#include "BehaviorTree/BTDecorator.h"
+#include "BTDecorator_BlueprintBase.generated.h"
+
+/**
+ *  Base class for blueprint based decorator nodes. Do NOT use it for creating native c++ classes!
+ *
+ *  Unlike task and services, decorator have two execution chains: 
+ *   ExecutionStart-ExecutionFinish and ObserverActivated-ObserverDeactivated
+ *  which makes automatic latent action cleanup impossible. Keep in mind, that
+ *  you HAVE TO verify is given chain is still active after resuming from any
+ *  latent action (like Delay, Timelines, etc).
+ *
+ *  Helper functions:
+ *  - IsDecoratorExecutionActive (true after ExecutionStart, until ExecutionFinish)
+ *  - IsDecoratorObserverActive (true after ObserverActivated, until ObserverDeactivated)
+ */
+
+UCLASS(Abstract, Blueprintable)
+class AIMODULE_API UBTDecorator_BlueprintBase : public UBTDecorator
+{
+	GENERATED_UCLASS_BODY()
+
+	/** initialize data about blueprint defined properties */
+	void InitializeProperties();
+
+	/** setup node name */
+	virtual void PostInitProperties() OVERRIDE;
+
+	/** notify about changes in blackboard */
+	void OnBlackboardChange(const class UBlackboardComponent* Blackboard, uint8 ChangedKeyID);
+
+	virtual FString GetStaticDescription() const OVERRIDE;
+	virtual void DescribeRuntimeValues(const class UBehaviorTreeComponent* OwnerComp, uint8* NodeMemory, EBTDescriptionVerbosity::Type Verbosity, TArray<FString>& Values) const OVERRIDE;
+	virtual bool CalculateRawConditionValue(class UBehaviorTreeComponent* OwnerComp, uint8* NodeMemory) const OVERRIDE;
+
+protected:
+
+	/** blackboard key names that should be observed */
+	TArray<FName> ObservedKeyNames;
+
+	/** properties with runtime values, stored only in class default object */
+	TArray<UProperty*> PropertyData;
+
+	/** show detailed information about properties */
+	UPROPERTY(EditInstanceOnly, Category=Description)
+	uint32 bShowPropertyDetails : 1;
+
+	/** temporary variable for ReceiveConditionCheck-FinishConditionCheck chain */
+	mutable uint32 CurrentCallResult : 1;
+
+	/** set if ReceiveTick is implemented by blueprint */
+	uint32 bImplementsReceiveTick : 1;
+
+	/** set if ReceiveExecutionStart is implemented by blueprint */
+	uint32 bImplementsReceiveExecutionStart : 1;
+
+	/** set if ReceiveExecutionFinish is implemented by blueprint */
+	uint32 bImplementsReceiveExecutionFinish : 1;
+
+	/** set if ReceiveObserverActivated is implemented by blueprint */
+	uint32 bImplementsReceiveObserverActivated : 1;
+
+	/** set if ReceiveObserverDeactivated is implemented by blueprint */
+	uint32 bImplementsReceiveObserverDeactivated : 1;
+
+	/** set if ReceiveConditionCheck is implemented by blueprint */
+	uint32 bImplementsReceiveConditionCheck : 1;
+
+	FOnBlackboardChange BBKeyObserver;
+
+	virtual void OnBecomeRelevant(UBehaviorTreeComponent* OwnerComp, uint8* NodeMemory) OVERRIDE;
+	virtual void OnCeaseRelevant(UBehaviorTreeComponent* OwnerComp, uint8* NodeMemory) OVERRIDE;
+	virtual void OnNodeActivation(struct FBehaviorTreeSearchData& SearchData) OVERRIDE;
+	virtual void OnNodeDeactivation(struct FBehaviorTreeSearchData& SearchData, EBTNodeResult::Type NodeResult) OVERRIDE;
+	virtual void TickNode(UBehaviorTreeComponent* OwnerComp, uint8* NodeMemory, float DeltaSeconds) OVERRIDE;
+
+	/** tick function */
+	UFUNCTION(BlueprintImplementableEvent)
+	virtual void ReceiveTick(AActor* OwnerActor, float DeltaSeconds);
+
+	/** called on execution of underlying node */
+	UFUNCTION(BlueprintImplementableEvent)
+	virtual void ReceiveExecutionStart(AActor* OwnerActor);
+
+	/** called when execution of underlying node is finished */
+	UFUNCTION(BlueprintImplementableEvent)
+	virtual void ReceiveExecutionFinish(AActor* OwnerActor, enum EBTNodeResult::Type NodeResult);
+
+	/** called when observer is activated (flow controller) */
+	UFUNCTION(BlueprintImplementableEvent)
+	virtual void ReceiveObserverActivated(AActor* OwnerActor);
+
+	/** called when observer is deactivated (flow controller) */
+	UFUNCTION(BlueprintImplementableEvent)
+	virtual void ReceiveObserverDeactivated(AActor* OwnerActor);
+
+	/** called when testing if underlying node can be executed, must call FinishConditionCheck */
+	UFUNCTION(BlueprintImplementableEvent)
+	virtual void ReceiveConditionCheck(AActor* OwnerActor);
+
+	/** finishes condition check */
+	UFUNCTION(BlueprintCallable, Category="AI|BehaviorTree")
+	void FinishConditionCheck(bool bAllowExecution);
+		
+	/** check if decorator is part of currently active branch */
+	UFUNCTION(BlueprintCallable, Category="AI|BehaviorTree")
+	bool IsDecoratorExecutionActive() const;
+
+	/** check if decorator's observer is currently active */
+	UFUNCTION(BlueprintCallable, Category="AI|BehaviorTree")
+	bool IsDecoratorObserverActive() const;
+
+	friend class FBehaviorBlueprintDetails;
+};
diff --git a/Engine/Source/Runtime/AIModule/Classes/BehaviorTree/Decorators/BTDecorator_CompareBBEntries.h b/Engine/Source/Runtime/AIModule/Classes/BehaviorTree/Decorators/BTDecorator_CompareBBEntries.h
new file mode 100644
--- /dev/null
+++ b/Engine/Source/Runtime/AIModule/Classes/BehaviorTree/Decorators/BTDecorator_CompareBBEntries.h
@@ -0,0 +1,50 @@
+// Copyright 1998-2014 Epic Games, Inc. All Rights Reserved.
+
+#pragma once
+#include "BehaviorTree/BTDecorator.h"
+#include "BTDecorator_CompareBBEntries.generated.h"
+
+UENUM()
+namespace EBlackBoardEntryComparison
+{
+	enum Type
+	{
+		Equal			UMETA(DisplayName="Is Equal To"),
+		NotEqual		UMETA(DisplayName="Is Not Equal To")
+	};
+}
+
+UCLASS(HideCategories=(Condition))
+class UBTDecorator_CompareBBEntries : public UBTDecorator
+{
+	GENERATED_UCLASS_BODY()
+
+protected:
+
+	/** operation type */
+	UPROPERTY(EditAnywhere, Category=Blackboard)
+	TEnumAsByte<EBlackBoardEntryComparison::Type> Operator;
+
+	/** blackboard key selector */
+	UPROPERTY(EditAnywhere, Category=Blackboard)
+	struct FBlackboardKeySelector BlackboardKeyA;
+
+	/** blackboard key selector */
+	UPROPERTY(EditAnywhere, Category=Blackboard)
+	struct FBlackboardKeySelector BlackboardKeyB;
+
+	FOnBlackboardChange BBKeyObserver;
+
+public:
+
+	virtual void PostInitProperties() OVERRIDE;
+
+	virtual bool CalculateRawConditionValue(class UBehaviorTreeComponent* OwnerComp, uint8* NodeMemory) const OVERRIDE;
+	virtual FString GetStaticDescription() const OVERRIDE;
+
+	virtual void InitializeFromAsset(class UBehaviorTree* Asset) OVERRIDE;
+
+	virtual void OnBlackboardChange(const class UBlackboardComponent* Blackboard, uint8 ChangedKeyID);	
+	virtual void OnBecomeRelevant(UBehaviorTreeComponent* OwnerComp, uint8* NodeMemory) OVERRIDE;
+	virtual void OnCeaseRelevant(UBehaviorTreeComponent* OwnerComp, uint8* NodeMemory) OVERRIDE;
+};
diff --git a/Engine/Source/Runtime/AIModule/Classes/BehaviorTree/Decorators/BTDecorator_ConeCheck.h b/Engine/Source/Runtime/AIModule/Classes/BehaviorTree/Decorators/BTDecorator_ConeCheck.h
new file mode 100644
--- /dev/null
+++ b/Engine/Source/Runtime/AIModule/Classes/BehaviorTree/Decorators/BTDecorator_ConeCheck.h
@@ -0,0 +1,54 @@
+// Copyright 1998-2014 Epic Games, Inc. All Rights Reserved.
+
+#pragma once
+#include "BehaviorTree/BTDecorator.h"
+#include "BTDecorator_ConeCheck.generated.h"
+
+struct FBTConeCheckDecoratorMemory
+{
+	bool bLastRawResult;
+};
+
+UCLASS()
+class AIMODULE_API UBTDecorator_ConeCheck : public UBTDecorator
+{
+	GENERATED_UCLASS_BODY()
+
+	typedef FBTConeCheckDecoratorMemory TNodeInstanceMemory;
+
+	/** max allowed time for execution of underlying node */
+	UPROPERTY(Category=Decorator, EditAnywhere)
+	float ConeHalfAngle;
+	
+	/** blackboard key selector */
+	UPROPERTY(EditAnywhere, Category=Blackboard)
+	struct FBlackboardKeySelector ConeOrigin;
+
+	/** "None" means "use ConeOrigin's direction" */
+	UPROPERTY(EditAnywhere, Category=Blackboard)
+	struct FBlackboardKeySelector ConeDirection;
+
+	/** blackboard key selector */
+	UPROPERTY(EditAnywhere, Category=Blackboard)
+	struct FBlackboardKeySelector Observed;
+	
+	float ConeHalfAngleDot;
+
+	virtual void InitializeFromAsset(class UBehaviorTree* Asset) OVERRIDE;
+	virtual uint16 GetInstanceMemorySize() const OVERRIDE;
+	virtual void DescribeRuntimeValues(const class UBehaviorTreeComponent* OwnerComp, uint8* NodeMemory, EBTDescriptionVerbosity::Type Verbosity, TArray<FString>& Values) const OVERRIDE;
+	virtual FString GetStaticDescription() const OVERRIDE;
+
+protected:
+
+	virtual bool CalculateRawConditionValue(class UBehaviorTreeComponent* OwnerComp, uint8* NodeMemory) const OVERRIDE;
+	virtual void OnBecomeRelevant(UBehaviorTreeComponent* OwnerComp, uint8* NodeMemory) OVERRIDE;
+	void OnBlackboardChange(const class UBlackboardComponent* Blackboard, uint8 ChangedKeyID);
+
+	virtual void TickNode(UBehaviorTreeComponent* OwnerComp, uint8* NodeMemory, float DeltaSeconds) OVERRIDE;
+	
+	bool CalculateDirection(const UBlackboardComponent* BlackboardComp, const FBlackboardKeySelector& Origin, const FBlackboardKeySelector& End, FVector& Direction) const;
+
+private:
+	bool CalcConditionImpl(class UBehaviorTreeComponent* OwnerComp, uint8* NodeMemory) const;
+};
diff --git a/Engine/Source/Runtime/AIModule/Classes/BehaviorTree/Decorators/BTDecorator_Cooldown.h b/Engine/Source/Runtime/AIModule/Classes/BehaviorTree/Decorators/BTDecorator_Cooldown.h
new file mode 100644
--- /dev/null
+++ b/Engine/Source/Runtime/AIModule/Classes/BehaviorTree/Decorators/BTDecorator_Cooldown.h
@@ -0,0 +1,36 @@
+// Copyright 1998-2014 Epic Games, Inc. All Rights Reserved.
+
+#pragma once
+#include "BehaviorTree/BTDecorator.h"
+#include "BTDecorator_Cooldown.generated.h"
+
+struct FBTCooldownDecoratorMemory
+{
+	float LastUseTimestamp;
+	uint8 RequestedRestart;
+};
+
+UCLASS(HideCategories=(Condition))
+class AIMODULE_API UBTDecorator_Cooldown : public UBTDecorator
+{
+	GENERATED_UCLASS_BODY()
+
+	/** max allowed time for execution of underlying node */
+	UPROPERTY(Category=Decorator, EditAnywhere)
+	float CoolDownTime;
+
+	// Begin UObject Interface
+	virtual void PostLoad() OVERRIDE;
+	// End UObject Interface
+
+	virtual bool CalculateRawConditionValue(class UBehaviorTreeComponent* OwnerComp, uint8* NodeMemory) const OVERRIDE;
+	virtual void InitializeMemory(class UBehaviorTreeComponent* OwnerComp, uint8* NodeMemory) const OVERRIDE;
+	virtual uint16 GetInstanceMemorySize() const OVERRIDE;
+	virtual void DescribeRuntimeValues(const class UBehaviorTreeComponent* OwnerComp, uint8* NodeMemory, EBTDescriptionVerbosity::Type Verbosity, TArray<FString>& Values) const OVERRIDE;
+	virtual FString GetStaticDescription() const OVERRIDE;
+
+protected:
+
+	virtual void OnNodeDeactivation(struct FBehaviorTreeSearchData& SearchData, EBTNodeResult::Type NodeResult) OVERRIDE;
+	virtual void TickNode(class UBehaviorTreeComponent* OwnerComp, uint8* NodeMemory, float DeltaSeconds) OVERRIDE;
+};
diff --git a/Engine/Source/Runtime/AIModule/Classes/BehaviorTree/Decorators/BTDecorator_DoesPathExist.h b/Engine/Source/Runtime/AIModule/Classes/BehaviorTree/Decorators/BTDecorator_DoesPathExist.h
new file mode 100644
--- /dev/null
+++ b/Engine/Source/Runtime/AIModule/Classes/BehaviorTree/Decorators/BTDecorator_DoesPathExist.h
@@ -0,0 +1,49 @@
+// Copyright 1998-2014 Epic Games, Inc. All Rights Reserved.
+
+#pragma once
+#include "BehaviorTree/BTDecorator.h"
+#include "BTDecorator_DoesPathExist.generated.h"
+
+UENUM()
+namespace EPathExistanceQueryType
+{
+	enum Type
+	{
+		NavmeshRaycast2D UMETA(ToolTip = "Really Fast"),
+		HierarchicalQuery UMETA(ToolTip = "Fast"),
+		RegularPathFinding UMETA(ToolTip = "Slow"),
+	};
+}
+
+UCLASS()
+class AIMODULE_API UBTDecorator_DoesPathExist : public UBTDecorator
+{
+	GENERATED_UCLASS_BODY()
+
+protected:
+
+	/** blackboard key selector */
+	UPROPERTY(EditAnywhere, Category=Condition)
+	struct FBlackboardKeySelector BlackboardKeyA;
+
+	/** blackboard key selector */
+	UPROPERTY(EditAnywhere, Category=Condition)
+	struct FBlackboardKeySelector BlackboardKeyB;
+
+public:
+
+	// deprecated, set value of blackboard key A on initialization
+	UPROPERTY()
+	uint32 bUseSelf:1;
+
+	UPROPERTY(EditAnywhere, Category=Condition)
+	TEnumAsByte<EPathExistanceQueryType::Type> PathQueryType;
+
+	/** "None" will result in default filter being used */
+	UPROPERTY(Category=Node, EditAnywhere)
+	TSubclassOf<class UNavigationQueryFilter> FilterClass;
+
+	virtual bool CalculateRawConditionValue(class UBehaviorTreeComponent* OwnerComp, uint8* NodeMemory) const OVERRIDE;
+	virtual FString GetStaticDescription() const OVERRIDE;
+	virtual void InitializeFromAsset(class UBehaviorTree* Asset) OVERRIDE;
+};
diff --git a/Engine/Source/Runtime/AIModule/Classes/BehaviorTree/Decorators/BTDecorator_KeepInCone.h b/Engine/Source/Runtime/AIModule/Classes/BehaviorTree/Decorators/BTDecorator_KeepInCone.h
new file mode 100644
--- /dev/null
+++ b/Engine/Source/Runtime/AIModule/Classes/BehaviorTree/Decorators/BTDecorator_KeepInCone.h
@@ -0,0 +1,52 @@
+// Copyright 1998-2014 Epic Games, Inc. All Rights Reserved.
+
+#pragma once
+#include "BehaviorTree/BTDecorator.h"
+#include "BTDecorator_KeepInCone.generated.h"
+
+struct FBTKeepInConeDecoratorMemory
+{
+	FVector InitialDirection;
+};
+
+UCLASS(HideCategories=(Condition))
+class AIMODULE_API UBTDecorator_KeepInCone : public UBTDecorator
+{
+	GENERATED_UCLASS_BODY()
+
+	typedef FBTKeepInConeDecoratorMemory TNodeInstanceMemory;
+
+	/** max allowed time for execution of underlying node */
+	UPROPERTY(Category=Decorator, EditAnywhere)
+	float ConeHalfAngle;
+	
+	/** blackboard key selector */
+	UPROPERTY(EditAnywhere, Category=Blackboard)
+	struct FBlackboardKeySelector ConeOrigin;
+
+	/** blackboard key selector */
+	UPROPERTY(EditAnywhere, Category=Blackboard)
+	struct FBlackboardKeySelector Observed;
+
+	// deprecated, set value of ConeOrigin on initialization
+	UPROPERTY()
+	uint32 bUseSelfAsOrigin:1;
+
+	// deprecated, set value of Observed on initialization
+	UPROPERTY()
+	uint32 bUseSelfAsObserved:1;
+	
+	float ConeHalfAngleDot;
+
+	virtual void InitializeFromAsset(class UBehaviorTree* Asset) OVERRIDE;
+	virtual uint16 GetInstanceMemorySize() const OVERRIDE;
+	virtual void DescribeRuntimeValues(const class UBehaviorTreeComponent* OwnerComp, uint8* NodeMemory, EBTDescriptionVerbosity::Type Verbosity, TArray<FString>& Values) const OVERRIDE;
+	virtual FString GetStaticDescription() const OVERRIDE;
+
+protected:
+
+	virtual void OnBecomeRelevant(UBehaviorTreeComponent* OwnerComp, uint8* NodeMemory) OVERRIDE;
+	virtual void TickNode(UBehaviorTreeComponent* OwnerComp, uint8* NodeMemory, float DeltaSeconds) OVERRIDE;
+
+	bool CalculateCurrentDirection(const UBehaviorTreeComponent* OwnerComp, FVector& Direction) const;
+};
diff --git a/Engine/Source/Runtime/AIModule/Classes/BehaviorTree/Decorators/BTDecorator_Loop.h b/Engine/Source/Runtime/AIModule/Classes/BehaviorTree/Decorators/BTDecorator_Loop.h
new file mode 100644
--- /dev/null
+++ b/Engine/Source/Runtime/AIModule/Classes/BehaviorTree/Decorators/BTDecorator_Loop.h
@@ -0,0 +1,31 @@
+// Copyright 1998-2014 Epic Games, Inc. All Rights Reserved.
+#pragma once
+#include "BehaviorTree/BTDecorator.h"
+#include "BTDecorator_Loop.generated.h"
+
+struct FBTLoopDecoratorMemory
+{
+	uint8 RemainingExecutions;
+};
+
+UCLASS(HideCategories=(Condition))
+class AIMODULE_API UBTDecorator_Loop : public UBTDecorator
+{
+	GENERATED_UCLASS_BODY()
+
+	/** number of executions */
+	UPROPERTY(Category=Decorator, EditAnywhere)
+	int32 NumLoops;
+
+	/** infinite loop */
+	UPROPERTY(Category=Decorator, EditAnywhere)
+	bool bInfiniteLoop;
+
+	virtual uint16 GetInstanceMemorySize() const OVERRIDE;
+	virtual void DescribeRuntimeValues(const class UBehaviorTreeComponent* OwnerComp, uint8* NodeMemory, EBTDescriptionVerbosity::Type Verbosity, TArray<FString>& Values) const OVERRIDE;
+	virtual FString GetStaticDescription() const OVERRIDE;
+
+protected:
+
+	virtual void OnNodeActivation(struct FBehaviorTreeSearchData& SearchData) OVERRIDE;
+};
diff --git a/Engine/Source/Runtime/AIModule/Classes/BehaviorTree/Decorators/BTDecorator_Optional.h b/Engine/Source/Runtime/AIModule/Classes/BehaviorTree/Decorators/BTDecorator_Optional.h
new file mode 100644
--- /dev/null
+++ b/Engine/Source/Runtime/AIModule/Classes/BehaviorTree/Decorators/BTDecorator_Optional.h
@@ -0,0 +1,15 @@
+// Copyright 1998-2014 Epic Games, Inc. All Rights Reserved.
+
+#pragma once
+#include "BehaviorTree/BTDecorator.h"
+#include "BTDecorator_Optional.generated.h"
+
+UCLASS(HideCategories=(Condition), meta=(DeprecatedNode,DeprecationMessage="Please use Force Success decorator instead."))
+class UBTDecorator_Optional : public UBTDecorator
+{
+	GENERATED_UCLASS_BODY()
+
+protected:
+
+	virtual void OnNodeProcessed(struct FBehaviorTreeSearchData& SearchData, EBTNodeResult::Type& NodeResult);
+};
diff --git a/Engine/Source/Runtime/AIModule/Classes/BehaviorTree/Decorators/BTDecorator_ReachedMoveGoal.h b/Engine/Source/Runtime/AIModule/Classes/BehaviorTree/Decorators/BTDecorator_ReachedMoveGoal.h
new file mode 100644
--- /dev/null
+++ b/Engine/Source/Runtime/AIModule/Classes/BehaviorTree/Decorators/BTDecorator_ReachedMoveGoal.h
@@ -0,0 +1,13 @@
+// Copyright 1998-2014 Epic Games, Inc. All Rights Reserved.
+
+#pragma once
+#include "BehaviorTree/BTDecorator.h"
+#include "BTDecorator_ReachedMoveGoal.generated.h"
+
+UCLASS()
+class AIMODULE_API UBTDecorator_ReachedMoveGoal : public UBTDecorator
+{
+	GENERATED_UCLASS_BODY()
+
+	virtual bool CalculateRawConditionValue(class UBehaviorTreeComponent* OwnerComp, uint8* NodeMemory) const OVERRIDE;
+};
diff --git a/Engine/Source/Runtime/AIModule/Classes/BehaviorTree/Decorators/BTDecorator_TimeLimit.h b/Engine/Source/Runtime/AIModule/Classes/BehaviorTree/Decorators/BTDecorator_TimeLimit.h
new file mode 100644
--- /dev/null
+++ b/Engine/Source/Runtime/AIModule/Classes/BehaviorTree/Decorators/BTDecorator_TimeLimit.h
@@ -0,0 +1,22 @@
+// Copyright 1998-2014 Epic Games, Inc. All Rights Reserved.
+#pragma once
+#include "BehaviorTree/BTDecorator.h"
+#include "BTDecorator_TimeLimit.generated.h"
+
+UCLASS(HideCategories=(Condition))
+class AIMODULE_API UBTDecorator_TimeLimit : public UBTDecorator
+{
+	GENERATED_UCLASS_BODY()
+		
+	/** max allowed time for execution of underlying node */
+	UPROPERTY(Category=Decorator, EditAnywhere)
+	float TimeLimit;
+
+	virtual void DescribeRuntimeValues(const class UBehaviorTreeComponent* OwnerComp, uint8* NodeMemory, EBTDescriptionVerbosity::Type Verbosity, TArray<FString>& Values) const OVERRIDE;
+	virtual FString GetStaticDescription() const OVERRIDE;
+
+protected:
+
+	virtual void OnBecomeRelevant(UBehaviorTreeComponent* OwnerComp, uint8* NodeMemory) OVERRIDE;
+	virtual void TickNode(UBehaviorTreeComponent* OwnerComp, uint8* NodeMemory, float DeltaSeconds) OVERRIDE;
+};
diff --git a/Engine/Source/Runtime/AIModule/Classes/BehaviorTree/Services/BTService_BlackboardBase.h b/Engine/Source/Runtime/AIModule/Classes/BehaviorTree/Services/BTService_BlackboardBase.h
new file mode 100644
--- /dev/null
+++ b/Engine/Source/Runtime/AIModule/Classes/BehaviorTree/Services/BTService_BlackboardBase.h
@@ -0,0 +1,30 @@
+// Copyright 1998-2014 Epic Games, Inc. All Rights Reserved.
+#pragma once
+#include "BehaviorTree/BTService.h"
+#include "BTService_BlackboardBase.generated.h"
+
+UCLASS(Abstract)
+class AIMODULE_API UBTService_BlackboardBase : public UBTService
+{
+	GENERATED_UCLASS_BODY()
+
+	/** initialize any asset related data */
+	virtual void InitializeFromAsset(class UBehaviorTree* Asset) OVERRIDE;
+
+	/** get name of selected blackboard key */
+	FName GetSelectedBlackboardKey() const;
+
+protected:
+
+	/** blackboard key selector */
+	UPROPERTY(EditAnywhere, Category=Blackboard)
+	struct FBlackboardKeySelector BlackboardKey;
+};
+
+//////////////////////////////////////////////////////////////////////////
+// Inlines
+
+FORCEINLINE FName UBTService_BlackboardBase::GetSelectedBlackboardKey() const
+{
+	return BlackboardKey.SelectedKeyName;
+}
diff --git a/Engine/Source/Runtime/AIModule/Classes/BehaviorTree/Services/BTService_BlueprintBase.h b/Engine/Source/Runtime/AIModule/Classes/BehaviorTree/Services/BTService_BlueprintBase.h
new file mode 100644
--- /dev/null
+++ b/Engine/Source/Runtime/AIModule/Classes/BehaviorTree/Services/BTService_BlueprintBase.h
@@ -0,0 +1,64 @@
+// Copyright 1998-2014 Epic Games, Inc. All Rights Reserved.
+
+#pragma once
+#include "BehaviorTree/BTService.h"
+#include "BTService_BlueprintBase.generated.h"
+
+/**
+ *  Base class for blueprint based service nodes. Do NOT use it for creating native c++ classes!
+ *
+ *  When service receives Deactivation event, all latent actions associated this instance are being removed.
+ *  This prevents from resuming activity started by Activation, but does not handle external events.
+ *  Please use them safely (unregister at abort) and call IsServiceActive() when in doubt.
+ */
+
+UCLASS(Abstract, Blueprintable)
+class AIMODULE_API UBTService_BlueprintBase : public UBTService
+{
+	GENERATED_UCLASS_BODY()
+
+	virtual void PostInitProperties() OVERRIDE;
+
+	virtual FString GetStaticDescription() const OVERRIDE;
+	virtual void DescribeRuntimeValues(const class UBehaviorTreeComponent* OwnerComp, uint8* NodeMemory, EBTDescriptionVerbosity::Type Verbosity, TArray<FString>& Values) const OVERRIDE;
+
+protected:
+
+	/** properties with runtime values, stored only in class default object */
+	TArray<UProperty*> PropertyData;
+
+	/** show detailed information about properties */
+	UPROPERTY(EditInstanceOnly, Category=Description)
+	uint32 bShowPropertyDetails : 1;
+
+	/** set if ReceiveTick is implemented by blueprint */
+	uint32 bImplementsReceiveTick : 1;
+	
+	/** set if ReceiveActivation is implemented by blueprint */
+	uint32 bImplementsReceiveActivation : 1;
+
+	/** set if ReceiveDeactivation is implemented by blueprint */
+	uint32 bImplementsReceiveDeactivation : 1;
+
+	virtual void OnBecomeRelevant(class UBehaviorTreeComponent* OwnerComp, uint8* NodeMemory) OVERRIDE;
+	virtual void OnCeaseRelevant(class UBehaviorTreeComponent* OwnerComp, uint8* NodeMemory) OVERRIDE;
+	virtual void TickNode(class UBehaviorTreeComponent* OwnerComp, uint8* NodeMemory, float DeltaSeconds) OVERRIDE;
+
+	/** tick function */
+	UFUNCTION(BlueprintImplementableEvent)
+	virtual void ReceiveTick(AActor* OwnerActor, float DeltaSeconds);
+
+	/** service became active */
+	UFUNCTION(BlueprintImplementableEvent)
+	virtual void ReceiveActivation(AActor* OwnerActor);
+
+	/** service became inactive */
+	UFUNCTION(BlueprintImplementableEvent)
+	virtual void ReceiveDeactivation(AActor* OwnerActor);
+
+	/** check if service is currently being active */
+	UFUNCTION(BlueprintCallable, Category="AI|BehaviorTree")
+	bool IsServiceActive() const;
+
+	friend class FBehaviorBlueprintDetails;
+};
diff --git a/Engine/Source/Runtime/AIModule/Classes/BehaviorTree/Services/BTService_DefaultFocus.h b/Engine/Source/Runtime/AIModule/Classes/BehaviorTree/Services/BTService_DefaultFocus.h
new file mode 100644
--- /dev/null
+++ b/Engine/Source/Runtime/AIModule/Classes/BehaviorTree/Services/BTService_DefaultFocus.h
@@ -0,0 +1,30 @@
+// Copyright 1998-2014 Epic Games, Inc. All Rights Reserved.
+
+#pragma once
+#include "BehaviorTree/Services/BTService_BlackboardBase.h"
+#include "BTService_DefaultFocus.generated.h"
+
+struct FBTDefaultFocusMemory
+{
+	AActor* FocusActorSet;
+	FVector FocusLocationSet;
+	bool bActorSet;
+
+	void Reset()
+	{
+		FMemory::MemZero(*this);
+	}
+};
+
+UCLASS(hidecategories=(Service))
+class UBTService_DefaultFocus : public UBTService_BlackboardBase
+{
+	GENERATED_UCLASS_BODY()
+
+	virtual uint16 GetInstanceMemorySize() const OVERRIDE;
+	virtual void OnBecomeRelevant(class UBehaviorTreeComponent* OwnerComp, uint8* NodeMemory) OVERRIDE;
+	virtual void OnCeaseRelevant(class UBehaviorTreeComponent* OwnerComp, uint8* NodeMemory) OVERRIDE;
+
+	virtual FString GetStaticDescription() const OVERRIDE;
+	virtual void DescribeRuntimeValues(const class UBehaviorTreeComponent* OwnerComp, uint8* NodeMemory, EBTDescriptionVerbosity::Type Verbosity, TArray<FString>& Values) const OVERRIDE;
+};
diff --git a/Engine/Source/Runtime/AIModule/Classes/BehaviorTree/Tasks/BTTask_BlackboardBase.h b/Engine/Source/Runtime/AIModule/Classes/BehaviorTree/Tasks/BTTask_BlackboardBase.h
new file mode 100644
--- /dev/null
+++ b/Engine/Source/Runtime/AIModule/Classes/BehaviorTree/Tasks/BTTask_BlackboardBase.h
@@ -0,0 +1,32 @@
+// Copyright 1998-2014 Epic Games, Inc. All Rights Reserved.
+
+#pragma once
+
+#include "BehaviorTree/BTTaskNode.h"
+#include "BTTask_BlackboardBase.generated.h"
+
+UCLASS(Abstract)
+class AIMODULE_API UBTTask_BlackboardBase : public UBTTaskNode
+{
+	GENERATED_UCLASS_BODY()
+
+	/** initialize any asset related data */
+	virtual void InitializeFromAsset(class UBehaviorTree* Asset) OVERRIDE;
+
+	/** get name of selected blackboard key */
+	FName GetSelectedBlackboardKey() const;
+
+protected:
+
+	/** blackboard key selector */
+	UPROPERTY(EditAnywhere, Category=Blackboard)
+	struct FBlackboardKeySelector BlackboardKey;
+};
+
+//////////////////////////////////////////////////////////////////////////
+// Inlines
+
+FORCEINLINE FName UBTTask_BlackboardBase::GetSelectedBlackboardKey() const
+{
+	return BlackboardKey.SelectedKeyName;
+}
diff --git a/Engine/Source/Runtime/AIModule/Classes/BehaviorTree/Tasks/BTTask_BlueprintBase.h b/Engine/Source/Runtime/AIModule/Classes/BehaviorTree/Tasks/BTTask_BlueprintBase.h
new file mode 100644
--- /dev/null
+++ b/Engine/Source/Runtime/AIModule/Classes/BehaviorTree/Tasks/BTTask_BlueprintBase.h
@@ -0,0 +1,89 @@
+// Copyright 1998-2014 Epic Games, Inc. All Rights Reserved.
+
+#pragma once
+#include "BTTask_BlueprintBase.generated.h"
+
+/**
+ *  Base class for blueprint based task nodes. Do NOT use it for creating native c++ classes!
+ *
+ *  When task receives Abort event, all latent actions associated this instance are being removed.
+ *  This prevents from resuming activity started by Execute, but does not handle external events.
+ *  Please use them safely (unregister at abort) and call IsTaskExecuting() when in doubt.
+ */
+
+UCLASS(Abstract, Blueprintable)
+class AIMODULE_API UBTTask_BlueprintBase : public UBTTaskNode
+{
+	GENERATED_UCLASS_BODY()
+
+	/** setup node name */
+	virtual void PostInitProperties() OVERRIDE;
+
+	virtual EBTNodeResult::Type ExecuteTask(class UBehaviorTreeComponent* OwnerComp, uint8* NodeMemory) OVERRIDE;
+	virtual EBTNodeResult::Type AbortTask(class UBehaviorTreeComponent* OwnerComp, uint8* NodeMemory) OVERRIDE;
+
+	virtual FString GetStaticDescription() const OVERRIDE;
+	virtual void DescribeRuntimeValues(const class UBehaviorTreeComponent* OwnerComp, uint8* NodeMemory, EBTDescriptionVerbosity::Type Verbosity, TArray<FString>& Values) const OVERRIDE;
+
+protected:
+
+	/** temporary variable for ReceiveExecute(Abort)-FinishExecute(Abort) chain */
+	mutable TEnumAsByte<EBTNodeResult::Type> CurrentCallResult;
+
+	/** properties that should be copied */
+	TArray<UProperty*> PropertyData;
+
+	/** show detailed information about properties */
+	UPROPERTY(EditInstanceOnly, Category=Description)
+	uint32 bShowPropertyDetails : 1;
+
+	/** set if ReceiveTick is implemented by blueprint */
+	uint32 bImplementsReceiveTick : 1;
+
+	/** set if ReceiveExecute is implemented by blueprint */
+	uint32 bImplementsReceiveExecute : 1;
+
+	/** set if ReceiveAbort is implemented by blueprint */
+	uint32 bImplementsReceiveAbort : 1;
+
+	/** if set, execution is inside blueprint's ReceiveExecute(Abort) event
+	  * FinishExecute(Abort) function should store their result in CurrentCallResult variable */
+	mutable uint32 bStoreFinishResult : 1;
+
+	/** entry point, task will stay active until FinishExecute is called */
+	UFUNCTION(BlueprintImplementableEvent)
+	virtual void ReceiveExecute(AActor* OwnerActor);
+
+	/** if blueprint graph contains this event, task will stay active until FinishAbort is called */
+	UFUNCTION(BlueprintImplementableEvent)
+	virtual void ReceiveAbort(AActor* OwnerActor);
+
+	/** tick function */
+	UFUNCTION(BlueprintImplementableEvent)
+	virtual void ReceiveTick(AActor* OwnerActor, float DeltaSeconds);
+
+	/** finishes task execution with Success or Fail result */
+	UFUNCTION(BlueprintCallable, Category="AI|BehaviorTree")
+	void FinishExecute(bool bSuccess);
+
+	/** aborts task execution */
+	UFUNCTION(BlueprintCallable, Category="AI|BehaviorTree")
+	void FinishAbort();
+
+	/** task execution will be finished (with result \'Success\') after receiving specified message */
+	UFUNCTION(BlueprintCallable, Category="AI|BehaviorTree")
+	void SetFinishOnMessage(FName MessageName);
+
+	/** task execution will be finished (with result \'Success\') after receiving specified message with indicated ID */
+	UFUNCTION(BlueprintCallable, Category="AI|BehaviorTree")
+	void SetFinishOnMessageWithId(FName MessageName, int32 RequestID = -1);
+
+	/** check if task is currently being executed */
+	UFUNCTION(BlueprintCallable, Category="AI|BehaviorTree")
+	bool IsTaskExecuting() const;
+
+	/** ticks this task */
+	virtual void TickTask(class UBehaviorTreeComponent* OwnerComp, uint8* NodeMemory, float DeltaSeconds) OVERRIDE;
+
+	friend class FBehaviorBlueprintDetails;
+};
diff --git a/Engine/Source/Runtime/AIModule/Classes/BehaviorTree/Tasks/BTTask_MakeNoise.h b/Engine/Source/Runtime/AIModule/Classes/BehaviorTree/Tasks/BTTask_MakeNoise.h
new file mode 100644
--- /dev/null
+++ b/Engine/Source/Runtime/AIModule/Classes/BehaviorTree/Tasks/BTTask_MakeNoise.h
@@ -0,0 +1,16 @@
+// Copyright 1998-2014 Epic Games, Inc. All Rights Reserved.
+
+#pragma once
+#include "BTTask_MakeNoise.generated.h"
+
+UCLASS()
+class UBTTask_MakeNoise : public UBTTaskNode
+{
+	GENERATED_UCLASS_BODY()
+
+	/** Loudnes of generated noise */
+	UPROPERTY(Category=Node, EditAnywhere, meta=(ClampMin = "0", ClampMax = "1", UIMin = "0", UIMax = "1"))
+	float Loudnes;
+
+	virtual EBTNodeResult::Type ExecuteTask(UBehaviorTreeComponent* OwnerComp, uint8* NodeMemory) OVERRIDE;
+};
diff --git a/Engine/Source/Runtime/AIModule/Classes/BehaviorTree/Tasks/BTTask_MoveDirectlyToward.h b/Engine/Source/Runtime/AIModule/Classes/BehaviorTree/Tasks/BTTask_MoveDirectlyToward.h
new file mode 100644
--- /dev/null
+++ b/Engine/Source/Runtime/AIModule/Classes/BehaviorTree/Tasks/BTTask_MoveDirectlyToward.h
@@ -0,0 +1,39 @@
+// Copyright 1998-2014 Epic Games, Inc. All Rights Reserved.
+
+#pragma once
+#include "BehaviorTree/Tasks/BTTask_BlackboardBase.h"
+#include "BTTask_MoveDirectlyToward.generated.h"
+
+struct FBTMoveDirectlyTowardMemory
+{
+	/** Move request ID */
+	FAIRequestID MoveRequestID;
+};
+
+UCLASS(config=Game, Meta=(
+	Tooltip="Moves the AI pawn toward the specified Actor or Location (Vector) blackboard entry in a straight line, without regard to any navigation system.  If you need the AI to navigate, use the \"Move To\" node instead."))
+class AIMODULE_API UBTTask_MoveDirectlyToward : public UBTTask_BlackboardBase
+{
+	GENERATED_UCLASS_BODY()
+
+	UPROPERTY(config, Category=Node, EditAnywhere, meta=(ClampMin = "0.0"))
+	float AcceptableRadius;
+
+	/** set to true will result in not updating move destination in 
+	 *	case where goal is an Actor that can change 
+	 *	his location while task is being performed */
+	UPROPERTY(Category=Node, EditAnywhere)
+	uint32 bForceMoveToLocation:1;
+
+	UPROPERTY(Category=Node, EditAnywhere)
+	uint32 bProjectVectorGoalToNavigation:1;
+
+	UPROPERTY(Category=Node, EditAnywhere)
+	uint32 bAllowStrafe : 1;
+
+	virtual EBTNodeResult::Type ExecuteTask(class UBehaviorTreeComponent* OwnerComp, uint8* NodeMemory) OVERRIDE;
+	virtual EBTNodeResult::Type AbortTask(class UBehaviorTreeComponent* OwnerComp, uint8* NodeMemory) OVERRIDE;
+	virtual uint16 GetInstanceMemorySize() const OVERRIDE;
+	virtual void DescribeRuntimeValues(const class UBehaviorTreeComponent* OwnerComp, uint8* NodeMemory, EBTDescriptionVerbosity::Type Verbosity, TArray<FString>& Values) const OVERRIDE;
+	virtual FString GetStaticDescription() const OVERRIDE;
+};
diff --git a/Engine/Source/Runtime/AIModule/Classes/BehaviorTree/Tasks/BTTask_MoveTo.h b/Engine/Source/Runtime/AIModule/Classes/BehaviorTree/Tasks/BTTask_MoveTo.h
new file mode 100644
--- /dev/null
+++ b/Engine/Source/Runtime/AIModule/Classes/BehaviorTree/Tasks/BTTask_MoveTo.h
@@ -0,0 +1,43 @@
+// Copyright 1998-2014 Epic Games, Inc. All Rights Reserved.
+
+#pragma once
+#include "BehaviorTree/Tasks/BTTask_BlackboardBase.h"
+#include "BTTask_MoveTo.generated.h"
+
+struct FBTMoveToTaskMemory
+{
+	/** Move request ID */
+	FAIRequestID MoveRequestID;
+
+	uint8 bWaitingForPath : 1;
+};
+
+UCLASS(config=Game)
+class AIMODULE_API UBTTask_MoveTo : public UBTTask_BlackboardBase
+{
+	GENERATED_UCLASS_BODY()
+
+	UPROPERTY(config, Category=Node, EditAnywhere, meta=(ClampMin = "0.0"))
+	float AcceptableRadius;
+
+	/** "None" will result in default filter being used */
+	UPROPERTY(Category=Node, EditAnywhere)
+	TSubclassOf<class UNavigationQueryFilter> FilterClass;
+
+	UPROPERTY(Category=Node, EditAnywhere)
+	uint32 bAllowStrafe : 1;
+	
+	virtual EBTNodeResult::Type ExecuteTask(class UBehaviorTreeComponent* OwnerComp, uint8* NodeMemory) OVERRIDE;
+	virtual EBTNodeResult::Type AbortTask(class UBehaviorTreeComponent* OwnerComp, uint8* NodeMemory) OVERRIDE;
+	virtual void TickTask(class UBehaviorTreeComponent* OwnerComp, uint8* NodeMemory, float DeltaSeconds) OVERRIDE;
+	virtual uint16 GetInstanceMemorySize() const OVERRIDE;
+
+	virtual void OnMessage(class UBehaviorTreeComponent* OwnerComp, uint8* NodeMemory, FName Message, int32 RequestID, bool bSuccess) OVERRIDE;
+
+	virtual void DescribeRuntimeValues(const class UBehaviorTreeComponent* OwnerComp, uint8* NodeMemory, EBTDescriptionVerbosity::Type Verbosity, TArray<FString>& Values) const OVERRIDE;
+	virtual FString GetStaticDescription() const OVERRIDE;
+
+protected:
+
+	EBTNodeResult::Type PerformMoveTask(class UBehaviorTreeComponent* OwnerComp, uint8* NodeMemory);
+};
diff --git a/Engine/Source/Runtime/AIModule/Classes/BehaviorTree/Tasks/BTTask_PlaySound.h b/Engine/Source/Runtime/AIModule/Classes/BehaviorTree/Tasks/BTTask_PlaySound.h
new file mode 100644
--- /dev/null
+++ b/Engine/Source/Runtime/AIModule/Classes/BehaviorTree/Tasks/BTTask_PlaySound.h
@@ -0,0 +1,17 @@
+// Copyright 1998-2014 Epic Games, Inc. All Rights Reserved.
+
+#pragma once
+#include "BTTask_PlaySound.generated.h"
+
+UCLASS()
+class UBTTask_PlaySound : public UBTTaskNode
+{
+	GENERATED_UCLASS_BODY()
+
+	/** CUE to play */
+	UPROPERTY(Category=Node, EditAnywhere)
+	class USoundCue* SoundToPlay;
+
+	virtual EBTNodeResult::Type ExecuteTask(UBehaviorTreeComponent* OwnerComp, uint8* NodeMemory) OVERRIDE;
+	virtual FString GetStaticDescription() const OVERRIDE;
+};
