Extracting AI code from the engine into a separate module, pass 1 #UE4

diff --git a/Engine/Source/Runtime/AIModule/Classes/BehaviorTree/Tasks/BTTask_RunBehavior.h b/Engine/Source/Runtime/AIModule/Classes/BehaviorTree/Tasks/BTTask_RunBehavior.h
new file mode 100644
--- /dev/null
+++ b/Engine/Source/Runtime/AIModule/Classes/BehaviorTree/Tasks/BTTask_RunBehavior.h
@@ -0,0 +1,57 @@
+// Copyright 1998-2014 Epic Games, Inc. All Rights Reserved.
+
+#pragma once
+#include "BehaviorTree/BehaviorTree.h"
+#include "BehaviorTree/BTTaskNode.h"
+#include "BTTask_RunBehavior.generated.h"
+
+/**
+ * RunBehavior task allows pushing subtrees on execution stack.
+ * Subtree asset can't be changed in runtime! 
+ *
+ * This limitation is caused by support for subtree's root level decorators,
+ * which are injected into parent tree, and structure of running tree
+ * cannot be modified in runtime (see: BTNode: ExecutionIndex, MemoryOffset)
+ *
+ * Dynamic subtrees can be implemented, but at the cost of some features:
+ * - no root level decorators
+ * - no or limited editor preview
+ */
+
+UCLASS(MinimalAPI)
+class UBTTask_RunBehavior : public UBTTaskNode
+{
+	GENERATED_UCLASS_BODY()
+
+	virtual EBTNodeResult::Type ExecuteTask(UBehaviorTreeComponent* OwnerComp, uint8* NodeMemory) OVERRIDE;
+	virtual uint16 GetInstanceMemorySize() const OVERRIDE;
+	virtual FString GetStaticDescription() const OVERRIDE;
+
+	/** called on instance startup, prepares root level nodes to use */
+	void InjectNodes(UBehaviorTreeComponent* OwnerComp, uint8* NodeMemory, int32& InstancedIndex) const;
+
+	/** @returns number of injected nodes */
+	int32 GetInjectedNodesCount() const;
+
+	/** @returns subtree asset */
+	class UBehaviorTree* GetSubtreeAsset() const;
+
+protected:
+
+	/** behavior to run */
+	UPROPERTY(Category = Node, EditAnywhere)
+	class UBehaviorTree* BehaviorAsset;
+};
+
+//////////////////////////////////////////////////////////////////////////
+// Inlines
+
+FORCEINLINE class UBehaviorTree* UBTTask_RunBehavior::GetSubtreeAsset() const
+{
+	return BehaviorAsset;
+}
+
+FORCEINLINE int32 UBTTask_RunBehavior::GetInjectedNodesCount() const
+{
+	return BehaviorAsset ? BehaviorAsset->RootDecorators.Num() : 0;
+}
diff --git a/Engine/Source/Runtime/AIModule/Classes/BehaviorTree/Tasks/BTTask_RunEQSQuery.h b/Engine/Source/Runtime/AIModule/Classes/BehaviorTree/Tasks/BTTask_RunEQSQuery.h
new file mode 100644
--- /dev/null
+++ b/Engine/Source/Runtime/AIModule/Classes/BehaviorTree/Tasks/BTTask_RunEQSQuery.h
@@ -0,0 +1,45 @@
+// Copyright 1998-2014 Epic Games, Inc. All Rights Reserved.
+
+#pragma once
+#include "../../EnvironmentQuery/EnvQueryTypes.h"
+#include "BTTask_RunEQSQuery.generated.h"
+
+struct FBTEnvQueryTaskMemory
+{
+	/** Query request ID */
+	int32 RequestID;
+};
+
+UCLASS(MinimalAPI)
+class UBTTask_RunEQSQuery : public UBTTask_BlackboardBase
+{
+	GENERATED_UCLASS_BODY()
+
+	/** query to run */
+	UPROPERTY(Category=Node, EditAnywhere)
+	class UEnvQuery* QueryTemplate;
+
+	/** optional parameters for query */
+	UPROPERTY(Category=Node, EditAnywhere)
+	TArray<struct FEnvNamedValue> QueryParams;
+
+	virtual EBTNodeResult::Type ExecuteTask(class UBehaviorTreeComponent* OwnerComp, uint8* NodeMemory) OVERRIDE;
+	virtual EBTNodeResult::Type AbortTask(class UBehaviorTreeComponent* OwnerComp, uint8* NodeMemory) OVERRIDE;
+
+	virtual void DescribeRuntimeValues(const class UBehaviorTreeComponent* OwnerComp, uint8* NodeMemory, EBTDescriptionVerbosity::Type Verbosity, TArray<FString>& Values) const OVERRIDE;
+	virtual FString GetStaticDescription() const OVERRIDE;
+	virtual uint16 GetInstanceMemorySize() const OVERRIDE;
+
+	/** finish task */
+	void OnQueryFinished(TSharedPtr<struct FEnvQueryResult> Result);
+
+#if WITH_EDITOR
+	/** prepare query params */
+	virtual void PostEditChangeProperty(struct FPropertyChangedEvent& PropertyChangedEvent) OVERRIDE;
+#endif
+
+protected:
+
+	/** gather all filters from existing EnvQueryItemTypes */
+	void CollectKeyFilters();
+};
diff --git a/Engine/Source/Runtime/AIModule/Classes/BehaviorTree/Tasks/BTTask_Wait.h b/Engine/Source/Runtime/AIModule/Classes/BehaviorTree/Tasks/BTTask_Wait.h
new file mode 100644
--- /dev/null
+++ b/Engine/Source/Runtime/AIModule/Classes/BehaviorTree/Tasks/BTTask_Wait.h
@@ -0,0 +1,29 @@
+// Copyright 1998-2014 Epic Games, Inc. All Rights Reserved.
+#pragma once
+
+#include "BTTask_Wait.generated.h"
+
+struct FBTWaitTaskMemory
+{
+	/** time left */
+	float RemainingWaitTime;
+};
+
+UCLASS()
+class UBTTask_Wait : public UBTTaskNode
+{
+	GENERATED_UCLASS_BODY()
+
+	/** wait time in seconds */
+	UPROPERTY(Category=Node, EditAnywhere)
+	float WaitTime;
+
+	virtual EBTNodeResult::Type ExecuteTask(class UBehaviorTreeComponent* OwnerComp, uint8* NodeMemory) OVERRIDE;
+	virtual uint16 GetInstanceMemorySize() const OVERRIDE;
+	virtual void DescribeRuntimeValues(const class UBehaviorTreeComponent* OwnerComp, uint8* NodeMemory, EBTDescriptionVerbosity::Type Verbosity, TArray<FString>& Values) const OVERRIDE;
+	virtual FString GetStaticDescription() const OVERRIDE;
+
+protected:
+
+	virtual void TickTask(class UBehaviorTreeComponent* OwnerComp, uint8* NodeMemory, float DeltaSeconds) OVERRIDE;
+};
diff --git a/Engine/Source/Runtime/AIModule/Classes/Blueprint/AIAsyncTaskBlueprintProxy.h b/Engine/Source/Runtime/AIModule/Classes/Blueprint/AIAsyncTaskBlueprintProxy.h
new file mode 100644
--- /dev/null
+++ b/Engine/Source/Runtime/AIModule/Classes/Blueprint/AIAsyncTaskBlueprintProxy.h
@@ -0,0 +1,32 @@
+// Copyright 1998-2014 Epic Games, Inc. All Rights Reserved.
+#pragma once
+
+#include "Navigation/PathFollowingComponent.h"
+#include "AIAsyncTaskBlueprintProxy.generated.h"
+
+DECLARE_DYNAMIC_MULTICAST_DELEGATE_OneParam(FOAISimpleDelegate, EPathFollowingResult::Type, MovementResult);
+
+UCLASS(MinimalAPI)
+class UAIAsyncTaskBlueprintProxy : public UObject
+{
+	GENERATED_UCLASS_BODY()
+
+	UPROPERTY(BlueprintAssignable)
+	FOAISimpleDelegate	OnSuccess;
+
+	UPROPERTY(BlueprintAssignable)
+	FOAISimpleDelegate	OnFail;
+
+public:
+	UFUNCTION()
+	void OnMoveCompleted(FAIRequestID RequestID, EPathFollowingResult::Type MovementResult);
+
+	void OnNoPath();
+
+	// Begin UObject interface
+	virtual void BeginDestroy() OVERRIDE;
+	// End UObject interface
+
+	TWeakObjectPtr<class AAIController> AIController;
+	FAIRequestID MoveRequestId;
+};
diff --git a/Engine/Source/Runtime/AIModule/Classes/Blueprint/AIBlueprintHelperLibrary.h b/Engine/Source/Runtime/AIModule/Classes/Blueprint/AIBlueprintHelperLibrary.h
new file mode 100644
--- /dev/null
+++ b/Engine/Source/Runtime/AIModule/Classes/Blueprint/AIBlueprintHelperLibrary.h
@@ -0,0 +1,38 @@
+// Copyright 1998-2014 Epic Games, Inc. All Rights Reserved.
+
+/**
+ * This kismet library is used for helper functions primarily used in the kismet compiler for AI related nodes
+ * NOTE: Do not change the signatures for any of these functions as it can break the kismet compiler and/or the nodes referencing them
+ */
+
+#pragma once
+#include "AIBlueprintHelperLibrary.generated.h"
+
+UCLASS()
+class AIMODULE_API UAIBlueprintHelperLibrary : public UBlueprintFunctionLibrary
+{
+	GENERATED_UCLASS_BODY()
+
+	UFUNCTION(BlueprintCallable, meta=(HidePin="WorldContextObject", DefaultToSelf="WorldContextObject", BlueprintInternalUseOnly = "TRUE"))
+	static class UAIAsyncTaskBlueprintProxy* CreateMoveToProxyObject(class UObject* WorldContextObject, APawn* Pawn, FVector Destination, AActor* TargetActor = NULL, float AcceptanceRadius = 0, bool bStopOnOverlap = false);
+
+	UFUNCTION(BlueprintCallable, Category="AI", meta=(DefaultToSelf="MessageSource"))
+	static void SendAIMessage(APawn* Target, FName Message, UObject* MessageSource, bool bSuccess = true);
+
+	UFUNCTION(BlueprintCallable, Category="AI", meta=(HidePin="WorldContextObject", DefaultToSelf="WorldContextObject", DeprecatedFunction, DeprecationMessage="Use SpawnAIFromClass"))
+	static class APawn* SpawnAI(class UObject* WorldContextObject, UBlueprint* Pawn, class UBehaviorTree* BehaviorTree, FVector Location, FRotator Rotation=FRotator::ZeroRotator, bool bNoCollisionFail=false);
+
+	UFUNCTION(BlueprintCallable, Category="AI", meta=(HidePin="WorldContextObject", DefaultToSelf="WorldContextObject"))
+	static class APawn* SpawnAIFromClass(class UObject* WorldContextObject, TSubclassOf<APawn> PawnClass, class UBehaviorTree* BehaviorTree, FVector Location, FRotator Rotation=FRotator::ZeroRotator, bool bNoCollisionFail=false);
+
+	UFUNCTION(BlueprintPure, Category="AI", meta=(DefaultToSelf="Target"))
+	static class UBlackboardComponent* GetBlackboard(AActor* Target);
+
+	/** locks indicated AI resources of animated pawn */
+	UFUNCTION(BlueprintCallable, Category = "Animation", BlueprintAuthorityOnly, meta = (DefaultToSelf = "AnimInstance"))
+	static void LockAIResourcesWithAnimation(class UAnimInstance* AnimInstance, bool bLockMovement, bool LockAILogic);
+
+	/** unlocks indicated AI resources of animated pawn. Will unlock only animation-locked resources */
+	UFUNCTION(BlueprintCallable, Category = "Animation", BlueprintAuthorityOnly, meta = (DefaultToSelf = "AnimInstance"))
+	static void UnlockAIResourcesWithAnimation(class UAnimInstance* AnimInstance, bool bUnlockMovement, bool UnlockAILogic);
+};
diff --git a/Engine/Source/Runtime/AIModule/Classes/BrainComponent.h b/Engine/Source/Runtime/AIModule/Classes/BrainComponent.h
new file mode 100644
--- /dev/null
+++ b/Engine/Source/Runtime/AIModule/Classes/BrainComponent.h
@@ -0,0 +1,184 @@
+// Copyright 1998-2014 Epic Games, Inc. All Rights Reserved.
+
+#pragma once
+
+#include "AITypes.h"
+#include "AIResourceInterface.h"
+#include "BrainComponent.generated.h"
+
+
+DECLARE_DELEGATE_TwoParams(FOnAIMessage, class UBrainComponent*, const struct FAIMessage&);
+
+DECLARE_LOG_CATEGORY_EXTERN(LogBrain, Warning, All);
+
+struct AIMODULE_API FAIMessage
+{
+	enum EStatus
+	{
+		Failure,
+		Success,
+	};
+
+	/** type of message */
+	FName MessageName;
+
+	/** message source */
+	UObject* Sender;
+
+	/** message param: ID */
+	FAIRequestID RequestID;
+
+	/** message param: status */
+	TEnumAsByte<EStatus> Status;
+
+	FAIMessage() : MessageName(NAME_None), Sender(NULL), RequestID(0), Status(FAIMessage::Success) {}
+	FAIMessage(FName InMessage, UObject* InSender) : MessageName(InMessage), Sender(InSender), RequestID(0), Status(FAIMessage::Success) {}
+	FAIMessage(FName InMessage, UObject* InSender, FAIRequestID InID, EStatus InStatus) : MessageName(InMessage), Sender(InSender), RequestID(InID), Status(InStatus) {}
+	FAIMessage(FName InMessage, UObject* InSender, FAIRequestID InID, bool bSuccess) : MessageName(InMessage), Sender(InSender), RequestID(InID), Status(bSuccess ? Success : Failure) {}
+	FAIMessage(FName InMessage, UObject* InSender, EStatus InStatus) : MessageName(InMessage), Sender(InSender), RequestID(0), Status(InStatus) {}
+	FAIMessage(FName InMessage, UObject* InSender, bool bSuccess) : MessageName(InMessage), Sender(InSender), RequestID(0), Status(bSuccess ? Success : Failure) {}
+
+	static void Send(class AController* Controller, const FAIMessage& Message);
+	static void Send(class APawn* Pawn, const FAIMessage& Message);
+	static void Send(class UBrainComponent* BrainComp, const FAIMessage& Message);
+
+	static void Broadcast(class UObject* WorldContextObject, const FAIMessage& Message);
+};
+
+typedef TSharedPtr<struct FAIMessageObserver, ESPMode::Fast> FAIMessageObserverHandle;
+
+struct AIMODULE_API FAIMessageObserver : public TSharedFromThis<FAIMessageObserver>
+{
+public:
+
+	static FAIMessageObserverHandle Create(class AController* Controller, FName MessageType, FOnAIMessage const& Delegate);
+	static FAIMessageObserverHandle Create(class AController* Controller, FName MessageType, FAIRequestID MessageID, FOnAIMessage const& Delegate);
+
+	static FAIMessageObserverHandle Create(APawn* Pawn, FName MessageType, FOnAIMessage const& Delegate);
+	static FAIMessageObserverHandle Create(APawn* Pawn, FName MessageType, FAIRequestID MessageID, FOnAIMessage const& Delegate);
+
+	static FAIMessageObserverHandle Create(UBrainComponent* BrainComp, FName MessageType, FOnAIMessage const& Delegate);
+	static FAIMessageObserverHandle Create(UBrainComponent* BrainComp, FName MessageType, FAIRequestID MessageID, FOnAIMessage const& Delegate);
+
+	~FAIMessageObserver();
+
+	void OnMessage(const FAIMessage& Message);
+	FString DescribeObservedMessage() const;
+	
+	FORCEINLINE FName GetObservedMessageType() const { return MessageType; }
+	FORCEINLINE FAIRequestID GetObservedMessageID() const { return MessageID; }
+	FORCEINLINE bool IsObservingMessageID() const { return bFilterByID; }
+
+private:
+
+	void Register(class UBrainComponent* OwnerComp);
+	void Unregister();
+
+	/** observed message type */
+	FName MessageType;
+
+	/** filter: message ID */
+	FAIRequestID MessageID;
+	bool bFilterByID;
+
+	/** delegate to call */
+	FOnAIMessage ObserverDelegate;
+
+	/** brain component owning this observer */
+	TWeakObjectPtr<UBrainComponent> Owner;
+};
+
+UCLASS(abstract, BlueprintType)
+class AIMODULE_API UBrainComponent : public UActorComponent, public IAIResourceInterface
+{
+	GENERATED_UCLASS_BODY()
+
+protected:
+	/** blackboard component */
+	UPROPERTY(transient)
+	class UBlackboardComponent* BlackboardComp;
+
+	// @TODO this is a temp contraption to implement delayed messages delivering
+	// until proper AI messaging is implemented
+	TArray<FAIMessage> MessagesToProcess;
+
+public:
+	virtual FString GetDebugInfoString() const { return TEXT(""); }
+
+	/** To be called in case we want to restart AI logic while it's still being locked.
+	 *	On subsequent ResumeLogic instead RestartLogic will be called. 
+	 *	@note this call does nothing if logic is not locked at the moment of call */
+	void RequestLogicRestartOnUnlock();
+	virtual void RestartLogic() {}
+protected:
+	virtual void StopLogic(const FString& Reason) {}
+	virtual void PauseLogic(const FString& Reason) {}
+	/** MUST be called by child implementations!
+	 *	@return indicates whether child class' ResumeLogic should be called (true) or has it been 
+	 *	handled in a different way and no other actions are required (false)*/
+	virtual EAILogicResuming::Type ResumeLogic(const FString& Reason);
+public:
+	virtual bool IsRunning() const { return false; }
+	virtual bool IsPaused() const { return false; }
+
+#if ENABLE_VISUAL_LOG
+	virtual void DescribeSelfToVisLog(struct FVisLogEntry* Snapshot) const;
+#endif // ENABLE_VISUAL_LOG
+	
+	// IAIResourceInterface begin
+	virtual void LockResource(EAILockSource::Type LockSource) OVERRIDE;
+	virtual void ClearResourceLock(EAILockSource::Type LockSource) OVERRIDE;
+	virtual void ForceUnlockResource() OVERRIDE;
+	virtual bool IsResourceLocked() const OVERRIDE;
+	// IAIResourceInterface end
+	
+	virtual void HandleMessage(const FAIMessage& Message);
+	
+	/** BEGIN UActorComponent overrides */
+	virtual void InitializeComponent() OVERRIDE;
+	virtual void TickComponent(float DeltaTime, enum ELevelTick TickType, FActorComponentTickFunction *ThisTickFunction);
+	/** END UActorComponent overrides */
+
+	/** caches BlackboardComponent's pointer to be used with this brain component */
+	void CacheBlackboardComponent(class UBlackboardComponent* BBComp);
+
+	/** @return blackboard used with this component */
+	class UBlackboardComponent* GetBlackboardComponent();
+
+	/** @return blackboard used with this component */
+	const class UBlackboardComponent* GetBlackboardComponent() const;
+
+protected:
+
+	/** active message observers */
+	TArray<FAIMessageObserver*> MessageObservers;
+
+	friend struct FAIMessageObserver;
+	friend struct FAIMessage;
+
+private:
+	/** used to keep track of which subsystem requested this AI resource be locked */
+	FAIResourceLock ResourceLock;
+
+	uint32 bDoLogicRestartOnUnlock : 1;
+
+public:
+	// static names to be used with SendMessage. Fell free to define game-specific
+	// messages anywhere you want
+	static const FName AIMessage_MoveFinished;
+	static const FName AIMessage_RepathFailed;
+	static const FName AIMessage_QueryFinished;
+};
+
+//////////////////////////////////////////////////////////////////////////
+// Inlines
+
+FORCEINLINE class UBlackboardComponent* UBrainComponent::GetBlackboardComponent()
+{
+	return BlackboardComp;
+}
+
+FORCEINLINE const class UBlackboardComponent* UBrainComponent::GetBlackboardComponent() const
+{
+	return BlackboardComp;
+}
\ No newline at end of file
diff --git a/Engine/Source/Runtime/AIModule/Classes/EnvironmentQuery/Contexts/EnvQueryContext_Querier.h b/Engine/Source/Runtime/AIModule/Classes/EnvironmentQuery/Contexts/EnvQueryContext_Querier.h
new file mode 100644
--- /dev/null
+++ b/Engine/Source/Runtime/AIModule/Classes/EnvironmentQuery/Contexts/EnvQueryContext_Querier.h
@@ -0,0 +1,13 @@
+// Copyright 1998-2014 Epic Games, Inc. All Rights Reserved.
+
+#pragma once
+#include "EnvironmentQuery/EnvQueryContext.h"
+#include "EnvQueryContext_Querier.generated.h"
+
+UCLASS(MinimalAPI)
+class UEnvQueryContext_Querier : public UEnvQueryContext
+{
+	GENERATED_UCLASS_BODY()
+ 
+	virtual void ProvideContext(struct FEnvQueryInstance& QueryInstance, struct FEnvQueryContextData& ContextData) const OVERRIDE;
+};
diff --git a/Engine/Source/Runtime/AIModule/Classes/EnvironmentQuery/EQSRenderingComponent.h b/Engine/Source/Runtime/AIModule/Classes/EnvironmentQuery/EQSRenderingComponent.h
new file mode 100644
--- /dev/null
+++ b/Engine/Source/Runtime/AIModule/Classes/EnvironmentQuery/EQSRenderingComponent.h
@@ -0,0 +1,45 @@
+// Copyright 1998-2014 Epic Games, Inc. All Rights Reserved.
+
+#pragma once
+#include "DebugRenderSceneProxy.h"
+#include "Debug/DebugDrawService.h"
+#include "EnvQueryTypes.h"
+#include "EQSRenderingComponent.generated.h"
+
+class FEQSSceneProxy : public FDebugRenderSceneProxy
+{
+public:
+	FEQSSceneProxy(const UPrimitiveComponent* InComponent, const FString& ViewFlagName = TEXT("DebugAI"), bool bDrawOnlyWhenSelected = true);
+	FEQSSceneProxy(const UPrimitiveComponent* InComponent, const FString& ViewFlagName, bool bDrawOnlyWhenSelected, const TArray<FSphere>& Spheres, const TArray<FText3d>& Texts);
+
+	virtual void DrawDebugLabels(UCanvas* Canvas, APlayerController*) OVERRIDE;
+	
+	//virtual void DrawStaticElements(FStaticPrimitiveDrawInterface* PDI) OVERRIDE;
+	virtual FPrimitiveViewRelevance GetViewRelevance(const FSceneView* View) OVERRIDE;
+
+	static void CollectEQSData(const UPrimitiveComponent* InComponent, const class IEQSQueryResultSourceInterface* QueryDataSource, TArray<FSphere>& Spheres, TArray<FText3d>& Texts);
+private:
+	FEnvQueryResult QueryResult;	
+	// can be 0
+	AActor* ActorOwner;
+	const class IEQSQueryResultSourceInterface* QueryDataSource;
+	bool bDrawOnlyWhenSelected;
+
+	static const FVector ItemDrawRadius;
+
+	bool SafeIsActorSelected() const;
+};
+
+UCLASS(hidecategories=Object)
+class AIMODULE_API UEQSRenderingComponent : public UPrimitiveComponent
+{
+	GENERATED_UCLASS_BODY()
+
+	FString DrawFlagName;
+	bool bDrawOnlyWhenSelected;
+
+	virtual FPrimitiveSceneProxy* CreateSceneProxy() OVERRIDE;
+	virtual FBoxSphereBounds CalcBounds(const FTransform &LocalToWorld) const OVERRIDE;
+	virtual void CreateRenderState_Concurrent() OVERRIDE;
+	virtual void DestroyRenderState_Concurrent() OVERRIDE;
+};
diff --git a/Engine/Source/Runtime/AIModule/Classes/EnvironmentQuery/EQSTestingPawn.h b/Engine/Source/Runtime/AIModule/Classes/EnvironmentQuery/EQSTestingPawn.h
new file mode 100644
--- /dev/null
+++ b/Engine/Source/Runtime/AIModule/Classes/EnvironmentQuery/EQSTestingPawn.h
@@ -0,0 +1,80 @@
+// Copyright 1998-2014 Epic Games, Inc. All Rights Reserved.
+
+#pragma once
+
+#include "GameFramework/Character.h"
+#include "EnvironmentQuery/EQSQueryResultSourceInterface.h"
+#include "EQSTestingPawn.generated.h"
+
+/** this class is abstract even though it's perfectly functional on its own.
+ *	The reason is to stop it from showing as valid player pawn type when configuring 
+ *	project's game mode. */
+UCLASS(abstract, hidecategories=(Advanced, Attachment, Collision, Animation),DependsOn=(UEnvQueryTypes))
+class AIMODULE_API AEQSTestingPawn : public ACharacter, public IEQSQueryResultSourceInterface
+{
+	GENERATED_UCLASS_BODY()
+	
+	UPROPERTY(Category=EQS, EditAnywhere)
+	class UEnvQuery* QueryTemplate;
+
+	/** optional parameters for query */
+	UPROPERTY(Category=EQS, EditAnywhere)
+	TArray<struct FEnvNamedValue> QueryParams;
+
+	UPROPERTY(Category=EQS, EditAnywhere)
+	float TimeLimitPerStep;
+
+	UPROPERTY(Category=EQS, EditAnywhere)
+	int32 StepToDebugDraw;
+
+	UPROPERTY(Category=EQS, EditAnywhere)
+	uint32 bDrawLabels:1;
+
+	UPROPERTY(Category=EQS, EditAnywhere)
+	uint32 bDrawFailedItems:1;
+
+	UPROPERTY(Category=EQS, EditAnywhere)
+	uint32 bReRunQueryOnlyOnFinishedMove:1;
+
+	UPROPERTY(Category=EQS, EditAnywhere)
+	uint32 bShouldBeVisibleInGame:1;
+
+#if WITH_EDITORONLY_DATA
+	/** Editor Preview */
+	UPROPERTY()
+	TSubobjectPtr<class UEQSRenderingComponent> EdRenderComp;
+#endif // WITH_EDITORONLY_DATA
+
+protected:
+	TSharedPtr<struct FEnvQueryInstance> QueryInstance;
+
+	TArray<FEnvQueryInstance> StepResults;
+
+public:
+	/** This pawn class spawns its controller in PostInitProperties to have it available in editor mode*/
+	virtual void TickActor( float DeltaTime, enum ELevelTick TickType, FActorTickFunction& ThisTickFunction ) OVERRIDE;
+	virtual void PostLoad() OVERRIDE;
+
+#if WITH_EDITOR
+	virtual void PostEditChangeProperty(struct FPropertyChangedEvent& PropertyChangedEvent) OVERRIDE;
+	virtual void PostEditMove(bool bFinished) OVERRIDE;
+#endif // WITH_EDITOR
+
+	// IEQSQueryResultSourceInterface start
+	virtual const struct FEnvQueryResult* GetQueryResult() const OVERRIDE;
+	virtual const struct FEnvQueryInstance* GetQueryInstance() const  OVERRIDE;
+
+	virtual bool GetShouldDebugDrawLabels() const OVERRIDE { return bDrawLabels; }
+	virtual bool GetShouldDrawFailedItems() const OVERRIDE{ return bDrawFailedItems; }
+	// IEQSQueryResultSourceInterface end
+
+	void RunEQSQuery();
+
+protected:	
+	void Reset();
+	void MakeOneStep();
+
+	void UpdateDrawing();
+
+	static void OnEditorSelectionChanged(UObject* NewSelection);
+};
diff --git a/Engine/Source/Runtime/AIModule/Classes/EnvironmentQuery/EnvQuery.h b/Engine/Source/Runtime/AIModule/Classes/EnvironmentQuery/EnvQuery.h
new file mode 100644
--- /dev/null
+++ b/Engine/Source/Runtime/AIModule/Classes/EnvironmentQuery/EnvQuery.h
@@ -0,0 +1,24 @@
+// Copyright 1998-2014 Epic Games, Inc. All Rights Reserved.
+
+#pragma once
+
+#include "EnvQueryTypes.h"
+#include "EnvQuery.generated.h"
+
+UCLASS()
+class AIMODULE_API UEnvQuery : public UObject
+{
+	GENERATED_UCLASS_BODY()
+
+#if WITH_EDITORONLY_DATA
+	/** Graph for query */
+	UPROPERTY()
+	class UEdGraph*	EdGraph;
+#endif
+
+	UPROPERTY()
+	TArray<class UEnvQueryOption*> Options;
+
+	/** Gather all required named params */
+	void CollectQueryParams(TArray<struct FEnvNamedValue>& NamedValues) const;
+};
diff --git a/Engine/Source/Runtime/AIModule/Classes/EnvironmentQuery/EnvQueryContext.h b/Engine/Source/Runtime/AIModule/Classes/EnvironmentQuery/EnvQueryContext.h
new file mode 100644
--- /dev/null
+++ b/Engine/Source/Runtime/AIModule/Classes/EnvironmentQuery/EnvQueryContext.h
@@ -0,0 +1,12 @@
+// Copyright 1998-2014 Epic Games, Inc. All Rights Reserved.
+
+#pragma once
+#include "EnvQueryContext.generated.h"
+
+UCLASS(Abstract)
+class AIMODULE_API UEnvQueryContext : public UObject
+{
+	GENERATED_UCLASS_BODY()
+
+	virtual void ProvideContext(struct FEnvQueryInstance& QueryInstance, struct FEnvQueryContextData& ContextData) const;
+};
diff --git a/Engine/Source/Runtime/AIModule/Classes/EnvironmentQuery/EnvQueryGenerator.h b/Engine/Source/Runtime/AIModule/Classes/EnvironmentQuery/EnvQueryGenerator.h
new file mode 100644
--- /dev/null
+++ b/Engine/Source/Runtime/AIModule/Classes/EnvironmentQuery/EnvQueryGenerator.h
@@ -0,0 +1,27 @@
+// Copyright 1998-2014 Epic Games, Inc. All Rights Reserved.
+
+#pragma once
+#include "EnvironmentQuery/EnvQueryTypes.h"
+#include "EnvQueryGenerator.generated.h"
+
+UCLASS(Abstract)
+class AIMODULE_API UEnvQueryGenerator : public UObject
+{
+	GENERATED_UCLASS_BODY()
+
+	UPROPERTY(EditDefaultsOnly, Category=Option)
+	FString OptionName;
+
+	/** type of generated items */
+	TSubclassOf<class UEnvQueryItemType> ItemType;
+
+	FGenerateItemsSignature GenerateDelegate;
+
+	/** get description of generator */
+	virtual FText GetDescriptionTitle() const;
+	virtual FText GetDescriptionDetails() const;
+
+#if WITH_EDITOR
+	virtual void PostEditChangeProperty(struct FPropertyChangedEvent& PropertyChangedEvent) OVERRIDE;
+#endif //WITH_EDITOR
+};
diff --git a/Engine/Source/Runtime/AIModule/Classes/EnvironmentQuery/EnvQueryManager.h b/Engine/Source/Runtime/AIModule/Classes/EnvironmentQuery/EnvQueryManager.h
new file mode 100644
--- /dev/null
+++ b/Engine/Source/Runtime/AIModule/Classes/EnvironmentQuery/EnvQueryManager.h
@@ -0,0 +1,156 @@
+// Copyright 1998-2014 Epic Games, Inc. All Rights Reserved.
+
+#pragma once
+#include "EnvQueryTypes.h"
+#include "Tickable.h"
+#include "EnvQueryManager.generated.h"
+
+/** wrapper for easy query execution */
+USTRUCT()
+struct AIMODULE_API FEnvQueryRequest
+{
+	GENERATED_USTRUCT_BODY()
+
+	FEnvQueryRequest() : QueryTemplate(NULL), Owner(NULL), World(NULL) {}
+
+	// basic constructor: owner will be taken from finish delegate bindings
+	FEnvQueryRequest(const class UEnvQuery* Query) : QueryTemplate(Query), Owner(NULL), World(NULL) {}
+
+	// use when owner is different from finish delegate binding
+	FEnvQueryRequest(const class UEnvQuery* Query, UObject* RequestOwner) : QueryTemplate(Query), Owner(RequestOwner), World(NULL) {}
+
+	// set named params
+	FORCEINLINE FEnvQueryRequest& SetFloatParam(FName ParamName, float Value) { NamedParams.Add(ParamName, Value); return *this; }
+	FORCEINLINE FEnvQueryRequest& SetIntParam(FName ParamName, int32 Value) { NamedParams.Add(ParamName, *((float*)&Value)); return *this; }
+	FORCEINLINE FEnvQueryRequest& SetBoolParam(FName ParamName, bool Value) { NamedParams.Add(ParamName, Value ? 1.0f : -1.0f); return *this; }
+	FORCEINLINE FEnvQueryRequest& SetNamedParam(const struct FEnvNamedValue& ParamData) { NamedParams.Add(ParamData.ParamName, ParamData.Value); return *this; }
+	FEnvQueryRequest& SetNamedParams(const TArray<struct FEnvNamedValue>& Params);
+
+	// set world (for accessing query manager) when owner can't provide it
+	FORCEINLINE FEnvQueryRequest& SetWorldOverride(UWorld* InWorld) { World = InWorld; return *this; }
+
+	template< class UserClass >	
+	FORCEINLINE int32 Execute(EEnvQueryRunMode::Type Mode, UserClass* InObj, typename FQueryFinishedSignature::TUObjectMethodDelegate< UserClass >::FMethodPtr InMethod)
+	{
+		return Execute(Mode, FQueryFinishedSignature::CreateUObject(InObj, InMethod));
+	}
+	template< class UserClass >	
+	FORCEINLINE int32 Execute(EEnvQueryRunMode::Type Mode, UserClass* InObj, typename FQueryFinishedSignature::TUObjectMethodDelegate_Const< UserClass >::FMethodPtr InMethod)
+	{
+		return Execute(Mode, FQueryFinishedSignature::CreateUObject(InObj, InMethod));
+	}
+	int32 Execute(EEnvQueryRunMode::Type RunMode, FQueryFinishedSignature const& FinishDelegate);
+
+protected:
+
+	/** query to run */
+	UPROPERTY()
+	const class UEnvQuery* QueryTemplate;
+
+	/** querier */
+	UPROPERTY()
+	UObject* Owner;
+
+	/** world */
+	UPROPERTY()
+	UWorld* World;
+
+	/** list of named params */
+	TMap<FName, float> NamedParams;
+
+	friend class UEnvQueryManager;
+};
+
+/** cache of instances with sorted tests */
+struct FEnvQueryInstanceCache
+{
+	/** query template */
+	TWeakObjectPtr<class UEnvQuery> Template;
+
+	/** instance to duplicate */
+	struct FEnvQueryInstance Instance;
+};
+
+#if USE_EQS_DEBUGGER
+struct AIMODULE_API FEQSDebugger
+{
+	void StoreQuery(TSharedPtr<FEnvQueryInstance>& Query);
+	TSharedPtr<FEnvQueryInstance> GetQueryForOwner(const UObject* Owner);
+	const TSharedPtr<FEnvQueryInstance> GetQueryForOwner(const UObject* Owner) const;
+
+protected:
+	// maps owner to performed queries
+	TMap<const UObject*, TSharedPtr<FEnvQueryInstance> > StoredQueries;
+};
+#endif // USE_EQS_DEBUGGER
+
+UCLASS()
+class AIMODULE_API UEnvQueryManager : public UObject, public FTickableGameObject
+{
+	GENERATED_UCLASS_BODY()
+
+	/** [FTickableGameObject] tick function */
+	virtual void Tick(float DeltaTime) OVERRIDE;
+
+	/** [FTickableGameObject] always tick */
+	virtual bool IsTickable() const OVERRIDE { return true; }
+
+	/** [FTickableGameObject] tick stats */
+	virtual TStatId GetStatId() const OVERRIDE;
+
+	/** execute query */
+	int32 RunQuery(const struct FEnvQueryRequest& Request, EEnvQueryRunMode::Type RunMode, FQueryFinishedSignature const& FinishDelegate);
+
+	/** alternative way to run queries. Do not use for anything other then testing! (worse performance) */
+	TSharedPtr<FEnvQueryResult> RunInstantQuery(const struct FEnvQueryRequest& Request, EEnvQueryRunMode::Type RunMode);
+
+	/** Creates a query instance configured for execution */
+	TSharedPtr<struct FEnvQueryInstance> PrepareQueryInstance(const struct FEnvQueryRequest& Request, EEnvQueryRunMode::Type RunMode);
+
+	/** finds UEnvQuery matching QueryName by first looking at instantiated queries (from InstanceCache)
+	 *	falling back to looking up all UEnvQuery and testing their name */
+	UEnvQuery* FindQueryTemplate(const FString& QueryName) const;
+
+	/** execute query */
+	bool AbortQuery(int32 RequestID);
+
+	/** fail all running queries on changing persistent map */
+	virtual void OnPreLoadMap();
+
+	/** cleanup hooks for map loading */
+	virtual void FinishDestroy() OVERRIDE;
+
+	/** list of all known item types */
+	static TArray<TSubclassOf<UEnvQueryItemType> > RegisteredItemTypes;
+
+	static UEnvQueryManager* GetCurrent(class UWorld* World);
+	static UEnvQueryManager* GetCurrent(class UObject* WorldContextObject);
+
+#if USE_EQS_DEBUGGER
+	static void NotifyAssetUpdate(UEnvQuery* Query);
+
+	FEQSDebugger& GetDebugger() { return EQSDebugger; }
+
+protected:
+	FEQSDebugger EQSDebugger;
+#endif // USE_EQS_DEBUGGER
+
+protected:
+
+	/** currently running queries */
+	TArray<TSharedPtr<struct FEnvQueryInstance> > RunningQueries;
+
+	/** cache of instances */
+	TArray<FEnvQueryInstanceCache> InstanceCache;
+
+	/** next ID for running query */
+	int32 NextQueryID;
+
+	/** create new instance, using cached data is possible */
+	TSharedPtr<struct FEnvQueryInstance> CreateQueryInstance(const class UEnvQuery* Template, EEnvQueryRunMode::Type RunMode);
+
+private:
+
+	/** create and bind delegates in instance */
+	void CreateOptionInstance(class UEnvQueryOption* OptionTemplate, const TArray<class UEnvQueryTest*>& SortedTests, struct FEnvQueryInstance& Instance);
+};
diff --git a/Engine/Source/Runtime/AIModule/Classes/EnvironmentQuery/EnvQueryOption.h b/Engine/Source/Runtime/AIModule/Classes/EnvironmentQuery/EnvQueryOption.h
new file mode 100644
--- /dev/null
+++ b/Engine/Source/Runtime/AIModule/Classes/EnvironmentQuery/EnvQueryOption.h
@@ -0,0 +1,19 @@
+// Copyright 1998-2014 Epic Games, Inc. All Rights Reserved.
+
+#pragma once
+#include "EnvQueryOption.generated.h"
+
+UCLASS()
+class AIMODULE_API UEnvQueryOption : public UObject
+{
+	GENERATED_UCLASS_BODY()
+
+	UPROPERTY()
+	class UEnvQueryGenerator* Generator;
+
+	UPROPERTY()
+	TArray<class UEnvQueryTest*> Tests;
+
+	FText GetDescriptionTitle() const;
+	FText GetDescriptionDetails() const;
+};
diff --git a/Engine/Source/Runtime/AIModule/Classes/EnvironmentQuery/EnvQueryTest.h b/Engine/Source/Runtime/AIModule/Classes/EnvironmentQuery/EnvQueryTest.h
new file mode 100644
--- /dev/null
+++ b/Engine/Source/Runtime/AIModule/Classes/EnvironmentQuery/EnvQueryTest.h
@@ -0,0 +1,164 @@
+// Copyright 1998-2014 Epic Games, Inc. All Rights Reserved.
+
+#pragma once
+#include "EnvironmentQuery/EnvQueryTypes.h"
+#include "EnvQueryTest.generated.h"
+
+UCLASS(Abstract)
+class AIMODULE_API UEnvQueryTest : public UObject
+{
+	GENERATED_UCLASS_BODY()
+
+	/** Number of test as defined in data asset */
+	UPROPERTY()
+	int32 TestOrder;
+
+	UPROPERTY(EditDefaultsOnly, Category=Filter)
+	TEnumAsByte<EEnvTestPurpose::Type> TestPurpose;
+	
+	UPROPERTY(EditDefaultsOnly, Category=Filter)
+	TEnumAsByte<EEnvTestFilterType::Type> FilterType;
+
+	/** Boolean value that must be matched in order for scoring to occur or filtering test to pass. */
+	UPROPERTY(EditDefaultsOnly, Category=Filter)
+	FEnvBoolParam BoolFilter;
+
+	/** filter for test value; currently this is the ONLY value, because we can only filter in one direction. */
+	UPROPERTY(EditDefaultsOnly, Category=Filter)
+	FEnvFloatParam FloatFilter;
+
+	UPROPERTY(EditDefaultsOnly, Category=Filter,
+		Meta=(ToolTip="Minimum limit (inclusive) of valid values for the raw test value.  Lower values will be discarded as invalid."))
+	FEnvFloatParam FloatFilterMin;
+
+	UPROPERTY(EditDefaultsOnly, Category=Filter,
+		Meta=(ToolTip="Maximum limit (inclusive) of valid values for the raw test value.  Higher values will be discarded as invalid."))
+	FEnvFloatParam FloatFilterMax;
+
+	// Using this property just as a way to do file version conversion without having to up the version number!
+	UPROPERTY()
+	bool bFormatUpdated;
+
+	// TODO: REMOVE, Deprecated!
+	/** Condition for discarding item */
+	UPROPERTY(EditDefaultsOnly, Category=DeprecatedFilt)
+	TEnumAsByte<EEnvTestCondition::Type> Condition;
+
+	// TODO: REMOVE, Deprecated!
+	/** Weight modifier */
+	UPROPERTY(EditDefaultsOnly, Category=Weight)
+	TEnumAsByte<EEnvTestWeight::Type> WeightModifier;
+
+	/** Cost of test */
+	TEnumAsByte<EEnvTestCost::Type> Cost;
+
+// 	UPROPERTY(EditDefaultsOnly, Category=Score, Meta=(DisplayName="Take Absolute Value of Test Result?",
+// 		ToolTip="If checked, the raw value of the test result will be converted to an absolute value BEFORE clamping.  Otherwise, the raw value will reamin as-is."))
+// 	bool bUseAbsoluteValueBeforeClamping;
+
+	UPROPERTY(EditDefaultsOnly, Category=Score,
+		Meta=(ToolTip="How should the lower bound for normalization of the raw test value before applying the scoring formula be determined?  Should it use the lowest value found (tested), the lower threshold for filtering, or a separate specified normalization minimum?"))
+	TEnumAsByte<EEnvQueryTestClamping::Type> ClampMinType;
+
+	UPROPERTY(EditDefaultsOnly, Category=Score,
+		Meta=(ToolTip="How should the upper bound for normalization of the raw test value before applying the scoring formula be determined?  Should it use the highes value found (tested), the upper threshold for filtering, or a separate specified normalization maximum?"))
+	TEnumAsByte<EEnvQueryTestClamping::Type> ClampMaxType;
+
+	UPROPERTY(EditDefaultsOnly, Category=Score,
+		Meta=(ToolTip="Minimum value to use to normalize the raw test value before applying scoring formula."))
+	FEnvFloatParam ScoreClampingMin;
+
+	UPROPERTY(EditDefaultsOnly, Category=Score,
+		Meta=(ToolTip="Maximum value to use to normalize the raw test value before applying scoring formula."))
+	FEnvFloatParam ScoreClampingMax;
+
+	// TODO: Should we hide this option for boolean tests?  In those cases, presumably scoring flipped would just
+	// reverse the true/false match we're looking for.
+// 	UPROPERTY(EditDefaultsOnly, Category=Score,
+// 		Meta=(Tooltip="If checked, reverse the normalized score by taking (1.0 - Normalized Score) before applying the curve scoring equation.  Use this if lower test values should result in a higher score."))
+// 	bool bMirrorNormalizedScore;
+
+	UPROPERTY(EditDefaultsOnly, Category=Score)
+	TEnumAsByte<EEnvTestScoreEquation::Type> ScoringEquation;
+
+	// For now, I'm keeping "Weight".  TODO: Reconsider Weight vs. "Scoring Factor", Constant, Power, etc. once we have a final plan for what properties are available.
+	/** Weight of test */
+	UPROPERTY(EditDefaultsOnly, Category=Score)
+	FEnvFloatParam Weight;
+
+// 	UPROPERTY(EditDefaultsOnly, Category=Score, Meta=(DisplayName="Factor"))
+// 	FEnvFloatParam ScoringFactor;
+
+// 	UPROPERTY(EditDefaultsOnly, Category=Score, Meta=(DisplayName="Constant"))
+// 	FEnvFloatParam ScoringConstant;
+
+// Only need this if we add in support for parametric equations
+// 	UPROPERTY(EditDefaultsOnly, Category=Score, Meta=(DisplayName="Power"))
+// 	FEnvFloatParam ScoringPower;
+
+	/** Validation: item type that can be used with this test */
+	TSubclassOf<UEnvQueryItemType> ValidItemType;
+
+	void SetWorkOnFloatValues(bool bWorkOnFloats);
+	FORCEINLINE bool GetWorkOnFloatValues() const { return bWorkOnFloatValues; }
+
+	// TODO: REMOVE, Deprecated!
+	/** Normalize test result starting from 0 */
+	uint32 bNormalizeFromZero : 1;
+
+private:
+	/** When set, test operates on float values (e.g. distance, with AtLeast, UpTo conditions),
+	 *  otherwise it will accept bool values (e.g. visibility, with Equals condition) */
+	UPROPERTY()
+	uint32 bWorkOnFloatValues : 1;
+
+public:
+	// TODO: REMOVE, Deprecated!
+	/** If set, items not passing filter condition will be removed from query, otherwise they will receive score 0 */
+	UPROPERTY(EditDefaultsOnly, Category=DeprecatedFilt)
+	uint32 bDiscardFailedItems : 1;
+
+	FExecuteTestSignature ExecuteDelegate;
+
+	/** check if test supports item type */
+	bool IsSupportedItem(TSubclassOf<class UEnvQueryItemType> ItemType) const;
+
+	/** check if context needs to be updated for every item */
+	bool IsContextPerItem(TSubclassOf<class UEnvQueryContext> CheckContext) const;
+
+	/** helper: get location of item */
+	FVector GetItemLocation(struct FEnvQueryInstance& QueryInstance, int32 ItemIndex) const;
+
+	/** helper: get location of item */
+	FRotator GetItemRotation(struct FEnvQueryInstance& QueryInstance, int32 ItemIndex) const;
+
+	/** helper: get actor from item */
+	AActor* GetItemActor(struct FEnvQueryInstance& QueryInstance, int32 ItemIndex) const;
+
+	/** normalize scores in range */
+	void NormalizeItemScores(struct FEnvQueryInstance& QueryInstance);
+
+	FORCEINLINE bool IsScoring() const { return (TestPurpose != EEnvTestPurpose::Filter); } // ((TestPurpose == EEnvTestPurpose::Score) || (TestPurpose == EEnvTestPurpose::FilterAndScore));
+	FORCEINLINE bool IsFiltering() const { return (TestPurpose != EEnvTestPurpose::Score); } // ((TestPurpose == EEnvTestPurpose::Filter) || (TestPurpose == EEnvTestPurpose::FilterAndScore));
+
+	/** get description of test */
+	virtual FString GetDescriptionTitle() const;
+	virtual FText GetDescriptionDetails() const;
+
+	FText DescribeFloatTestParams() const;
+	FText DescribeBoolTestParams(const FString& ConditionDesc) const;
+
+	virtual void PostLoad() OVERRIDE;
+
+#if WITH_EDITOR
+	virtual void PostEditChangeProperty(struct FPropertyChangedEvent& PropertyChangedEvent) OVERRIDE;
+#endif //WITH_EDITOR
+};
+
+//////////////////////////////////////////////////////////////////////////
+// Inlines
+
+FORCEINLINE bool UEnvQueryTest::IsSupportedItem(TSubclassOf<class UEnvQueryItemType> ItemType) const
+{
+	return ItemType && (ItemType == ValidItemType || ItemType->IsChildOf(ValidItemType));
+};
diff --git a/Engine/Source/Runtime/AIModule/Classes/EnvironmentQuery/EnvQueryTypes.h b/Engine/Source/Runtime/AIModule/Classes/EnvironmentQuery/EnvQueryTypes.h
new file mode 100644
--- /dev/null
+++ b/Engine/Source/Runtime/AIModule/Classes/EnvironmentQuery/EnvQueryTypes.h
@@ -0,0 +1,935 @@
+// Copyright 1998-2014 Epic Games, Inc. All Rights Reserved.
+
+#pragma once
+
+#include "EnvironmentQuery/Items/EnvQueryItemType.h"
+#include "EnvQueryTypes.generated.h"
+
+DECLARE_LOG_CATEGORY_EXTERN(LogEQS, Log, All);
+
+// If set, execution details will be processed by debugger
+#define USE_EQS_DEBUGGER				(1 && !(UE_BUILD_SHIPPING || UE_BUILD_TEST))
+
+DECLARE_STATS_GROUP(TEXT("Environment Query"), STATGROUP_AI_EQS, STATCAT_Advanced);
+
+DECLARE_CYCLE_STAT_EXTERN(TEXT("Tick"),STAT_AI_EQS_Tick,STATGROUP_AI_EQS, );
+DECLARE_CYCLE_STAT_EXTERN(TEXT("Load Time"),STAT_AI_EQS_LoadTime,STATGROUP_AI_EQS, );
+DECLARE_CYCLE_STAT_EXTERN(TEXT("Generator Time"),STAT_AI_EQS_GeneratorTime,STATGROUP_AI_EQS, );
+DECLARE_CYCLE_STAT_EXTERN(TEXT("Test Time"),STAT_AI_EQS_TestTime,STATGROUP_AI_EQS, );
+DECLARE_DWORD_COUNTER_STAT_EXTERN(TEXT("Num Instances"),STAT_AI_EQS_NumInstances,STATGROUP_AI_EQS, );
+DECLARE_DWORD_COUNTER_STAT_EXTERN(TEXT("Num Items"),STAT_AI_EQS_NumItems,STATGROUP_AI_EQS, );
+DECLARE_MEMORY_STAT_EXTERN(TEXT("Instance memory"),STAT_AI_EQS_InstanceMemory,STATGROUP_AI_EQS, AIMODULE_API);
+
+UENUM()
+namespace EEnvTestPurpose
+{
+	enum Type
+	{
+		Filter,
+		Score,
+		FilterAndScore UMETA(DisplayName="Filter and Score")
+	};
+}
+
+UENUM()
+namespace EEnvTestFilterType
+{
+	enum Type
+	{
+		Minimum,	// For numeric tests
+		Maximum,	// For numeric tests
+		Range,		// For numeric tests
+		Match		// For boolean tests
+	};
+}
+
+UENUM()
+namespace EEnvTestCondition
+{
+	enum Type
+	{
+		NoCondition		UMETA(DisplayName="Always pass"),
+		AtLeast			UMETA(DisplayName="At least"),
+		UpTo			UMETA(DisplayName="Up to"),
+		Match,
+	};
+}
+
+UENUM()
+namespace EEnvTestScoreEquation
+{
+	enum Type
+	{
+		Linear,
+		Square,
+		InverseLinear,	// For now...
+		Constant
+		// What other curve shapes should be supported?  At first I was thinking we'd have parametric (F*V^P + C), but
+		// many versions of that curve would violate the [0, 1] output range which I think we should preserve.  So instead
+		// I think we should define these by "curve shape".  I'm not sure if we need to allow full tweaks to the curves,
+		// such as supporting other "Exponential" curves (positive even powers).  However, I think it's likely that we'll
+		// want to support "smooth LERP" / S-shaped curve of the form 2x^3 - 3x^2, and possibly a "sideways" version of
+		// the same S-curve.  We also might want to allow "Sine" curves, basically adjusted to match the range and then
+		// simply offset by some amount to allow a peak or valley in the middle or on the ends.  (Four Sine options are
+		// probably sufficient.)  I'm not sure if Sine is really needed though, so probably we should only add it if
+		// there's a need identified.  One other curve shape we might want is "Square Root", which might optionally
+		// support any positive fractional power (if we also supported any positive even number for an "Exponential"
+		// type.
+	};
+}
+
+UENUM()
+namespace EEnvTestWeight
+{
+	enum Type
+	{
+		None,
+		Square,
+		Inverse,
+		Absolute,	// Removed after discussion with EGP folks
+		Constant,
+		Skip			UMETA(DisplayName = "Do not weight"),
+	};
+}
+
+UENUM()
+namespace EEnvTestCost
+{
+	enum Type
+	{
+		Low,				// reading data, math operations (e.g. distance)
+		Medium,				// processing data from multiple sources (e.g. fire tickets)
+		High,				// really expensive calls (e.g. visibility traces, pathfinding)
+	};
+}
+
+UENUM()
+namespace EEnvQueryStatus
+{
+	enum Type
+	{
+		Processing,
+		Success,
+		Failed,
+		Aborted,
+		OwnerLost,
+		MissingParam,
+	};
+}
+
+UENUM()
+namespace EEnvQueryRunMode
+{
+	enum Type
+	{
+		SingleResult,		// weight scoring first, try conditions from best result and stop after first item pass
+		AllMatching,		// conditions first (limit set of items), weight scoring later
+	};
+}
+
+UENUM()
+namespace EEnvQueryParam
+{
+	enum Type
+	{
+		Float,
+		Int,
+		Bool,
+	};
+}
+
+UENUM()
+namespace EEnvQueryTrace
+{
+	enum Type
+	{
+		None,
+		Navigation,
+		Geometry,
+	};
+}
+
+UENUM()
+namespace EEnvTraceShape
+{
+	enum Type
+	{
+		Line,
+		Box,
+		Sphere,
+		Capsule,
+	};
+}
+
+UENUM()
+namespace EEnvDirection
+{
+	enum Type
+	{
+		TwoPoints	UMETA(DisplayName="Two Points",ToolTip="Direction from location of one context to another."),
+		Rotation	UMETA(ToolTip="Context's rotation will be used as a direction."),
+	};
+}
+
+UENUM()
+namespace EEnvQueryTestClamping
+{
+	enum Type
+	{
+		None,			
+		SpecifiedValue,	// Clamp to value specified in test
+		FilterThreshold	// Clamp to test's filter threshold
+	};
+}
+
+USTRUCT()
+struct AIMODULE_API FEnvFloatParam
+{
+	GENERATED_USTRUCT_BODY();
+
+	typedef float FValueType;
+
+	FEnvFloatParam() :
+		Value(0.f)
+	{
+	}
+
+	/** default value */
+	UPROPERTY(EditDefaultsOnly, Category=Param)
+	float Value;
+
+	/** name of parameter */
+	UPROPERTY(EditDefaultsOnly, Category=Param)
+	FName ParamName;
+
+	bool IsNamedParam() const { return ParamName != NAME_None; }
+};
+
+USTRUCT()
+struct AIMODULE_API FEnvIntParam
+{
+	GENERATED_USTRUCT_BODY();
+
+	typedef int32 FValueType;
+
+	FEnvIntParam() :
+		Value(0)
+	{
+	}
+
+	/** default value */
+	UPROPERTY(EditDefaultsOnly, Category=Param)
+	int32 Value;
+
+	/** name of parameter */
+	UPROPERTY(EditDefaultsOnly, Category=Param)
+	FName ParamName;
+
+	bool IsNamedParam() const { return ParamName != NAME_None; }
+};
+
+USTRUCT()
+struct AIMODULE_API FEnvBoolParam
+{
+	GENERATED_USTRUCT_BODY();
+
+	typedef bool FValueType;
+
+	FEnvBoolParam() :
+		Value(false)
+	{
+	}
+
+	/** default value */
+	UPROPERTY(EditDefaultsOnly, Category=Param)
+	bool Value;
+
+	/** name of parameter */
+	UPROPERTY(EditDefaultsOnly, Category=Param)
+	FName ParamName;
+
+	bool IsNamedParam() const { return ParamName != NAME_None; }
+};
+
+USTRUCT(BlueprintType)
+struct AIMODULE_API FEnvNamedValue
+{
+	GENERATED_USTRUCT_BODY();
+
+	UPROPERTY(EditAnywhere, Category=Param)
+	FName ParamName;
+
+	UPROPERTY(EditAnywhere, Category=Param)
+	TEnumAsByte<EEnvQueryParam::Type> ParamType;
+
+	UPROPERTY(EditAnywhere, Category=Param)
+	float Value;
+};
+
+USTRUCT()
+struct AIMODULE_API FEnvDirection
+{
+	GENERATED_USTRUCT_BODY()
+
+	/** line A: start context */
+	UPROPERTY(EditDefaultsOnly, Category=Direction)
+	TSubclassOf<class UEnvQueryContext> LineFrom;
+
+	/** line A: finish context */
+	UPROPERTY(EditDefaultsOnly, Category=Direction)
+	TSubclassOf<class UEnvQueryContext> LineTo;
+
+	/** line A: direction context */
+	UPROPERTY(EditDefaultsOnly, Category=Direction)
+	TSubclassOf<class UEnvQueryContext> Rotation;
+
+	/** defines direction of second line used by test */
+	UPROPERTY(EditDefaultsOnly, Category=Direction, meta=(DisplayName="Mode"))
+	TEnumAsByte<EEnvDirection::Type> DirMode;
+
+	FText ToText() const;
+};
+
+USTRUCT()
+struct AIMODULE_API FEnvTraceData
+{
+	GENERATED_USTRUCT_BODY()
+
+	enum EDescriptionMode
+	{
+		Brief,
+		Detailed,
+	};
+
+	FEnvTraceData() : 
+		ProjectDown(1024.0f), ProjectUp(1024.0f), ExtentX(10.0f), ExtentY(10.0f), ExtentZ(10.0f), bOnlyBlockingHits(true),
+		bCanTraceOnNavMesh(true), bCanTraceOnGeometry(true), bCanDisableTrace(true), bCanProjectDown(false)
+	{
+	}
+
+	/** navigation filter for tracing */
+	UPROPERTY(EditDefaultsOnly, Category=Trace)
+	TSubclassOf<class UNavigationQueryFilter> NavigationFilter;
+
+	/** search height: below point */
+	UPROPERTY(EditDefaultsOnly, Category=Trace, meta=(UIMin=0, ClampMin=0))
+	float ProjectDown;
+
+	/** search height: above point */
+	UPROPERTY(EditDefaultsOnly, Category=Trace, meta=(UIMin=0, ClampMin=0))
+	float ProjectUp;
+
+	/** shape parameter for trace */
+	UPROPERTY(EditDefaultsOnly, Category=Trace, meta=(UIMin=0, ClampMin=0))
+	float ExtentX;
+
+	/** shape parameter for trace */
+	UPROPERTY(EditDefaultsOnly, Category=Trace, meta=(UIMin=0, ClampMin=0))
+	float ExtentY;
+
+	/** shape parameter for trace */
+	UPROPERTY(EditDefaultsOnly, Category=Trace, meta=(UIMin=0, ClampMin=0))
+	float ExtentZ;
+
+	/** geometry trace channel */
+	UPROPERTY(EditDefaultsOnly, Category=Trace)
+	TEnumAsByte<enum ETraceTypeQuery> TraceChannel;
+
+	/** shape used for geometry tracing */
+	UPROPERTY(EditDefaultsOnly, Category=Trace)
+	TEnumAsByte<EEnvTraceShape::Type> TraceShape;
+
+	/** shape used for geometry tracing */
+	UPROPERTY(EditDefaultsOnly, Category=Trace)
+	TEnumAsByte<EEnvQueryTrace::Type> TraceMode;
+
+	/** if set, trace will run on complex collisions */
+	UPROPERTY(EditDefaultsOnly, Category=Trace, AdvancedDisplay)
+	uint32 bTraceComplex : 1;
+
+	/** if set, trace will look only for blocking hits */
+	UPROPERTY(EditDefaultsOnly, Category=Trace, AdvancedDisplay)
+	uint32 bOnlyBlockingHits : 1;
+
+	/** if set, editor will allow picking navmesh trace */
+	UPROPERTY(EditDefaultsOnly, Category=Trace)
+	uint32 bCanTraceOnNavMesh : 1;
+
+	/** if set, editor will allow picking geometry trace */
+	UPROPERTY(EditDefaultsOnly, Category=Trace)
+	uint32 bCanTraceOnGeometry : 1;
+
+	/** if set, editor will allow  */
+	UPROPERTY(EditDefaultsOnly, Category=Trace)
+	uint32 bCanDisableTrace : 1;
+
+	/** if set, editor show height up/down properties for projection */
+	UPROPERTY(EditDefaultsOnly, Category=Trace)
+	uint32 bCanProjectDown : 1;
+
+	FText ToText(EDescriptionMode DescMode) const;
+
+	void SetGeometryOnly();
+	void SetNavmeshOnly();
+};
+
+//////////////////////////////////////////////////////////////////////////
+// Returned results
+
+struct AIMODULE_API FEnvQueryItem
+{
+	/** total score of item */
+	float Score;
+
+	/** raw data offset */
+	int32 DataOffset:31;
+
+	/** has this item been discarded? */
+	int32 bIsDiscarded:1;
+
+	FORCEINLINE bool IsValid() const { return DataOffset >= 0 && !bIsDiscarded; }
+	FORCEINLINE void Discard() { bIsDiscarded = true; }
+
+	bool operator<(const FEnvQueryItem& Other) const
+	{
+		// sort by validity
+		if (IsValid() != Other.IsValid())
+		{
+			// self not valid = less important
+			return !IsValid();
+		}
+
+		// sort by score
+		return Score < Other.Score;
+	}
+
+	FEnvQueryItem() : Score(0.0f), DataOffset(-1), bIsDiscarded(false) {}
+	FEnvQueryItem(int32 InOffset) : Score(0.0f), DataOffset(InOffset), bIsDiscarded(false) {}
+};
+
+template <> struct TIsZeroConstructType<FEnvQueryItem> { enum { Value = true }; };
+
+struct AIMODULE_API FEnvQueryResult
+{
+	TArray<FEnvQueryItem> Items;
+
+	/** type of generated items */
+	TSubclassOf<class UEnvQueryItemType> ItemType;
+
+	/** query status */
+	TEnumAsByte<EEnvQueryStatus::Type> Status;
+
+	/** raw data of items */
+	TArray<uint8> RawData;
+
+	/** index of query option, that generated items */
+	int32 OptionIndex;
+
+	/** instance ID */
+	int32 QueryID;
+
+	/** instance owner */
+	TWeakObjectPtr<UObject> Owner;
+
+	FORCEINLINE float GetItemScore(int32 Index) const { return Items.IsValidIndex(Index) ? Items[Index].Score : 0.0f; }
+
+	/** item accessors for basic types */
+	AActor* GetItemAsActor(int32 Index) const;
+	FVector GetItemAsLocation(int32 Index) const;
+
+	FEnvQueryResult() : ItemType(NULL), Status(EEnvQueryStatus::Processing), OptionIndex(0) {}
+	FEnvQueryResult(const EEnvQueryStatus::Type& InStatus) : ItemType(NULL), Status(InStatus), OptionIndex(0) {}
+};
+
+
+//////////////////////////////////////////////////////////////////////////
+// Runtime processing structures
+
+DECLARE_DELEGATE_OneParam(FQueryFinishedSignature, TSharedPtr<struct FEnvQueryResult>);
+DECLARE_DELEGATE_OneParam(FGenerateItemsSignature, struct FEnvQueryInstance&);
+DECLARE_DELEGATE_OneParam(FExecuteTestSignature, struct FEnvQueryInstance&);
+
+struct AIMODULE_API FEnvQuerySpatialData
+{
+	FVector Location;
+	FRotator Rotation;
+};
+
+/** Detailed information about item, used by tests */
+struct AIMODULE_API FEnvQueryItemDetails
+{
+	/** Results assigned by option's tests, before any modifications */
+	TArray<float> TestResults;
+
+#if USE_EQS_DEBUGGER
+	/** Results assigned by option's tests, after applying modifiers, normalization and weight */
+	TArray<float> TestWeightedScores;
+
+	int32 FailedTestIndex;
+	int32 ItemIndex;
+#endif // USE_EQS_DEBUGGER
+
+	FEnvQueryItemDetails() {}
+	FEnvQueryItemDetails(int32 NumTests, int32 InItemIndex)
+	{
+		TestResults.AddZeroed(NumTests);
+#if USE_EQS_DEBUGGER
+		TestWeightedScores.AddZeroed(NumTests);
+		ItemIndex = InItemIndex;
+		FailedTestIndex = INDEX_NONE;
+#endif
+	}
+
+	FORCEINLINE uint32 GetAllocatedSize() const
+	{
+		return sizeof(*this) +
+#if USE_EQS_DEBUGGER
+			TestWeightedScores.GetAllocatedSize() +
+#endif
+			TestResults.GetAllocatedSize();
+	}
+};
+
+struct AIMODULE_API FEnvQueryContextData
+{
+	/** type of context values */
+	TSubclassOf<class UEnvQueryItemType> ValueType;
+
+	/** number of stored values */
+	int32 NumValues;
+
+	/** data of stored values */
+	TArray<uint8> RawData;
+
+	FEnvQueryContextData()
+		: NumValues(0)
+	{}
+
+	FORCEINLINE uint32 GetAllocatedSize() const { return sizeof(*this) + RawData.GetAllocatedSize(); }
+};
+
+struct AIMODULE_API FEnvQueryOptionInstance
+{
+	/** generator's delegate */
+	FGenerateItemsSignature GenerateDelegate;
+
+	/** tests' delegates */
+	TArray<FExecuteTestSignature> TestDelegates;
+
+	/** type of generated items */
+	TSubclassOf<class UEnvQueryItemType> ItemType;
+
+	/** is set, items will be shuffled after tests */
+	bool bShuffleItems;
+
+	FORCEINLINE uint32 GetAllocatedSize() const { return sizeof(*this) + TestDelegates.GetAllocatedSize(); }
+};
+
+#if NO_LOGGING
+#define EQSHEADERLOG(...)
+#else
+#define EQSHEADERLOG(msg) Log(msg)
+#endif // NO_LOGGING
+
+struct FEQSQueryDebugData
+{
+	TArray<FEnvQueryItem> DebugItems;
+	TArray<FEnvQueryItemDetails> DebugItemDetails;
+	TArray<uint8> RawData;
+	TArray<FString> PerformedTestNames;
+
+	void Store(const FEnvQueryInstance* QueryInstance);
+	void Reset()
+	{
+		DebugItems.Reset();
+		DebugItemDetails.Reset();
+		RawData.Reset();
+		PerformedTestNames.Reset();
+	}
+};
+
+struct AIMODULE_API FEnvQueryInstance : public FEnvQueryResult
+{
+	typedef float FNamedParamValueType;
+
+	/** short name of query template */
+	FString QueryName;
+
+	/** world owning this query instance */
+	UWorld* World;
+
+	/** observer's delegate */
+	FQueryFinishedSignature FinishDelegate;
+
+	/** execution params */
+	TMap<FName, FNamedParamValueType> NamedParams;
+
+	/** contexts in use */
+	TMap<UClass*, FEnvQueryContextData> ContextCache;
+
+	/** list of options */
+	TArray<struct FEnvQueryOptionInstance> Options;
+
+	/** currently processed test (-1 = generator) */
+	int32 CurrentTest;
+
+	/** non-zero if test run last step has been stopped mid-process. This indicates
+	 *	index of the first item that needs processing when resumed */
+	int32 CurrentTestStartingItem;
+
+	/** list of item details */
+	TArray<struct FEnvQueryItemDetails> ItemDetails;
+
+	/** number of valid items on list */
+	int32 NumValidItems;
+
+	/** size of current value */
+	uint16 ValueSize;
+
+	/** used to breaking from item iterator loops */
+	uint8 bFoundSingleResult : 1;
+
+	/** set when testing final condition of an option */
+	uint8 bPassOnSingleResult : 1;
+
+#if USE_EQS_DEBUGGER
+	/** set to true to store additional debug info */
+	uint8 bStoreDebugInfo : 1;
+#endif // USE_EQS_DEBUGGER
+
+	/** run mode */
+	TEnumAsByte<EEnvQueryRunMode::Type> Mode;
+
+	/** item type's CDO for location tests */
+	class UEnvQueryItemType_VectorBase* ItemTypeVectorCDO;
+
+	/** item type's CDO for actor tests */
+	class UEnvQueryItemType_ActorBase* ItemTypeActorCDO;
+
+	/** if > 0 then it's how much time query has for performing current step */
+	double TimeLimit;
+
+	FEnvQueryInstance() : World(NULL), CurrentTest(-1), NumValidItems(0), bFoundSingleResult(false)
+#if USE_EQS_DEBUGGER
+		, bStoreDebugInfo(bDebuggingInfoEnabled) 
+#endif // USE_EQS_DEBUGGER
+	{ IncStats(); }
+	FEnvQueryInstance(const FEnvQueryInstance& Other) { *this = Other; IncStats(); }
+	~FEnvQueryInstance() { DecStats(); }
+
+	/** execute single step of query */
+	void ExecuteOneStep(double TimeLimit);
+
+	/** update context cache */
+	bool PrepareContext(UClass* Context, FEnvQueryContextData& ContextData);
+
+	/** helpers for reading spatial data from context */
+	bool PrepareContext(UClass* Context, TArray<FEnvQuerySpatialData>& Data);
+	bool PrepareContext(UClass* Context, TArray<FVector>& Data);
+	bool PrepareContext(UClass* Context, TArray<FRotator>& Data);
+	/** helpers for reading actor data from context */
+	bool PrepareContext(UClass* Context, TArray<AActor*>& Data);
+
+	/** access named params */
+	template<typename TEQSParam>
+	FORCEINLINE bool GetParamValue(const TEQSParam& Param, typename TEQSParam::FValueType& OutValue, const FString& ParamDesc)
+	{
+		if (Param.ParamName != NAME_None)
+		{
+			FNamedParamValueType* PtrValue = NamedParams.Find(Param.ParamName);
+			if (PtrValue == NULL)
+			{
+				EQSHEADERLOG(FString::Printf(TEXT("Query [%s] is missing param [%s] for [%s] property!"), *QueryName, *Param.ParamName.ToString(), *ParamDesc));				
+				Status = EEnvQueryStatus::MissingParam;
+				return false;
+			}
+
+			OutValue = *((typename TEQSParam::FValueType*)PtrValue);
+		}
+		else
+		{
+			OutValue = Param.Value;
+		}
+
+		return true;
+	}
+
+	/** raw data operations */
+	void ReserveItemData(int32 NumAdditionalItems);
+
+	template<typename TypeItem, typename TypeValue>
+	void AddItemData(TypeValue ItemValue)
+	{
+		DEC_MEMORY_STAT_BY(STAT_AI_EQS_InstanceMemory, RawData.GetAllocatedSize() + Items.GetAllocatedSize());
+
+		const int32 DataOffset = RawData.AddUninitialized(ValueSize);
+		TypeItem::SetValue(RawData.GetTypedData() + DataOffset, ItemValue);
+		Items.Add(FEnvQueryItem(DataOffset));
+
+		INC_MEMORY_STAT_BY(STAT_AI_EQS_InstanceMemory, RawData.GetAllocatedSize() + Items.GetAllocatedSize());
+	}
+
+protected:
+
+	/** prepare item data after generator has finished */
+	void FinalizeGeneration();
+
+	/** update costs and flags after test has finished */
+	void FinalizeTest();
+	
+	/** final pass on items of finished query */
+	void FinalizeQuery();
+
+	/** normalize total score in range 0..1 */
+	void NormalizeScores();
+
+	/** sort all scores, from highest to lowest */
+	void SortScores();
+
+	/** pick one of items with highest score */
+	void PickBestItem();
+
+	/** prepare items on reaching final condition in SingleResult mode */
+	void OnFinalCondition();
+
+	/** discard all items but one */
+	void PickSingleItem(int32 ItemIndex);
+
+public:
+
+	/** removes all runtime data that can be used for debugging (not a part of actual query result) */
+	void StripRedundantData();
+
+#if STATS
+	FORCEINLINE void IncStats()
+	{
+		INC_MEMORY_STAT_BY(STAT_AI_EQS_InstanceMemory, GetAllocatedSize());
+		INC_DWORD_STAT_BY(STAT_AI_EQS_NumItems, Items.Num());
+	}
+
+	FORCEINLINE void DecStats()
+	{
+		DEC_MEMORY_STAT_BY(STAT_AI_EQS_InstanceMemory, GetAllocatedSize()); 
+		DEC_DWORD_STAT_BY(STAT_AI_EQS_NumItems, Items.Num());
+	}
+
+	uint32 GetAllocatedSize() const;
+	uint32 GetContextAllocatedSize() const;
+#else
+	FORCEINLINE uint32 GetAllocatedSize() const { return 0; }
+	FORCEINLINE uint32 GetContextAllocatedSize() const { return 0; }
+	FORCEINLINE void IncStats() {}
+	FORCEINLINE void DecStats() {}
+#endif // STATS
+
+#if !NO_LOGGING
+	void Log(const FString Msg) const;
+#endif // #if !NO_LOGGING
+
+#if CPP || UE_BUILD_DOCS
+	struct AIMODULE_API ItemIterator
+	{
+		ItemIterator(class UEnvQueryTest* QueryTest, FEnvQueryInstance& QueryInstance);
+
+		~ItemIterator()
+		{
+			Instance->CurrentTestStartingItem = CurrentItem;
+		}
+
+		void SetScore(EEnvTestPurpose::Type TestPurpose, EEnvTestFilterType::Type FilterType, float Score, float Min, float Max)
+		{
+			bool bPassedTest = true;
+			switch (FilterType)
+			{
+				case EEnvTestFilterType::Maximum:
+					if (Score > Max)
+					{
+						bPassedTest = false;
+					}
+					break;
+
+				case EEnvTestFilterType::Minimum:
+					if (Score < Min)
+					{
+						bPassedTest = false;
+					}
+					break;
+
+				case EEnvTestFilterType::Range:
+					if ((Score < Min) || (Score > Max))
+					{
+						bPassedTest = false;
+					}
+					break;
+
+				case EEnvTestFilterType::Match:
+					UE_LOG(LogEQS, Error, TEXT("Filtering Type set to 'Match' for floating point test.  Will consider test as failed in all cases."));
+					bPassedTest = false;
+					break;
+
+				default:
+					UE_LOG(LogEQS, Error, TEXT("Filtering Type set to invalid value for floating point test.  Will consider test as failed in all cases."));
+					bPassedTest = false;
+					break;
+			}
+
+			if (!bPassedTest)
+			{
+				if (TestPurpose == EEnvTestPurpose::Score)
+				{
+					bSkipped = true;
+					NumPartialScores++;
+				}
+				else // We are filtering!
+				{
+					bPassed = false;
+				}
+			}
+			else
+			{
+				ItemScore += Score;
+				NumPartialScores++;
+			}
+		}
+
+		void SetScore(EEnvTestPurpose::Type TestPurpose, EEnvTestFilterType::Type FilterType, bool bScore, bool bExpected)
+		{
+			bool bPassedTest = true;
+			switch (FilterType)
+			{
+				case EEnvTestFilterType::Match:
+					if (bScore != bExpected)
+					{
+						bPassedTest = false;
+					}
+					break;
+
+				case EEnvTestFilterType::Maximum:
+					UE_LOG(LogEQS, Error, TEXT("Filtering Type set to 'Maximum' for boolean test.  Will consider test as failed in all cases."));
+					bPassedTest = false;
+					break;
+
+				case EEnvTestFilterType::Minimum:
+					UE_LOG(LogEQS, Error, TEXT("Filtering Type set to 'Minimum' for boolean test.  Will consider test as failed in all cases."));
+					bPassedTest = false;
+					break;
+
+				case EEnvTestFilterType::Range:
+					UE_LOG(LogEQS, Error, TEXT("Filtering Type set to 'Range' for boolean test.  Will consider test as failed in all cases."));
+					bPassedTest = false;
+					break;
+
+				default:
+					UE_LOG(LogEQS, Error, TEXT("Filtering Type set to invalid value for boolean test.  Will consider test as failed in all cases."));
+					bPassedTest = false;
+					break;
+			}
+
+			if (!bPassedTest)
+			{
+				if (TestPurpose == EEnvTestPurpose::Score)
+				{
+					bSkipped = true;
+					NumPartialScores++;
+				}
+				else // We are filtering!
+				{
+					bPassed = false;
+				}
+			}
+			else
+			{
+				ItemScore += 1.0f;
+				NumPartialScores++;
+			}
+		}
+
+		uint8* GetItemData()
+		{
+			return Instance->RawData.GetTypedData() + Instance->Items[CurrentItem].DataOffset;
+		}
+
+		void DiscardItem()
+		{
+			bPassed = false;
+		}
+
+		void SkipItem()
+		{
+			bSkipped++;
+		}
+
+		operator bool() const
+		{
+			return CurrentItem < Instance->Items.Num() && !Instance->bFoundSingleResult && (Deadline < 0 || FPlatformTime::Seconds() < Deadline);
+		}
+
+		int32 operator*() const
+		{
+			return CurrentItem; 
+		}
+
+		void operator++()
+		{
+			StoreTestResult();
+			if (!Instance->bFoundSingleResult)
+			{
+				InitItemScore();
+				for (CurrentItem++; CurrentItem < Instance->Items.Num() && !Instance->Items[CurrentItem].IsValid(); CurrentItem++) ;
+			}
+		}
+
+	protected:
+
+		FEnvQueryInstance* Instance;
+		UEnvQueryTest* Test;
+		int32 CurrentItem;
+		int32 NumPartialScores;
+		double Deadline;
+		float ItemScore;
+		uint32 bPassed : 1;
+		uint32 bSkipped : 1;
+		uint32 bDiscardFailed : 1;
+
+		void InitItemScore()
+		{
+			NumPartialScores = 0;
+			ItemScore = 0.0f;
+			bPassed = true;
+			bSkipped = false;
+		}
+
+		void StoreTestResult();
+	};
+#endif
+
+#if USE_EQS_DEBUGGER
+	FEQSQueryDebugData DebugData;
+	static bool bDebuggingInfoEnabled;
+#endif // USE_EQS_DEBUGGER
+	FBox GetBoundingBox() const;
+};
+
+namespace FEQSHelpers
+{
+#if WITH_RECAST
+	const class ARecastNavMesh* FindNavMeshForQuery(struct FEnvQueryInstance& QueryInstance);
+#endif // WITH_RECAST
+}
+
+UCLASS(Abstract, CustomConstructor)
+class AIMODULE_API UEnvQueryTypes : public UObject
+{
+	GENERATED_UCLASS_BODY()
+
+	/** special test value assigned to items skipped by condition check */
+	static float SkippedItemValue;
+
+	UEnvQueryTypes(const class FPostConstructInitializeProperties& PCIP) : Super(PCIP) {}
+	static FText GetShortTypeName(const UObject* Ob);
+
+	static FText DescribeContext(TSubclassOf<class UEnvQueryContext> ContextClass);
+	static FString DescribeIntParam(const FEnvIntParam& Param);
+	static FString DescribeFloatParam(const FEnvFloatParam& Param);
+	static FString DescribeBoolParam(const FEnvBoolParam& Param);
+};
diff --git a/Engine/Source/Runtime/AIModule/Classes/EnvironmentQuery/Generators/EnvQueryGenerator_Composite.h b/Engine/Source/Runtime/AIModule/Classes/EnvironmentQuery/Generators/EnvQueryGenerator_Composite.h
new file mode 100644
--- /dev/null
+++ b/Engine/Source/Runtime/AIModule/Classes/EnvironmentQuery/Generators/EnvQueryGenerator_Composite.h
@@ -0,0 +1,22 @@
+// Copyright 1998-2014 Epic Games, Inc. All Rights Reserved.
+
+#pragma once
+#include "EnvironmentQuery/EnvQueryGenerator.h"
+#include "EnvQueryGenerator_Composite.generated.h"
+
+/**
+ * Composite generator allows using multiple generators in single query option
+ * Currently it's available only from code
+ */
+
+UCLASS(MinimalAPI)
+class UEnvQueryGenerator_Composite : public UEnvQueryGenerator
+{
+	GENERATED_UCLASS_BODY()
+
+	UPROPERTY()
+	TArray<class UEnvQueryGenerator*> Generators;
+
+	void GenerateItems(struct FEnvQueryInstance& QueryInstance);
+	virtual FText GetDescriptionTitle() const OVERRIDE;
+};
diff --git a/Engine/Source/Runtime/AIModule/Classes/EnvironmentQuery/Generators/EnvQueryGenerator_OnCircle.h b/Engine/Source/Runtime/AIModule/Classes/EnvironmentQuery/Generators/EnvQueryGenerator_OnCircle.h
new file mode 100644
--- /dev/null
+++ b/Engine/Source/Runtime/AIModule/Classes/EnvironmentQuery/Generators/EnvQueryGenerator_OnCircle.h
@@ -0,0 +1,62 @@
+// Copyright 1998-2014 Epic Games, Inc. All Rights Reserved.
+
+#pragma once
+#include "EnvironmentQuery/Generators/EnvQueryGenerator_ProjectedPoints.h"
+#include "EnvQueryGenerator_OnCircle.generated.h"
+
+UCLASS()	
+class UEnvQueryGenerator_OnCircle : public UEnvQueryGenerator_ProjectedPoints
+{
+	GENERATED_UCLASS_BODY()
+
+	/** max distance of path between point and context */
+	UPROPERTY(EditDefaultsOnly, Category=Generator)
+	FEnvFloatParam Radius;
+
+	/** items will be generated on a circle this much apart */
+	UPROPERTY(EditDefaultsOnly, Category=Generator)
+	FEnvFloatParam ItemSpacing;
+
+	/** If you generate items on a piece of circle you define direction of Arc cut here */
+	UPROPERTY(EditDefaultsOnly, Category=Generator, meta=(EditCondition="bDefineArc"))
+	FEnvDirection ArcDirection;
+
+	/** If you generate items on a piece of circle you define angle of Arc cut here */
+	UPROPERTY(EditDefaultsOnly, Category=Generator)
+	FEnvFloatParam Angle;
+	
+	UPROPERTY()
+	float AngleRadians;
+
+	/** context */
+	UPROPERTY(EditAnywhere, Category=Generator)
+	TSubclassOf<class UEnvQueryContext> CircleCenter;
+
+	/** horizontal trace for nearest obstacle */
+	UPROPERTY(EditAnywhere, Category=Generator)
+	FEnvTraceData TraceData;
+
+	UPROPERTY()
+	uint32 bDefineArc:1;
+
+	virtual void PostLoad() OVERRIDE;
+
+	void GenerateItems(struct FEnvQueryInstance& QueryInstance); 
+
+	virtual FText GetDescriptionTitle() const OVERRIDE;
+	virtual FText GetDescriptionDetails() const OVERRIDE;
+
+#if WITH_EDITOR
+	virtual void PostEditChangeProperty( struct FPropertyChangedEvent& PropertyChangedEvent) OVERRIDE;
+#endif // WITH_EDITOR
+
+protected:
+	FVector CalcDirection(struct FEnvQueryInstance& QueryInstance) const;
+
+	DECLARE_DELEGATE_SevenParams(FRunTraceSignature, const FVector&, const FVector&, UWorld*, enum ECollisionChannel, const FCollisionQueryParams&, const FVector&, FVector&);
+
+	void RunLineTrace(const FVector& StartPos, const FVector& EndPos, UWorld* World, enum ECollisionChannel Channel, const FCollisionQueryParams& Params, const FVector& Extent, FVector& HitPos);
+	void RunSphereTrace(const FVector& StartPos, const FVector& EndPos, UWorld* World, enum ECollisionChannel Channel, const FCollisionQueryParams& Params, const FVector& Extent, FVector& HitPos);
+	void RunCapsuleTrace(const FVector& StartPos, const FVector& EndPos, UWorld* World, enum ECollisionChannel Channel, const FCollisionQueryParams& Params, const FVector& Extent, FVector& HitPos);
+	void RunBoxTrace(const FVector& StartPos, const FVector& EndPos, UWorld* World, enum ECollisionChannel Channel, const FCollisionQueryParams& Params, const FVector& Extent, FVector& HitPos);
+};
diff --git a/Engine/Source/Runtime/AIModule/Classes/EnvironmentQuery/Generators/EnvQueryGenerator_PathingGrid.h b/Engine/Source/Runtime/AIModule/Classes/EnvironmentQuery/Generators/EnvQueryGenerator_PathingGrid.h
new file mode 100644
--- /dev/null
+++ b/Engine/Source/Runtime/AIModule/Classes/EnvironmentQuery/Generators/EnvQueryGenerator_PathingGrid.h
@@ -0,0 +1,55 @@
+// Copyright 1998-2014 Epic Games, Inc. All Rights Reserved.
+
+#pragma once
+#include "EnvironmentQuery/EnvQueryGenerator.h"
+#include "EnvQueryGenerator_PathingGrid.generated.h"
+
+/**
+ *  Navigation grid, generates points on navmesh
+ *  with paths to/from context no further than given limit
+ */
+
+UCLASS()
+class UEnvQueryGenerator_PathingGrid : public UEnvQueryGenerator
+{
+	GENERATED_UCLASS_BODY()
+
+	/** max distance of path between point and context */
+	UPROPERTY(EditDefaultsOnly, Category=Generator)
+	FEnvFloatParam MaxPathDistance;
+
+	/** generation density */
+	UPROPERTY(EditDefaultsOnly, Category=Generator)
+	FEnvFloatParam Density;
+
+	/** path direction switch: set = path from context to points, clear = path from points to context */
+	UPROPERTY(EditDefaultsOnly, Category=Generator)
+	FEnvBoolParam PathFromContext;
+
+	/** context */
+	UPROPERTY(EditAnywhere, Category=Generator)
+	TSubclassOf<class UEnvQueryContext> GenerateAround;
+
+	/** context */
+	UPROPERTY(EditAnywhere, Category=Generator)
+	TEnumAsByte<EEnvQueryTrace::Type> TraceType;
+
+	/** navigation filter for tracing */
+	UPROPERTY(EditDefaultsOnly, Category=Generator)
+	TSubclassOf<class UNavigationQueryFilter> NavigationFilter;
+
+	void GenerateItems(struct FEnvQueryInstance& QueryInstance); 
+
+	virtual FText GetDescriptionTitle() const OVERRIDE;
+	virtual FText GetDescriptionDetails() const OVERRIDE;
+
+protected:
+#if WITH_RECAST
+
+	/** find all navmesh node refs in pathing distance */
+	void FindNodeRefsInPathDistance(const class ARecastNavMesh* NavMesh, const FVector& ContextLocation, float MaxPathDistance, bool bPathFromContext, TArray<NavNodeRef>& NodeRefs, FBox& NodeRefsBounds) const;
+
+	/** check if nav location is in allowed set */
+	bool IsNavLocationInPathDistance(const class ARecastNavMesh* NavMesh, const struct FNavLocation& NavLocation, const TArray<NavNodeRef>& NodeRefs) const;
+#endif
+};
diff --git a/Engine/Source/Runtime/AIModule/Classes/EnvironmentQuery/Generators/EnvQueryGenerator_ProjectedPoints.h b/Engine/Source/Runtime/AIModule/Classes/EnvironmentQuery/Generators/EnvQueryGenerator_ProjectedPoints.h
new file mode 100644
--- /dev/null
+++ b/Engine/Source/Runtime/AIModule/Classes/EnvironmentQuery/Generators/EnvQueryGenerator_ProjectedPoints.h
@@ -0,0 +1,58 @@
+// Copyright 1998-2013 Epic Games, Inc. All Rights Reserved.
+
+#pragma once
+#include "EnvironmentQuery/EnvQueryGenerator.h"
+#include "EnvironmentQuery/EnvQueryTypes.h"
+#include "EnvQueryGenerator_ProjectedPoints.generated.h"
+
+UCLASS(Abstract)
+class UEnvQueryGenerator_ProjectedPoints : public UEnvQueryGenerator
+{
+	GENERATED_UCLASS_BODY()
+
+	/** trace params */
+	UPROPERTY(EditDefaultsOnly, Category=Generator)
+	FEnvTraceData ProjectionData;
+
+	struct FSortByHeight
+	{
+		float OriginalZ;
+
+		FSortByHeight(const FVector& OriginalPt) : OriginalZ(OriginalPt.Z) {}
+
+		FORCEINLINE bool operator()(const FNavLocation& A, const FNavLocation& B) const
+		{
+			return FMath::Abs(A.Location.Z - OriginalZ) < FMath::Abs(B.Location.Z - OriginalZ);
+		}
+	};
+
+	FORCEINLINE_DEBUGGABLE bool ProjectNavPointSimple(FVector& Point, const FVector& Extent, const class ANavigationData* NavData, TSharedPtr<const FNavigationQueryFilter> NavFilter)
+	{
+		FNavLocation ProjectedPoint;
+		if (NavData->ProjectPoint(Point, ProjectedPoint, Extent, NavFilter))
+		{
+			Point.Z = ProjectedPoint.Location.Z;
+			return true;
+		}
+
+		return false;
+	}
+
+#if WITH_RECAST
+	FORCEINLINE_DEBUGGABLE bool ProjectNavPointInRange(FVector& Point, const float Radius, const float MinZOffset, const float MaxZOffset, const class ARecastNavMesh* NavData, TSharedPtr<const FNavigationQueryFilter> NavFilter, TArray<FNavLocation>& TempPoints)
+	{
+		TempPoints.Reset();
+		if (NavData->ProjectPointMulti(Point, TempPoints, FVector(Radius, Radius, 0), Point.Z - MinZOffset, Point.Z + MaxZOffset, NavFilter))
+		{
+			TempPoints.Sort(FSortByHeight(Point));
+			Point.Z = TempPoints[0].Location.Z;
+			return true;
+		}
+
+		return false;
+	}
+#endif
+
+	/** project all points in array and remove those outside navmesh */
+	void ProjectAndFilterNavPoints(TArray<FVector>& Points, const class ANavigationData* NavData);
+};
diff --git a/Engine/Source/Runtime/AIModule/Classes/EnvironmentQuery/Generators/EnvQueryGenerator_SimpleGrid.h b/Engine/Source/Runtime/AIModule/Classes/EnvironmentQuery/Generators/EnvQueryGenerator_SimpleGrid.h
new file mode 100644
--- /dev/null
+++ b/Engine/Source/Runtime/AIModule/Classes/EnvironmentQuery/Generators/EnvQueryGenerator_SimpleGrid.h
@@ -0,0 +1,32 @@
+// Copyright 1998-2014 Epic Games, Inc. All Rights Reserved.
+
+#pragma once
+#include "EnvironmentQuery/Generators/EnvQueryGenerator_ProjectedPoints.h"
+#include "EnvQueryGenerator_SimpleGrid.generated.h"
+
+/**
+ *  Simple grid, generates points in 2D square around context
+ */
+
+UCLASS()
+class UEnvQueryGenerator_SimpleGrid : public UEnvQueryGenerator_ProjectedPoints
+{
+	GENERATED_UCLASS_BODY()
+
+	/** square's extent */
+	UPROPERTY(EditDefaultsOnly, Category=Generator)
+	FEnvFloatParam Radius;
+
+	/** generation density */
+	UPROPERTY(EditDefaultsOnly, Category=Generator)
+	FEnvFloatParam Density;
+
+	/** context */
+	UPROPERTY(EditDefaultsOnly, Category=Generator)
+	TSubclassOf<class UEnvQueryContext> GenerateAround;
+
+	void GenerateItems(struct FEnvQueryInstance& QueryInstance); 
+
+	virtual FText GetDescriptionTitle() const OVERRIDE;
+	virtual FText GetDescriptionDetails() const OVERRIDE;
+};
diff --git a/Engine/Source/Runtime/AIModule/Classes/EnvironmentQuery/Items/EnvQueryItemType.h b/Engine/Source/Runtime/AIModule/Classes/EnvironmentQuery/Items/EnvQueryItemType.h
new file mode 100644
--- /dev/null
+++ b/Engine/Source/Runtime/AIModule/Classes/EnvironmentQuery/Items/EnvQueryItemType.h
@@ -0,0 +1,44 @@
+// Copyright 1998-2014 Epic Games, Inc. All Rights Reserved.
+
+#pragma once
+#include "EnvQueryItemType.generated.h"
+
+UCLASS(Abstract)
+class AIMODULE_API UEnvQueryItemType : public UObject
+{
+	GENERATED_UCLASS_BODY()
+
+	/** get ValueSize */
+	FORCEINLINE uint16 GetValueSize() const { return ValueSize; }
+
+	/** add filters for blackboard key selector */
+	virtual void AddBlackboardFilters(struct FBlackboardKeySelector& KeySelector, UObject* FilterOwner) const;
+
+	/** store value in blackboard entry */
+	virtual bool StoreInBlackboard(struct FBlackboardKeySelector& KeySelector, class UBlackboardComponent* Blackboard, const uint8* RawData) const;
+
+	/** unregister from known types in EnvQueryManager */
+	virtual void FinishDestroy() OVERRIDE;
+
+	/** describe item */
+	virtual FString GetDescription(const uint8* RawData) const;
+
+protected:
+
+	/** size of value for this type */
+	uint16 ValueSize;
+
+	/** helper function for reading typed data from memory block */
+	template<typename T>
+	static T GetValueFromMemory(const uint8* MemoryBlock)
+	{
+		return *((T*)MemoryBlock);
+	}
+
+	/** helper function for writing typed data to memory block */
+	template<typename T>
+	static void SetValueInMemory(uint8* MemoryBlock, const T& Value)
+	{
+		*((T*)MemoryBlock) = Value;
+	}
+};
diff --git a/Engine/Source/Runtime/AIModule/Classes/EnvironmentQuery/Items/EnvQueryItemType_Actor.h b/Engine/Source/Runtime/AIModule/Classes/EnvironmentQuery/Items/EnvQueryItemType_Actor.h
new file mode 100644
--- /dev/null
+++ b/Engine/Source/Runtime/AIModule/Classes/EnvironmentQuery/Items/EnvQueryItemType_Actor.h
@@ -0,0 +1,23 @@
+// Copyright 1998-2014 Epic Games, Inc. All Rights Reserved.
+
+#pragma once
+
+#include "EnvironmentQuery/Items/EnvQueryItemType_ActorBase.h"
+#include "EnvQueryItemType_Actor.generated.h"
+
+UCLASS()
+class AIMODULE_API UEnvQueryItemType_Actor : public UEnvQueryItemType_ActorBase
+{
+	GENERATED_UCLASS_BODY()
+
+	static AActor* GetValue(const uint8* RawData);
+	static void SetValue(uint8* RawData, const AActor* Value);
+
+	static void SetContextHelper(struct FEnvQueryContextData& ContextData, const AActor* SingleActor);
+	static void SetContextHelper(struct FEnvQueryContextData& ContextData, const TArray<const AActor*>& MultipleActors);
+	static void SetContextHelper(struct FEnvQueryContextData& ContextData, const TArray<AActor*>& MultipleActors);
+
+	virtual FVector GetLocation(const uint8* RawData) const OVERRIDE;
+	virtual FRotator GetRotation(const uint8* RawData) const OVERRIDE;
+	virtual AActor* GetActor(const uint8* RawData) const OVERRIDE;
+};
diff --git a/Engine/Source/Runtime/AIModule/Classes/EnvironmentQuery/Items/EnvQueryItemType_ActorBase.h b/Engine/Source/Runtime/AIModule/Classes/EnvironmentQuery/Items/EnvQueryItemType_ActorBase.h
new file mode 100644
--- /dev/null
+++ b/Engine/Source/Runtime/AIModule/Classes/EnvironmentQuery/Items/EnvQueryItemType_ActorBase.h
@@ -0,0 +1,19 @@
+// Copyright 1998-2014 Epic Games, Inc. All Rights Reserved.
+
+#pragma once
+
+#include "EnvironmentQuery/Items/EnvQueryItemType_VectorBase.h"
+#include "EnvQueryItemType_ActorBase.generated.h"
+
+
+UCLASS(Abstract)
+class AIMODULE_API UEnvQueryItemType_ActorBase : public UEnvQueryItemType_VectorBase
+{
+	GENERATED_UCLASS_BODY()
+
+	virtual AActor* GetActor(const uint8* RawData) const;
+
+	virtual void AddBlackboardFilters(struct FBlackboardKeySelector& KeySelector, UObject* FilterOwner) const OVERRIDE;
+	virtual bool StoreInBlackboard(struct FBlackboardKeySelector& KeySelector, class UBlackboardComponent* Blackboard, const uint8* RawData) const OVERRIDE;
+	virtual FString GetDescription(const uint8* RawData) const OVERRIDE;
+};
diff --git a/Engine/Source/Runtime/AIModule/Classes/EnvironmentQuery/Items/EnvQueryItemType_Direction.h b/Engine/Source/Runtime/AIModule/Classes/EnvironmentQuery/Items/EnvQueryItemType_Direction.h
new file mode 100644
--- /dev/null
+++ b/Engine/Source/Runtime/AIModule/Classes/EnvironmentQuery/Items/EnvQueryItemType_Direction.h
@@ -0,0 +1,25 @@
+// Copyright 1998-2013 Epic Games, Inc. All Rights Reserved.
+
+#pragma once
+#include "EnvironmentQuery/Items/EnvQueryItemType_VectorBase.h"
+#include "EnvQueryItemType_Direction.generated.h"
+
+UCLASS()
+class AIMODULE_API UEnvQueryItemType_Direction : public UEnvQueryItemType_VectorBase
+{
+	GENERATED_UCLASS_BODY()
+
+	static FVector GetValue(const uint8* RawData);
+	static void SetValue(uint8* RawData, const FVector& Value);
+
+	static FRotator GetValueRot(const uint8* RawData);
+	static void SetValueRot(uint8* RawData, const FRotator& Value);
+
+	static void SetContextHelper(struct FEnvQueryContextData& ContextData, const FVector& SingleDirection);
+	static void SetContextHelper(struct FEnvQueryContextData& ContextData, const FRotator& SingleRotation);
+	static void SetContextHelper(struct FEnvQueryContextData& ContextData, const TArray<FVector>& MultipleDirections);
+	static void SetContextHelper(struct FEnvQueryContextData& ContextData, const TArray<FRotator>& MultipleRotations);
+
+	virtual FRotator GetRotation(const uint8* RawData) const OVERRIDE;
+	virtual FString GetDescription(const uint8* RawData) const OVERRIDE;
+};
diff --git a/Engine/Source/Runtime/AIModule/Classes/EnvironmentQuery/Items/EnvQueryItemType_Point.h b/Engine/Source/Runtime/AIModule/Classes/EnvironmentQuery/Items/EnvQueryItemType_Point.h
new file mode 100644
--- /dev/null
+++ b/Engine/Source/Runtime/AIModule/Classes/EnvironmentQuery/Items/EnvQueryItemType_Point.h
@@ -0,0 +1,19 @@
+// Copyright 1998-2014 Epic Games, Inc. All Rights Reserved.
+
+#pragma once
+#include "EnvironmentQuery/Items/EnvQueryItemType_VectorBase.h"
+#include "EnvQueryItemType_Point.generated.h"
+
+UCLASS()
+class AIMODULE_API UEnvQueryItemType_Point : public UEnvQueryItemType_VectorBase
+{
+	GENERATED_UCLASS_BODY()
+
+	static FVector GetValue(const uint8* RawData);
+	static void SetValue(uint8* RawData, const FVector& Value);
+
+	static void SetContextHelper(struct FEnvQueryContextData& ContextData, const FVector& SinglePoint);
+	static void SetContextHelper(struct FEnvQueryContextData& ContextData, const TArray<FVector>& MultiplePoints);
+
+	virtual FVector GetLocation(const uint8* RawData) const OVERRIDE;
+};
diff --git a/Engine/Source/Runtime/AIModule/Classes/EnvironmentQuery/Items/EnvQueryItemType_VectorBase.h b/Engine/Source/Runtime/AIModule/Classes/EnvironmentQuery/Items/EnvQueryItemType_VectorBase.h
new file mode 100644
--- /dev/null
+++ b/Engine/Source/Runtime/AIModule/Classes/EnvironmentQuery/Items/EnvQueryItemType_VectorBase.h
@@ -0,0 +1,18 @@
+// Copyright 1998-2014 Epic Games, Inc. All Rights Reserved.
+
+#pragma once
+#include "EnvironmentQuery/Items/EnvQueryItemType.h"
+#include "EnvQueryItemType_VectorBase.generated.h"
+
+UCLASS(Abstract)
+class AIMODULE_API UEnvQueryItemType_VectorBase : public UEnvQueryItemType
+{
+	GENERATED_UCLASS_BODY()
+
+	virtual FVector GetLocation(const uint8* RawData) const;
+	virtual FRotator GetRotation(const uint8* RawData) const;
+
+	virtual void AddBlackboardFilters(struct FBlackboardKeySelector& KeySelector, UObject* FilterOwner) const OVERRIDE;
+	virtual bool StoreInBlackboard(struct FBlackboardKeySelector& KeySelector, class UBlackboardComponent* Blackboard, const uint8* RawData) const OVERRIDE;
+	virtual FString GetDescription(const uint8* RawData) const OVERRIDE;
+};
diff --git a/Engine/Source/Runtime/AIModule/Classes/EnvironmentQuery/Tests/EnvQueryTest_Distance.h b/Engine/Source/Runtime/AIModule/Classes/EnvironmentQuery/Tests/EnvQueryTest_Distance.h
new file mode 100644
--- /dev/null
+++ b/Engine/Source/Runtime/AIModule/Classes/EnvironmentQuery/Tests/EnvQueryTest_Distance.h
@@ -0,0 +1,35 @@
+// Copyright 1998-2014 Epic Games, Inc. All Rights Reserved.
+
+#pragma once
+#include "EnvironmentQuery/EnvQueryTest.h"
+#include "EnvQueryTest_Distance.generated.h"
+
+UENUM()
+namespace EEnvTestDistance
+{
+	enum Type
+	{
+		Distance3D,
+		Distance2D,
+		DistanceZ,
+	};
+}
+
+UCLASS()
+class UEnvQueryTest_Distance : public UEnvQueryTest
+{
+	GENERATED_UCLASS_BODY()
+
+	/** testing mode */
+	UPROPERTY(EditDefaultsOnly, Category=Distance)
+	TEnumAsByte<EEnvTestDistance::Type> TestMode;
+
+	/** context */
+	UPROPERTY(EditDefaultsOnly, Category=Distance)
+	TSubclassOf<class UEnvQueryContext> DistanceTo;
+
+	void RunTest(struct FEnvQueryInstance& QueryInstance);
+
+	virtual FString GetDescriptionTitle() const OVERRIDE;
+	virtual FText GetDescriptionDetails() const OVERRIDE;
+};
diff --git a/Engine/Source/Runtime/AIModule/Classes/EnvironmentQuery/Tests/EnvQueryTest_Dot.h b/Engine/Source/Runtime/AIModule/Classes/EnvironmentQuery/Tests/EnvQueryTest_Dot.h
new file mode 100644
--- /dev/null
+++ b/Engine/Source/Runtime/AIModule/Classes/EnvironmentQuery/Tests/EnvQueryTest_Dot.h
@@ -0,0 +1,42 @@
+// Copyright 1998-2014 Epic Games, Inc. All Rights Reserved.
+
+#pragma once
+#include "EnvironmentQuery/EnvQueryTest.h"
+#include "EnvQueryTest_Dot.generated.h"
+
+UCLASS(MinimalAPI)
+class UEnvQueryTest_Dot : public UEnvQueryTest
+{
+	GENERATED_UCLASS_BODY()
+
+	/** defines direction of first line used by test */
+	UPROPERTY(EditDefaultsOnly, Category=Dot)
+	FEnvDirection LineA;
+
+	/** defines direction of second line used by test */
+	UPROPERTY(EditDefaultsOnly, Category=Dot)
+	FEnvDirection LineB;
+
+	void RunTest(struct FEnvQueryInstance& QueryInstance);
+
+	virtual FString GetDescriptionTitle() const OVERRIDE;
+	virtual FText GetDescriptionDetails() const OVERRIDE;
+
+protected:
+
+	/** helper function: gather directions from context pairs */
+	void GatherLineDirections(TArray<FVector>& Directions, struct FEnvQueryInstance& QueryInstance, const FVector& ItemLocation,
+		TSubclassOf<class UEnvQueryContext> LineFrom, TSubclassOf<class UEnvQueryContext> LineTo);
+
+	/** helper function: gather directions from context */
+	void GatherLineDirections(TArray<FVector>& Directions, struct FEnvQueryInstance& QueryInstance, const FRotator& ItemRotation,
+		TSubclassOf<class UEnvQueryContext> LineDirection);
+
+	/** helper function: gather directions from proper contexts */
+	void GatherLineDirections(TArray<FVector>& Directions, struct FEnvQueryInstance& QueryInstance,
+		TSubclassOf<class UEnvQueryContext> LineFrom, TSubclassOf<class UEnvQueryContext> LineTo, TSubclassOf<class UEnvQueryContext> LineDirection, bool bUseDirectionContext,
+		const FVector& ItemLocation = FVector::ZeroVector, const FRotator& ItemRotation = FRotator::ZeroRotator);
+
+	/** helper function: check if contexts are updated per item */
+	bool RequiresPerItemUpdates(TSubclassOf<class UEnvQueryContext> LineFrom, TSubclassOf<class UEnvQueryContext> LineTo, TSubclassOf<class UEnvQueryContext> LineDirection, bool bUseDirectionContext);
+};
diff --git a/Engine/Source/Runtime/AIModule/Classes/EnvironmentQuery/Tests/EnvQueryTest_Pathfinding.h b/Engine/Source/Runtime/AIModule/Classes/EnvironmentQuery/Tests/EnvQueryTest_Pathfinding.h
new file mode 100644
--- /dev/null
+++ b/Engine/Source/Runtime/AIModule/Classes/EnvironmentQuery/Tests/EnvQueryTest_Pathfinding.h
@@ -0,0 +1,66 @@
+// Copyright 1998-2014 Epic Games, Inc. All Rights Reserved.
+
+#pragma once
+#include "EnvironmentQuery/EnvQueryTest.h"
+#include "EnvQueryTest_Pathfinding.generated.h"
+
+UENUM()
+namespace EEnvTestPathfinding
+{
+	enum Type
+	{
+		PathExist,
+		PathCost,
+		PathLength,
+	};
+}
+
+UCLASS()
+class UEnvQueryTest_Pathfinding : public UEnvQueryTest
+{
+	GENERATED_UCLASS_BODY()
+
+	/** testing mode */
+	UPROPERTY(EditDefaultsOnly, Category=Pathfinding)
+	TEnumAsByte<EEnvTestPathfinding::Type> TestMode;
+
+	/** context: other end of pathfinding test */
+	UPROPERTY(EditDefaultsOnly, Category=Pathfinding)
+	TSubclassOf<class UEnvQueryContext> Context;
+
+	/** pathfinding direction */
+	UPROPERTY(EditDefaultsOnly, Category=Pathfinding)
+	FEnvBoolParam PathToItem;
+
+	/** if set, items with failed path will be invalidated (PathCost, PathLength) */
+	UPROPERTY(EditDefaultsOnly, Category=Pathfinding, AdvancedDisplay)
+	FEnvBoolParam DiscardUnreachable;
+
+	/** if set, hierarchical (faster) pathfinding will be used */
+	UPROPERTY(EditDefaultsOnly, Category=Pathfinding, AdvancedDisplay)
+	FEnvBoolParam HierarchicalPathfinding;
+
+	void RunTest(struct FEnvQueryInstance& QueryInstance);
+
+	virtual FString GetDescriptionTitle() const OVERRIDE;
+	virtual FText GetDescriptionDetails() const OVERRIDE;
+
+#if WITH_EDITOR
+	/** update test properties after changing mode */
+	virtual void PostEditChangeProperty(struct FPropertyChangedEvent& PropertyChangedEvent);
+#endif
+
+protected:
+
+	DECLARE_DELEGATE_RetVal_SixParams(bool, FTestPathSignature, const FVector&, const FVector&, EPathFindingMode::Type&, const ANavigationData*, UNavigationSystem*, const UObject*);
+	DECLARE_DELEGATE_RetVal_SixParams(float, FFindPathSignature, const FVector&, const FVector&, EPathFindingMode::Type&, const ANavigationData*, UNavigationSystem*, const UObject*);
+
+	bool TestPathFrom(const FVector& ItemPos, const FVector& ContextPos, EPathFindingMode::Type& Mode, const ANavigationData* NavData, UNavigationSystem* NavSys, const UObject* PathOwner);
+	bool TestPathTo(const FVector& ItemPos, const FVector& ContextPos, EPathFindingMode::Type& Mode, const ANavigationData* NavData, UNavigationSystem* NavSys, const UObject* PathOwner);
+	float FindPathCostFrom(const FVector& ItemPos, const FVector& ContextPos, EPathFindingMode::Type& Mode, const ANavigationData* NavData, UNavigationSystem* NavSys, const UObject* PathOwner);
+	float FindPathCostTo(const FVector& ItemPos, const FVector& ContextPos, EPathFindingMode::Type& Mode, const ANavigationData* NavData, UNavigationSystem* NavSys, const UObject* PathOwner);
+	float FindPathLengthFrom(const FVector& ItemPos, const FVector& ContextPos, EPathFindingMode::Type& Mode, const ANavigationData* NavData, UNavigationSystem* NavSys, const UObject* PathOwner);
+	float FindPathLengthTo(const FVector& ItemPos, const FVector& ContextPos, EPathFindingMode::Type& Mode, const ANavigationData* NavData, UNavigationSystem* NavSys, const UObject* PathOwner);
+
+	ANavigationData* FindNavigationData(UNavigationSystem* NavSys, UObject* Owner);
+};
diff --git a/Engine/Source/Runtime/AIModule/Classes/EnvironmentQuery/Tests/EnvQueryTest_Trace.h b/Engine/Source/Runtime/AIModule/Classes/EnvironmentQuery/Tests/EnvQueryTest_Trace.h
new file mode 100644
--- /dev/null
+++ b/Engine/Source/Runtime/AIModule/Classes/EnvironmentQuery/Tests/EnvQueryTest_Trace.h
@@ -0,0 +1,49 @@
+// Copyright 1998-2014 Epic Games, Inc. All Rights Reserved.
+
+#pragma once
+#include "EnvironmentQuery/EnvQueryTest.h"
+#include "EnvQueryTest_Trace.generated.h"
+
+UCLASS(MinimalAPI)
+class UEnvQueryTest_Trace : public UEnvQueryTest
+{
+	GENERATED_UCLASS_BODY()
+
+	/** trace data */
+	UPROPERTY(EditDefaultsOnly, Category=Trace)
+	FEnvTraceData TraceData;
+
+	/** trace direction */
+	UPROPERTY(EditDefaultsOnly, Category=Trace)
+	FEnvBoolParam TraceToItem;
+
+	/** Z offset from item */
+	UPROPERTY(EditDefaultsOnly, Category=Trace, AdvancedDisplay)
+	FEnvFloatParam ItemOffsetZ;
+
+	/** Z offset from querier */
+	UPROPERTY(EditDefaultsOnly, Category=Trace, AdvancedDisplay)
+	FEnvFloatParam ContextOffsetZ;
+
+	/** context: other end of trace test */
+	UPROPERTY(EditDefaultsOnly, Category=Trace)
+	TSubclassOf<class UEnvQueryContext> Context;
+
+	void RunTest(struct FEnvQueryInstance& QueryInstance);
+
+	virtual FString GetDescriptionTitle() const OVERRIDE;
+	virtual FText GetDescriptionDetails() const OVERRIDE;
+
+protected:
+
+	DECLARE_DELEGATE_RetVal_SevenParams(bool, FRunTraceSignature, const FVector&, const FVector&, AActor*, UWorld*, enum ECollisionChannel, const FCollisionQueryParams&, const FVector&);
+
+	bool RunLineTraceTo(const FVector& ItemPos, const FVector& ContextPos, AActor* ItemActor, UWorld* World, enum ECollisionChannel Channel, const FCollisionQueryParams& Params, const FVector& Extent);
+	bool RunLineTraceFrom(const FVector& ItemPos, const FVector& ContextPos, AActor* ItemActor, UWorld* World, enum ECollisionChannel Channel, const FCollisionQueryParams& Params, const FVector& Extent);
+	bool RunBoxTraceTo(const FVector& ItemPos, const FVector& ContextPos, AActor* ItemActor, UWorld* World, enum ECollisionChannel Channel, const FCollisionQueryParams& Params, const FVector& Extent);
+	bool RunBoxTraceFrom(const FVector& ItemPos, const FVector& ContextPos, AActor* ItemActor, UWorld* World, enum ECollisionChannel Channel, const FCollisionQueryParams& Params, const FVector& Extent);
+	bool RunSphereTraceTo(const FVector& ItemPos, const FVector& ContextPos, AActor* ItemActor, UWorld* World, enum ECollisionChannel Channel, const FCollisionQueryParams& Params, const FVector& Extent);
+	bool RunSphereTraceFrom(const FVector& ItemPos, const FVector& ContextPos, AActor* ItemActor, UWorld* World, enum ECollisionChannel Channel, const FCollisionQueryParams& Params, const FVector& Extent);
+	bool RunCapsuleTraceTo(const FVector& ItemPos, const FVector& ContextPos, AActor* ItemActor, UWorld* World, enum ECollisionChannel Channel, const FCollisionQueryParams& Params, const FVector& Extent);
+	bool RunCapsuleTraceFrom(const FVector& ItemPos, const FVector& ContextPos, AActor* ItemActor, UWorld* World, enum ECollisionChannel Channel, const FCollisionQueryParams& Params, const FVector& Extent);
+};
diff --git a/Engine/Source/Runtime/AIModule/Classes/Navigation/CrowdFollowingComponent.h b/Engine/Source/Runtime/AIModule/Classes/Navigation/CrowdFollowingComponent.h
new file mode 100644
--- /dev/null
+++ b/Engine/Source/Runtime/AIModule/Classes/Navigation/CrowdFollowingComponent.h
@@ -0,0 +1,145 @@
+// Copyright 1998-2014 Epic Games, Inc. All Rights Reserved.
+
+#pragma once
+#include "Navigation/CrowdAgentInterface.h"
+#include "Navigation/PathFollowingComponent.h"
+#include "CrowdFollowingComponent.generated.h"
+
+namespace ECrowdAvoidanceQuality
+{
+	enum Type
+	{
+		Low,
+		Medium,
+		Good,
+		High,
+	};
+}
+
+UCLASS()
+class AIMODULE_API UCrowdFollowingComponent : public UPathFollowingComponent, public ICrowdAgentInterface
+{
+	GENERATED_UCLASS_BODY()
+
+	UPROPERTY()
+	FVector CrowdAgentMoveDirection;
+
+	// ICrowdAgentInterface BEGIN
+	virtual FVector GetCrowdAgentLocation() const OVERRIDE;
+	virtual FVector GetCrowdAgentVelocity() const OVERRIDE;
+	virtual void GetCrowdAgentCollisions(float& CylinderRadius, float& CylinderHalfHeight) const OVERRIDE;
+	virtual float GetCrowdAgentMaxSpeed() const OVERRIDE;
+	// ICrowdAgentInterface END
+
+	// PathFollowingComponent BEGIN
+	virtual void Initialize() OVERRIDE;
+	virtual void Cleanup() OVERRIDE;
+	virtual void AbortMove(const FString& Reason, FAIRequestID RequestID = FAIRequestID::CurrentRequest, bool bResetVelocity = true, bool bSilent = false) OVERRIDE;
+	virtual void PauseMove(FAIRequestID RequestID = FAIRequestID::CurrentRequest, bool bResetVelocity = true) OVERRIDE;
+	virtual void ResumeMove(FAIRequestID RequestID = FAIRequestID::CurrentRequest) OVERRIDE;
+	virtual FVector GetMoveFocus(bool bAllowStrafe) const OVERRIDE;
+	virtual void OnLanded() OVERRIDE;
+	virtual void OnPathFinished(EPathFollowingResult::Type Result) OVERRIDE;
+	virtual void OnPathUpdated() OVERRIDE;
+	// PathFollowingComponent END
+
+	/** update params in crowd manager */
+	void UpdateCrowdAgentParams() const;
+
+	/** pass agent velocity to movement component */
+	virtual void ApplyCrowdAgentVelocity(const FVector& NewVelocity, const FVector& DestPathCorner, bool bTraversingLink);
+
+	/** master switch for crowd steering & avoidance */
+	void SuspendCrowdSteering(bool bSuspend);
+
+	/** switch between crowd simulation and parent implementation (following path segments) */
+	virtual void SetCrowdSimulation(bool bEnable);
+
+	/** called when agent moved to next nav node (poly) */
+	virtual void OnNavNodeChanged(NavNodeRef NewPolyRef, NavNodeRef PrevPolyRef, int32 CorridorSize);
+
+	void SetCrowdAnticipateTurns(bool bEnable, bool bUpdateAgent = true);
+	void SetCrowdObstacleAvoidance(bool bEnable, bool bUpdateAgent = true);
+	void SetCrowdSeparation(bool bEnable, bool bUpdateAgent = true);
+	void SetCrowdOptimizeVisibility(bool bEnable, bool bUpdateAgent = true);
+	void SetCrowdOptimizeTopology(bool bEnable, bool bUpdateAgent = true);
+	void SetCrowdSeparationWeight(float Weight, bool bUpdateAgent = true);
+	void SetCrowdCollisionQueryRange(float Range, bool bUpdateAgent = true);
+	void SetCrowdPathOptimizationRange(float Range, bool bUpdateAgent = true);
+	void SetCrowdAvoidanceQuality(ECrowdAvoidanceQuality::Type Quality, bool bUpdateAgent = true);
+
+	FORCEINLINE bool IsCrowdSimulationEnabled() const { return bEnableCrowdSimulation; }
+	FORCEINLINE bool IsCrowdAnticipateTurnsEnabled() const { return bEnableAnticipateTurns && !bSuspendCrowdSimulation; }
+	FORCEINLINE bool IsCrowdObstacleAvoidanceEnabled() const { return bEnableObstacleAvoidance && !bSuspendCrowdSimulation; }
+	FORCEINLINE bool IsCrowdSeparationEnabled() const { return bEnableSeparation && !bSuspendCrowdSimulation; }
+	FORCEINLINE bool IsCrowdOptimizeVisibilityEnabled() const { return bEnableOptimizeVisibility; /** don't check suspend here! */ }
+	FORCEINLINE bool IsCrowdOptimizeTopologyEnabled() const { return bEnableOptimizeTopology && !bSuspendCrowdSimulation; }
+	FORCEINLINE float GetCrowdSeparationWeight() const { return SeparationWeight; }
+	FORCEINLINE float GetCrowdCollisionQueryRange() const { return CollisionQueryRange; }
+	FORCEINLINE float GetCrowdPathOptimizationRange() const { return PathOptimizationRange; }
+	FORCEINLINE ECrowdAvoidanceQuality::Type GetCrowdAvoidanceQuality() const { return AvoidanceQuality; }
+
+	virtual void GetDebugStringTokens(TArray<FString>& Tokens, TArray<EPathFollowingDebugTokens::Type>& Flags) const;
+#if ENABLE_VISUAL_LOG
+	virtual void DescribeSelfToVisLog(struct FVisLogEntry* Snapshot) const;
+#endif // ENABLE_VISUAL_LOG
+
+protected:
+
+	UPROPERTY(transient)
+	class UCharacterMovementComponent* CharacterMovement;
+
+	/** if set, velocity will be updated even if agent is falling */
+	uint32 bAffectFallingVelocity : 1;
+
+	/** if set, move focus will match velocity direction */
+	uint32 bRotateToVelocity : 1;
+
+	/** if set, move velocity will be updated in every tick */
+	uint32 bUpdateDirectMoveVelocity : 1;
+
+	/** if set, agent will be simulated by crowd, otherwise it will act only as an obstacle */
+	uint32 bEnableCrowdSimulation : 1;
+
+	/** if set, avoidance and steering will be suspended (used for direct move requests) */
+	uint32 bSuspendCrowdSimulation : 1;
+
+	uint32 bEnableAnticipateTurns : 1;
+	uint32 bEnableObstacleAvoidance : 1;
+	uint32 bEnableSeparation : 1;
+	uint32 bEnableOptimizeVisibility : 1;
+	uint32 bEnableOptimizeTopology : 1;
+
+	/** if set, agent if moving on final path part, skip further updates (runtime flag) */
+	uint32 bFinalPathPart : 1;
+
+	/** if set, movement will be finished when velocity is opposite to path direction (runtime flag) */
+	uint32 bCheckMovementAngle : 1;
+
+	float SeparationWeight;
+	float CollisionQueryRange;
+	float PathOptimizationRange;
+
+	/** start index of current path part */
+	int32 PathStartIndex;
+
+	TEnumAsByte<ECrowdAvoidanceQuality::Type> AvoidanceQuality;
+
+	// PathFollowingComponent BEGIN
+	virtual int32 DetermineStartingPathPoint(const FNavigationPath* ConsideredPath) const OVERRIDE;
+	virtual void SetMoveSegment(int32 SegmentStartIndex) OVERRIDE;
+	virtual void UpdatePathSegment() OVERRIDE;
+	virtual void FollowPathSegment(float DeltaTime) OVERRIDE;
+	virtual bool ShouldCheckPathOnResume() const OVERRIDE;
+	virtual bool IsOnPath() const OVERRIDE;
+	virtual bool UpdateMovementComponent(bool bForce) OVERRIDE;
+	virtual void Reset() OVERRIDE;
+	// PathFollowingComponent END
+
+	void SwitchToNextPathPart();
+	bool ShouldSwitchPathPart(int32 CorridorSize) const;
+	bool HasMovedDuringPause() const;
+	void UpdateCachedDirections(const FVector& NewVelocity, const FVector& NextPathCorner, bool bTraversingLink);
+
+	friend class UCrowdManager;
+};
