Extracting AI code from the engine into a separate module, pass 1 #UE4

diff --git a/Engine/Source/Runtime/AIModule/Classes/Navigation/CrowdManager.h b/Engine/Source/Runtime/AIModule/Classes/Navigation/CrowdManager.h
new file mode 100644
--- /dev/null
+++ b/Engine/Source/Runtime/AIModule/Classes/Navigation/CrowdManager.h
@@ -0,0 +1,253 @@
+// Copyright 1998-2014 Epic Games, Inc. All Rights Reserved.
+
+#pragma once
+#include "Tickable.h"
+#include "CrowdManager.generated.h"
+
+/**
+ *  Crowd manager is responsible for handling crowds using Detour (Recast library)
+ *
+ *  Agents will respect navmesh for all steering and avoidance updates, 
+ *  but it's slower than AvoidanceManager solution (RVO, cares only about agents)
+ *
+ *  All agents will operate on the same navmesh data, which will be picked from
+ *  navigation system defaults (UNavigationSystem::SupportedAgents[0])
+ *
+ *  To use it, you have to add CrowdFollowingComponent to your agent
+ *  (usually: replace class of PathFollowingComponent in AIController by adding 
+ *   those lines in controller's constructor
+ *
+ *   ACrowdAIController::ACrowdAIController(const class FPostConstructInitializeProperties& PCIP)
+ *       : Super(PCIP.SetDefaultSubobjectClass<UCrowdFollowingComponent>(TEXT("PathFollowingComponent")))
+ *
+ *   or simply add both components and switch move requests between them)
+ *
+ *  Actors that should be avoided, but are not being simulated by crowd (like players)
+ *  should implement CrowdAgentInterface AND register/unregister themselves with crowd manager:
+ *  
+ *   UCrowdManager* CrowdManager = UCrowdManager::GetCurrent(this);
+ *   if (CrowdManager)
+ *   {
+ *      CrowdManager->RegisterAgent(this);
+ *   }
+ *
+ *   Check flags in CrowdDebugDrawing namespace (CrowdManager.cpp) for debugging options.
+ */
+
+USTRUCT()
+struct AIMODULE_API FCrowdAvoidanceConfig
+{
+	GENERATED_USTRUCT_BODY()
+
+	UPROPERTY(EditAnywhere, Category=Crowd)
+	float VelocityBias;
+
+	UPROPERTY(EditAnywhere, Category=Crowd)
+	float DesiredVelocityWeight;
+	
+	UPROPERTY(EditAnywhere, Category=Crowd)
+	float CurrentVelocityWeight;
+	
+	UPROPERTY(EditAnywhere, Category=Crowd)
+	float SideBiasWeight;
+	
+	UPROPERTY(EditAnywhere, Category=Crowd)
+	float ImpactTimeWeight;
+
+	UPROPERTY(EditAnywhere, Category=Crowd)
+	float ImpactTimeRange;
+
+	UPROPERTY(EditAnywhere, Category=Crowd)
+	uint8 AdaptiveDivisions;
+
+	UPROPERTY(EditAnywhere, Category=Crowd)
+	uint8 AdaptiveRings;
+	
+	UPROPERTY(EditAnywhere, Category=Crowd)
+	uint8 AdaptiveDepth;
+
+	FCrowdAvoidanceConfig() :
+		VelocityBias(0.4f), DesiredVelocityWeight(2.0f), CurrentVelocityWeight(0.75f),
+		SideBiasWeight(0.75f), ImpactTimeWeight(2.5f), ImpactTimeRange(2.5f),
+		AdaptiveDivisions(7), AdaptiveRings(2), AdaptiveDepth(5)
+	{}
+};
+
+struct AIMODULE_API FCrowdAgentData
+{
+#if WITH_RECAST
+	/** special filter for checking offmesh links */
+	struct dtQuerySpecialLinkFilter* LinkFilter;
+#endif
+
+	/** poly ref that agent is standing on from previous update */
+	NavNodeRef PrevPoly;
+
+	/** index of agent in detour crowd */
+	int32 AgentIndex;
+
+	/** is this agent fully simulated by crowd? */
+	uint32 bIsSimulated : 1;
+
+	FCrowdAgentData() : PrevPoly(0), AgentIndex(-1), bIsSimulated(false) {}
+};
+
+struct FCrowdTickHelper : FTickableGameObject
+{
+	TWeakObjectPtr<class UCrowdManager> Owner;
+
+	FCrowdTickHelper() : Owner(NULL) {}
+	virtual void Tick(float DeltaTime);
+	virtual bool IsTickable() const { return Owner.IsValid(); }
+	virtual bool IsTickableInEditor() const { return true; }
+	virtual bool IsTickableWhenPaused() const { return true; }
+	virtual TStatId GetStatId() const;
+};
+
+UCLASS()
+class AIMODULE_API UCrowdManager : public UObject
+{
+	GENERATED_UCLASS_BODY()
+
+	virtual void Tick(float DeltaTime);
+	virtual void BeginDestroy() OVERRIDE;
+
+	/** adds new agent to crowd */
+	bool RegisterAgent(const class ICrowdAgentInterface* Agent);
+
+	/** removes agent from crowd */
+	void UnregisterAgent(const class ICrowdAgentInterface* Agent);
+
+	/** updates agent data */
+	void UpdateAgentParams(const class ICrowdAgentInterface* Agent) const;
+
+	/** refresh agent state */
+	void UpdateAgentState(const class ICrowdAgentInterface* Agent) const;
+
+	/** sets move target for crowd agent (only for fully simulated) */
+	bool SetAgentMoveTarget(const class UCrowdFollowingComponent* AgentComponent, const FVector& MoveTarget, TSharedPtr<const FNavigationQueryFilter> Filter) const;
+
+	/** sets move direction for crowd agent (only for fully simulated) */
+	bool SetAgentMoveDirection(const class UCrowdFollowingComponent* AgentComponent, const FVector& MoveDirection) const;
+
+	/** sets move target using path (only for fully simulated) */
+	bool SetAgentMovePath(const class UCrowdFollowingComponent* AgentComponent, const FNavMeshPath* Path, int32 PathSectionStart, int32 PathSectionEnd) const;
+
+	/** clears move target for crowd agent (only for fully simulated) */
+	void ClearAgentMoveTarget(const class UCrowdFollowingComponent* AgentComponent) const;
+
+	/** switch agent to waiting state */
+	void PauseAgent(const class UCrowdFollowingComponent* AgentComponent) const;
+
+	/** resumes agent movement */
+	void ResumeAgent(const class UCrowdFollowingComponent* AgentComponent, bool bForceReplanPath = true) const;
+
+	/** returns number of nearby agents */
+	int32 GetNumNearbyAgents(const class ICrowdAgentInterface* Agent) const;
+
+	/** reads existing avoidance config or returns false */
+	bool GetAvoidanceConfig(int32 Idx, FCrowdAvoidanceConfig& Data) const;
+
+	/** updates existing avoidance config or returns false */
+	bool SetAvoidanceConfig(int32 Idx, const FCrowdAvoidanceConfig& Data);
+
+	/** remove started offmesh connections from corridor */
+	void SetOffmeshConnectionPruning(bool bRemoveFromCorridor);
+
+	/** block path visibility raycasts when crossing different nav areas */
+	void SetSingleAreaVisibilityOptimization(bool bEnable);
+
+	/** adjust current position in path's corridor, starting test from PathStartIdx */
+	void AdjustAgentPathStart(const class UCrowdFollowingComponent* AgentComponent, const FNavMeshPath* Path, int32& PathStartIdx) const;
+
+#if WITH_EDITOR
+	virtual void PostEditChangeProperty(struct FPropertyChangedEvent& PropertyChangedEvent) OVERRIDE;
+
+	void DebugTick() const;
+#endif
+
+	/** notify called when detour navmesh is changed */
+	void OnNavMeshUpdate();
+
+	class UWorld* GetWorld() const;
+
+	static UCrowdManager* GetCurrent(class UObject* WorldContextObject);
+	static UCrowdManager* GetCurrent(class UWorld* World);
+
+protected:
+
+	UPROPERTY(transient)
+	class ANavigationData* MyNavData;
+
+	/** obstacle avoidance params */
+	UPROPERTY(EditAnywhere, Category=Config)
+	TArray<FCrowdAvoidanceConfig> AvoidanceConfig;
+
+	/** max number of agents supported by crowd */
+	UPROPERTY(EditAnywhere, Category=Config)
+	int32 MaxAgents;
+
+	/** max radius of agent that can be added to crowd */
+	UPROPERTY(EditAnywhere, Category=Config)
+	float MaxAgentRadius;
+
+	/** how often should agents check their position after moving off navmesh? */
+	UPROPERTY(EditAnywhere, Category=Config)
+	float NavmeshCheckInterval;
+
+	uint32 bPruneStartedOffmeshConnections : 1;
+	uint32 bSingleAreaVisibilityOptimization : 1;
+
+	/** agents registered in crowd manager */
+	TMap<const class ICrowdAgentInterface*, FCrowdAgentData> ActiveAgents;
+
+#if WITH_RECAST
+	/** crowd manager */
+	class dtCrowd* DetourCrowd;
+
+	/** debug data */
+	struct dtCrowdAgentDebugInfo* DetourAgentDebug;
+	class dtObstacleAvoidanceDebugData* DetourAvoidanceDebug;
+#endif
+
+#if WITH_EDITOR
+	FCrowdTickHelper* TickHelper;
+#endif
+
+	/** try to initialize nav data from already existing ones */
+	virtual void UpdateNavData();
+
+	/** setup params of crowd avoidance */
+	virtual void UpdateAvoidanceConfig();
+
+	/** called from tick, just after updating agents proximity data */
+	virtual void PostProximityUpdate();
+
+#if WITH_RECAST
+	void AddAgent(const class ICrowdAgentInterface* Agent, FCrowdAgentData& AgentData) const;
+	void RemoveAgent(const class ICrowdAgentInterface* Agent, const FCrowdAgentData* AgentData) const;
+	void GetAgentParams(const class ICrowdAgentInterface* Agent, struct dtCrowdAgentParams* AgentParams) const;
+
+	/** pass position and velocity to crowd simulation */
+	void UpdateSpatialData(const class ICrowdAgentInterface* Agent, const FCrowdAgentData& AgentData) const;
+
+	/** pass new velocity to movement components */
+	void ApplyVelocity(const class UCrowdFollowingComponent* AgentComponent, int32 AgentIndex) const;
+
+	/** check changes in crowd simulation and adjust UE specific properties (smart links, poly updates) */
+	void UpdateAgentPaths();
+
+	/** switch debugger to object selected in PIE */
+	void UpdateSelectedDebug(const class ICrowdAgentInterface* Agent, int32 AgentIndex) const;
+
+	void CreateCrowdManager();
+	void DestroyCrowdManager();
+
+	void DrawDebugCorners(const struct dtCrowdAgent* CrowdAgent) const;
+	void DrawDebugCollisionSegments(const struct dtCrowdAgent* CrowdAgent) const;
+	void DrawDebugPath(const struct dtCrowdAgent* CrowdAgent) const;
+	void DrawDebugVelocityObstacles(const struct dtCrowdAgent* CrowdAgent) const;
+	void DrawDebugPathOptimization(const struct dtCrowdAgent* CrowdAgent) const;
+	void DrawDebugNeighbors(const struct dtCrowdAgent* CrowdAgent) const;
+#endif
+};
diff --git a/Engine/Source/Runtime/AIModule/Classes/Navigation/NavigationComponent.h b/Engine/Source/Runtime/AIModule/Classes/Navigation/NavigationComponent.h
new file mode 100644
--- /dev/null
+++ b/Engine/Source/Runtime/AIModule/Classes/Navigation/NavigationComponent.h
@@ -0,0 +1,314 @@
+// Copyright 1998-2014 Epic Games, Inc. All Rights Reserved.
+
+#pragma once
+#include "AI/Navigation/NavigationPathGenerator.h"
+#include "NavigationComponent.generated.h"
+
+UCLASS(dependson=(UNavAgentInterface, UNavPathObserverInterface, UNavigationSystem), config=Engine)
+class AIMODULE_API UNavigationComponent : public UActorComponent, public INavigationPathGenerator
+{
+	GENERATED_UCLASS_BODY()
+
+	static const bool QuerySuccess = true;
+	static const bool QueryFailure = false;
+
+	/** Actor this instance is observing. Can be NULL */
+	UPROPERTY()
+	const class AActor* GoalActor;
+
+	// Begin UActorComponent Interface
+	virtual void InitializeComponent() OVERRIDE;
+	virtual void TickComponent(float DeltaTime, enum ELevelTick TickType, FActorComponentTickFunction *ThisTickFunction) OVERRIDE;
+	virtual void SetComponentTickEnabledAsync(bool bEnabled) OVERRIDE;
+	// End UActorComponent Interface
+
+	// Begin INavigationPathGenerator Interface
+	virtual FNavPathSharedPtr GetGeneratedPath(class INavAgentInterface* Agent) OVERRIDE { return Path; }
+	// End INavigationPathGenerator Interface
+
+	/** updates cached pointers to relevant owner's components */
+	virtual void UpdateCachedComponents();
+
+	/** called to notify navigation component to update various cached data related to navigation agent it's driving */
+	virtual void OnNavAgentChanged();
+
+	//----------------------------------------------------------------------//
+	// querying                                                                
+	//----------------------------------------------------------------------//
+	
+	/** Retrieves ANavigationData instance being a context to this NavigationComponent. */
+	FORCEINLINE const class ANavigationData* GetNavData() const { return MyNavData ? MyNavData : PickNavData(); }
+	
+	/** Retrieves ANavigationData instance being a context to this NavigationComponent. */
+	FORCEINLINE class ANavigationData* GetNavData() { return MyNavData ? MyNavData : PickNavData(); }
+
+	/** Returns a random location on the navmesh. */
+	FNavLocation GetRandomPointOnNavMesh() const;
+
+	/** Projects given Location to navigation referenced by this NavigationComponent,
+	 *	using any additional internal information available. */
+	FNavLocation ProjectPointToNavigation(const FVector& Location) const;
+
+	/** Projects given Location to navigation referenced by this NavigationComponent,
+	 *	using any additional internal information available. */
+	bool ProjectPointToNavigation(const FVector& Location, FNavLocation& ProjectedLocation) const;
+	bool ProjectPointToNavigation(const FVector& Location, FVector& ProjectedLocation) const;
+	
+	/** Requests generating simple path in straight line to an Actor,
+	 *  and continuous update when goal actor moves more than RepathDistance from path end point.
+	 *	It doesn't check bDoAsyncPathfinding flag, and path will be always available after function call */
+	virtual bool FindSimplePathToActor(const AActor* NewGoalActor, float RepathDistance = 100.0f);
+
+	/** Requests generating path to an Actor, and continuous update of path 
+	 *	if GoalActor moves more than RepathDistance from path end point.
+	 *	If bDoAsyncPathfinding == true then MyNavAgent->OnPathUpdated will be 
+	 *	called when path finishes generating */
+	virtual bool FindPathToActor(const AActor* NewGoalActor, TSharedPtr<const FNavigationQueryFilter> QueryFilter = NULL, float RepathDistance = 100.f);
+
+	/** requests generating simple path in straight line to specified DestLocation. Aborts previous pathfinding requests. 
+	 *	It doesn't check bDoAsyncPathfinding flag, and path will be always available after function call */
+	virtual bool FindSimplePathToLocation(const FVector& DestLocation);
+
+	/** requests generation of a path to a specified pawn. Aborts previous pathfinding requests. 
+	 *	If bDoAsyncPathfinding == true then MyNavAgent->OnPathUpdated will be 
+	 *	called when path finishes generating */
+	virtual bool FindPathToLocation(const FVector& DestLocation, TSharedPtr<const FNavigationQueryFilter> QueryFilter = NULL);
+
+	/**
+	 *	Freezes NavComponent's actions like tracing goal actor. If following 
+	 *	pathfinding request will have the same destination then stored
+	 *	path will be reused (if still valid)
+	 */
+	void Pause();
+
+	/** Notifies navigation system given pathfinding request is no longer valid.
+	 *	Even if path ends up being generated (due to the async nature of its
+	 *	generation) result will be discarded. 
+	 *	@NOTE Function just aborts currently requested path. If NavigationComponent 
+	 *	instance observes and actor it will generate a new one. */
+	void AbortAsyncFindPathRequest();
+
+	virtual void SetPath(const FNavPathSharedPtr& NewPath);
+	
+	/** @todo: document */
+	bool GetPathPoint(uint32 PathVertIdx, FNavPathPoint& PathPoint) const;
+
+	FORCEINLINE const FNavPathSharedPtr& GetPath() const { return Path; }
+
+	/** @todo: document */
+	bool HasValidPath() const;
+
+	FORCEINLINE bool IsPathPartial() const { return bIsPathPartial != 0; }
+
+	/** @todo: document */
+	FORCEINLINE bool IsFollowing() const { return GoalActor != NULL; }
+
+	FORCEINLINE bool DoesAsyncPathfinding() const { return bDoAsyncPathfinding; }
+
+	FORCEINLINE bool IsPaused() const { return bIsPaused; }
+
+	FORCEINLINE bool IsWaitingForRepath() const { return bIsWaitingForRepath; }
+
+	/** Return true if path should be updated when nearby smart link change its state */
+	FORCEINLINE bool WantsSmartLinkUpdates() const { return bUpdateForSmartLinks; }
+
+	/** Returns current move goal for actor, using data from GetMoveGoalOffset() */
+	FORCEINLINE FVector GetCurrentMoveGoal(const class AActor* InGoalActor, class AActor* MovingActor) const
+	{
+		const INavAgentInterface* NavAgent = InterfaceCast<const INavAgentInterface>(InGoalActor);
+		if (NavAgent)
+		{
+			const FVector Offset = NavAgent->GetMoveGoalOffset(MovingActor);
+			const FVector MoveGoal = FRotationTranslationMatrix(InGoalActor->GetActorRotation(), NavAgent->GetNavAgentLocation()).TransformPosition(Offset);
+			return MoveGoal;
+		}
+		else
+		{
+			return InGoalActor->GetActorLocation();
+		}
+	}
+
+	FORCEINLINE FVector GetQueryExtent() const { return NavigationQueryExtent; }
+
+	void SwapCurrentMoveGoal(const AActor* NewGoalActor);
+
+	//----------------------------------------------------------------------//
+	// misc
+	//----------------------------------------------------------------------//
+
+	/** Aborts actor following process triggered with FindPathToActor */
+	void ResetTransientData();
+
+	/** Notify path following about changes to update current movement */
+	void SetPathFollowingUpdates(bool bEnabled);
+
+	/** Find new path when it gets invalidated (e.g. runtime navmesh rebuild) */
+	void SetRepathWhenInvalid(bool bEnabled);
+
+	UFUNCTION()
+	virtual void OnNavDataRegistered(class ANavigationData* NavData);
+
+	/** 
+	 *	Called when given path becomes invalid (via @see PathObserverDelegate)
+	 *	NOTE: InvalidatedPath doesn't have to be instance's current Path
+	 */
+	void OnPathInvalid(FNavigationPath* InvalidatedPath);
+
+	/** Called when nearby smart link changes its state */
+	virtual void OnSmartLinkBroadcast(class USmartNavLinkComponent* NearbyLink);
+
+	/** set whether this component should receive update broadcasts from nearby smart links or not */
+	void SetReceiveSmartLinkUpdates(bool bEnabled);
+
+	TSharedPtr<const FNavigationQueryFilter> GetStoredQueryFilter() { return StoredQueryFilter; }
+
+	/** get additional flags passed with path finding queries */
+	int32 GetNavDataFlags() const { return NavDataFlags;  }
+
+	/** add new NavData flag passed with path finding queries */
+	void SetNavDataFlag(int32 Flag) { NavDataFlags |= Flag; }
+
+	/** add new NavData flag passed with path finding queries */
+	void ClearNavDataFlag(int32 Flag) { NavDataFlags &= ~Flag; }
+
+	/** override NavData flags */
+	void SetNavDataFlags(int32 Flags) { NavDataFlags = Flags; }
+
+	//----------------------------------------------------------------------//
+	// debug
+	//----------------------------------------------------------------------//
+
+#if ENABLE_VISUAL_LOG
+	virtual void DescribeSelfToVisLog(struct FVisLogEntry* Snapshot) const;
+#endif // ENABLE_VISUAL_LOG
+
+	/** quick and dirty debug render. If Canvas not NULL additional data (if present) will be drawn */
+	void DrawDebugCurrentPath(class UCanvas* Canvas = NULL, FColor PathColor = FColor::White, const uint32 NextPathPointIndex = 0);
+
+
+protected:
+
+	/** associated path following component */
+	UPROPERTY(transient)
+	class UPathFollowingComponent* PathFollowComp;
+
+	/** default query extent to use. It's a sum of navigation data's default query extent and nav agent's size */
+	mutable FVector NavigationQueryExtent;
+
+	/** Shared path - there can be multiple users of one path instance. No not modify it, unless you know what you're doing */
+	FNavPathSharedPtr Path;
+
+	/** caches information whether contained path is partial (i.e. destination was not reachable so path leads to best guess location */
+	uint32 bIsPathPartial : 1;
+
+	/** whether component wants to be ticked */
+	uint32 bShoudTick : 1;
+
+	/** NavigationComponent can be "Paused" to not update it's path or observe it in any way
+	 *	but not release generated path neither. Subsequent calls to FindPathToActor or 
+	 *	FindPathToLocation will result in checking if new destination is the same as
+	 *	the old one, and if so internal ResumePath will be called 
+	 *	@NOTE: contained path invalidation while component is paused will result 
+	 *		in clearing pause along with other transient data */
+	uint32 bIsPaused : 1;
+
+	/** repath to goal actor using simple paths (without navigation data) */
+	uint32 bUseSimplePath : 1;
+
+	/** if set, path following component will be notified about path changes */
+	uint32 bNotifyPathFollowing : 1;
+
+	/** if set, invalid paths will trigger repath */
+	uint32 bRepathWhenInvalid : 1;
+
+	/** set when component if waiting for scheduled repath */
+	uint32 bIsWaitingForRepath : 1;
+
+	/** if set, path will be updated if agent receives smart link broadcast */
+	uint32 bUpdateForSmartLinks : 1;
+
+	UPROPERTY(config)
+	uint32 bDoAsyncPathfinding : 1;
+
+	/** use HPA* for pathfind requests (faster, but less accurate) */
+	UPROPERTY(config)
+	uint32 bDoHierarchicalPathfinding : 1;
+
+	/** if A* search fails on running out of nodes, try searching again but switch start with goal
+	 *  (uses backtracking feature to handle one way offmesh links properly) */
+	UPROPERTY(config)
+	uint32 bSearchFromGoalWhenOutOfNodes : 1;
+
+	mutable INavAgentInterface* MyNavAgent;
+	mutable INavPathObserverInterface* MyPathObserver;
+
+	FVector CurrentGoal;
+
+	TSharedPtr<const FNavigationQueryFilter> StoredQueryFilter;
+
+	/** contains location of GoalActor from the moment path has been generated. Used for 
+	 *	generated path invalidation when if goal actor moved away from original location */
+	FVector OriginalGoalActorLocation;
+
+	int32 AsynPathQueryID;
+
+	/** additional flags passed with every path finding query */
+	int32 NavDataFlags;
+	
+	/** if GoalActor moves move than this away from last point of Path Navigation Component will re-run path finding */
+	float RepathDistanceSq;
+
+	FNavigationPath::FPathObserverDelegate PathObserverDelegate;
+
+	struct FDeferredRepath
+	{
+		TWeakObjectPtr<AActor> GoalActor;
+		FVector GoalLocation;
+		bool bSimplePath;
+	};
+	FDeferredRepath RepathData;
+
+	/** asks NavigationSystem for NavigationData appropriate for self and caches it in MyNavData */
+	class ANavigationData* PickNavData() const;
+
+	/** called as part of PickNavData to update NavigationQueryExtent */
+	virtual void CacheNavQueryExtent() const;
+
+	/** Generates simple path in straight line without using navigation data. */
+	bool GenerateSimplePath(const FVector& GoalLocation);
+
+	/** Generates a path between component owner's location and GoalLocation and stores 
+	 *	it locally in PathPolys.  Path is available upon return from this function. */
+	bool GeneratePathTo(const FVector& GoalLocation, TSharedPtr<const FNavigationQueryFilter> QueryFilter = NULL);
+
+	/** Generates a path between the given points and stores it locally in PathPolys.  
+	 *	Path is available upon return from this function. */
+	bool GeneratePath(const FVector& FromLocation, const FVector& GoalLocation, TSharedPtr<const FNavigationQueryFilter> QueryFilter = NULL);
+
+	/** Asynchronous version of GeneratePathTo().  Path isn't available until callback is received. */
+	bool AsyncGeneratePathTo(const FVector& GoalLocation, TSharedPtr<const FNavigationQueryFilter> QueryFilter = NULL);
+
+	/** Asynchronous version of GeneratePath().  Path isn't available until callback is received. */
+	bool AsyncGeneratePath(const FVector& FromLocation, const FVector& GoalLocation, TSharedPtr<const FNavigationQueryFilter> QueryFilter = NULL);
+
+	bool ResumePath();
+
+	void DeferredResumePath();
+	
+	bool RePathTo(const FVector& GoalLocation, TSharedPtr<const FNavigationQueryFilter> StoredQueryFilter = NULL);
+
+	/** send notifies about updated path */
+	void NotifyPathUpdate();
+
+	/** find new path to GoalActor or CurrentGoal  */
+	void DeferredRepathToGoal();
+
+private:
+
+	/** cached pointer to ANavMesh this navigation component's using. Use Reset to force recaching. 
+	 *	@note read it via GetNavData that will update it if null. */
+	UPROPERTY()
+	mutable ANavigationData* MyNavData;
+
+	void AsyncGeneratePath_OnCompleteCallback(uint32 QueryID, ENavigationQueryResult::Type Result, FNavPathSharedPtr ResultPath);	
+};
diff --git a/Engine/Source/Runtime/AIModule/Classes/Navigation/PathFollowingComponent.h b/Engine/Source/Runtime/AIModule/Classes/Navigation/PathFollowingComponent.h
new file mode 100644
--- /dev/null
+++ b/Engine/Source/Runtime/AIModule/Classes/Navigation/PathFollowingComponent.h
@@ -0,0 +1,447 @@
+// Copyright 1998-2014 Epic Games, Inc. All Rights Reserved.
+
+#pragma once
+#include "AITypes.h"
+#include "AI/Navigation/NavigationTypes.h"
+#include "Components/ActorComponent.h"
+#include "AIResourceInterface.h"
+#include "PathFollowingComponent.generated.h"
+
+DECLARE_LOG_CATEGORY_EXTERN(LogPathFollowing, Warning, All);
+
+UENUM(BlueprintType)
+namespace EPathFollowingStatus
+{
+	enum Type
+	{
+		Idle,				// no requests
+		Waiting,			// request with incomplete path, will start after UpdateMove()
+		Paused,				// request paused, will continue after ResumeMove()
+		Moving,				// following path
+	};
+}
+
+UENUM(BlueprintType)
+namespace EPathFollowingResult
+{
+	enum Type
+	{
+		Success,			// reached destination
+		Blocked,			// movement was blocked
+		OffPath,			// agent is not on path
+		Aborted,			// aborted and stopped (failure)
+		Skipped,			// aborted and replaced with new request
+		Invalid,			// request was invalid
+	};
+}
+
+// left for now, will be removed soon! please use EPathFollowingStatus instead
+UENUM(BlueprintType)
+namespace EPathFollowingAction
+{
+	enum Type
+	{
+		Error,
+		NoMove,
+		DirectMove,
+		PartialPath,
+		PathToGoal,
+	};
+}
+
+UENUM(BlueprintType)
+namespace EPathFollowingRequestResult
+{
+	enum Type
+	{
+		Failed,
+		AlreadyAtGoal,
+		RequestSuccessful
+	};
+}
+
+namespace EPathFollowingDebugTokens
+{
+	enum Type
+	{
+		Description,
+		ParamName,
+		FailedValue,
+		PassedValue,
+	};
+}
+
+UCLASS(config=Engine)
+class AIMODULE_API UPathFollowingComponent : public UActorComponent, public IAIResourceInterface
+{
+	GENERATED_UCLASS_BODY()
+
+	DECLARE_DELEGATE_TwoParams(FPostProcessMoveSignature, UPathFollowingComponent* /*comp*/, FVector& /*velocity*/);
+	DECLARE_DELEGATE_OneParam(FRequestCompletedSignature, EPathFollowingResult::Type /*Result*/);
+	DECLARE_MULTICAST_DELEGATE_TwoParams(FMoveCompletedSignature, FAIRequestID /*RequestID*/, EPathFollowingResult::Type /*Result*/);
+
+	/** delegate for modifying path following velocity */
+	FPostProcessMoveSignature PostProcessMove;
+
+	/** delegate for move completion notify */
+	FMoveCompletedSignature OnMoveFinished;
+
+	// Begin UActorComponent Interface
+	virtual void TickComponent(float DeltaTime, enum ELevelTick TickType, FActorComponentTickFunction *ThisTickFunction) OVERRIDE;
+	// End UActorComponent Interface
+
+	/** initialize component to use */
+	virtual void Initialize();
+
+	/** cleanup component before destroying */
+	virtual void Cleanup();
+
+	/** updates cached pointers to relevant owner's components */
+	virtual void UpdateCachedComponents();
+
+	/** start movement along path
+	 *  @returns request ID or 0 when failed */
+	virtual FAIRequestID RequestMove(FNavPathSharedPtr Path, FRequestCompletedSignature OnComplete, const AActor* DestinationActor = NULL, float AcceptanceRadius = UPathFollowingComponent::DefaultAcceptanceRadius, bool bStopOnOverlap = true, FCustomMoveSharedPtr GameData = NULL);
+
+	/** start movement along path
+	 *  @returns request ID or 0 when failed */
+	FORCEINLINE FAIRequestID RequestMove(FNavPathSharedPtr InPath, const AActor* InDestinationActor = NULL, float InAcceptanceRadius = UPathFollowingComponent::DefaultAcceptanceRadius, bool InStopOnOverlap = true, FCustomMoveSharedPtr InGameData = NULL)
+	{
+		return RequestMove(InPath, UnboundRequestDelegate, InDestinationActor, InAcceptanceRadius, InStopOnOverlap, InGameData);
+	}
+
+	/** update path for specified request
+	 *  @param RequestID - request to update */
+	virtual bool UpdateMove(FNavPathSharedPtr Path, FAIRequestID RequestID = FAIRequestID::CurrentRequest);
+
+	/** aborts following path
+	 *  @param RequestID - request to abort, 0 = current
+	 *  @param bResetVelocity - try to stop movement component
+	 *  @param bSilent - finish with Skipped result instead of Aborted */
+	virtual void AbortMove(const FString& Reason, FAIRequestID RequestID = FAIRequestID::CurrentRequest, bool bResetVelocity = true, bool bSilent = false);
+
+	/** pause path following
+	*  @param RequestID - request to pause, FAIRequestID::CurrentRequest means pause current request, regardless of its ID */
+	virtual void PauseMove(FAIRequestID RequestID = FAIRequestID::CurrentRequest, bool bResetVelocity = true);
+
+	/** resume path following
+	*  @param RequestID - request to resume, FAIRequestID::CurrentRequest means restor current request, regardless of its ID*/
+	virtual void ResumeMove(FAIRequestID RequestID = FAIRequestID::CurrentRequest);
+
+	/** notify about finished movement */
+	virtual void OnPathFinished(EPathFollowingResult::Type Result);
+
+	/** notify about finishing move along current path segment */
+	virtual void OnSegmentFinished();
+
+	/** notify about changing current path */
+	virtual void OnPathUpdated();
+
+	/** set associated movement component */
+	virtual void SetMovementComponent(class UNavMovementComponent* MoveComp);
+
+	/** get current focal point of movement */
+	virtual FVector GetMoveFocus(bool bAllowStrafe) const;
+
+	/** simple test for stationary agent (used as early finish condition), check if reached given point
+	 *  @param TestPoint - point to test
+	 *  @param AcceptanceRadius - allowed 2D distance
+	 *  @param bExactSpot - false: increase AcceptanceRadius with agent's radius
+	 */
+	bool HasReached(const FVector& TestPoint, float AcceptanceRadius = UPathFollowingComponent::DefaultAcceptanceRadius, bool bExactSpot = false) const;
+
+	/** simple test for stationary agent (used as early finish condition), check if reached given goal
+	 *  @param TestGoal - actor to test
+	 *  @param AcceptanceRadius - allowed 2D distance
+	 *  @param bExactSpot - false: increase AcceptanceRadius with agent's radius
+	 */
+	bool HasReached(const AActor* TestGoal, float AcceptanceRadius = UPathFollowingComponent::DefaultAcceptanceRadius, bool bExactSpot = false) const;
+
+	/** update state of block detection */
+	void SetBlockDetectionState(bool bEnable);
+
+	/** @returns state of block detection */
+	bool IsBlockDetectionActive() const { return bUseBlockDetection; }
+
+	/** set block detection params */
+	void SetBlockDetection(float DistanceThreshold, float Interval, int32 NumSamples);
+
+	/** set threshold for precise reach tests in intermediate goals (minimal test radius)  */
+	void SetPreciseReachThreshold(float AgentRadiusMultiplier, float AgentHeightMultiplier);
+
+	/** set status of last requested move */
+	void SetLastMoveAtGoal(bool bFinishedAtGoal);
+
+	/** @returns estimated cost of unprocessed path segments
+	 *	@NOTE 0 means, that component is following final path segment or doesn't move */
+	float GetRemainingPathCost() const;
+	
+	/** Returns current location on navigation data */
+	FNavLocation GetCurrentNavLocation() const;
+
+	FORCEINLINE EPathFollowingStatus::Type GetStatus() const { return Status; }
+	FORCEINLINE float GetAcceptanceRadius() const { return AcceptanceRadius; }
+	FORCEINLINE float GetDefaultAcceptanceRadius() const { return MyDefaultAcceptanceRadius; }
+	FORCEINLINE void SetAcceptanceRadius(float InAcceptanceRadius) { AcceptanceRadius = InAcceptanceRadius; }
+	FORCEINLINE AActor* GetMoveGoal() const { return DestinationActor.Get(); }
+	FORCEINLINE bool HasPartialPath() const { return Path.IsValid() && Path->IsPartial(); }
+	FORCEINLINE bool DidMoveReachGoal() const { return bLastMoveReachedGoal && (Status == EPathFollowingStatus::Idle); }
+
+	FORCEINLINE FAIRequestID GetCurrentRequestId() const { return CurrentRequestId; }
+	FORCEINLINE uint32 GetCurrentPathIndex() const { return MoveSegmentStartIndex; }
+	FORCEINLINE uint32 GetNextPathIndex() const { return MoveSegmentEndIndex; }
+	FORCEINLINE FVector GetCurrentTargetLocation() const { return *CurrentDestination; }
+	FORCEINLINE FVector GetCurrentDirection() const { return MoveSegmentDirection; }
+	FORCEINLINE FBasedPosition GetCurrentTargetLocationBased() const { return CurrentDestination; }
+	FORCEINLINE USmartNavLinkComponent* GetCurrentSmartLink() const { return CurrentSmartLink;  }
+
+	/** will be deprecated soon, please use AIController.GetMoveStatus instead! */
+	UFUNCTION(BlueprintCallable, Category="AI|Components|PathFollowing")
+	EPathFollowingAction::Type GetPathActionType() const;
+
+	/** will be deprecated soon, please use AIController.GetImmediateMoveDestination instead! */
+	UFUNCTION(BlueprintCallable, Category="AI|Components|PathFollowing")
+	FVector GetPathDestination() const;
+
+	FORCEINLINE const FNavPathSharedPtr GetPath() const { return Path; }
+
+	/** readable name of current status */
+	FString GetStatusDesc() const;
+	/** readable name of result enum */
+	FString GetResultDesc(EPathFollowingResult::Type Result) const;
+
+	void SetDestinationActor(const AActor* InDestinationActor);
+
+	virtual void GetDebugStringTokens(TArray<FString>& Tokens, TArray<EPathFollowingDebugTokens::Type>& Flags) const;
+	virtual FString GetDebugString() const;
+
+	virtual void DisplayDebug(class UCanvas* Canvas, const FDebugDisplayInfo& DebugDisplay, float& YL, float& YPos) const;
+#if ENABLE_VISUAL_LOG
+	virtual void DescribeSelfToVisLog(struct FVisLogEntry* Snapshot) const;
+#endif // ENABLE_VISUAL_LOG
+
+	/** called when moving agent collides with another actor */
+	UFUNCTION()
+	virtual void OnActorBump(class AActor* SelfActor, class AActor* OtherActor, FVector NormalImpulse, const FHitResult& Hit);
+
+	/** Called when movement is blocked by a collision with another actor.  */
+	virtual void OnMoveBlockedBy(const FHitResult& BlockingImpact) {}
+
+	/** Called when falling movement ends. */
+	virtual void OnLanded() {}
+
+	// IAIResourceInterface begin
+	virtual void LockResource(EAILockSource::Type LockSource) OVERRIDE;
+	virtual void ClearResourceLock(EAILockSource::Type LockSource) OVERRIDE;
+	virtual void ForceUnlockResource() OVERRIDE;
+	virtual bool IsResourceLocked() const OVERRIDE;
+	// IAIResourceInterface end
+
+protected:
+
+	/** associated movement component */
+	UPROPERTY(transient)
+	class UNavMovementComponent* MovementComp;
+
+	/** associated navigation component */
+	UPROPERTY(transient)
+	class UNavigationComponent* NavComp;
+
+	/** smart link being currently traversed */
+	UPROPERTY(transient)
+	class USmartNavLinkComponent* CurrentSmartLink;
+
+	/** navigation data for agent described in movement component */
+	UPROPERTY(transient)
+	ANavigationData* MyNavData;
+
+	/** current status */
+	TEnumAsByte<EPathFollowingStatus::Type> Status;
+
+	/** requested path */
+	FNavPathSharedPtr Path;
+
+	/** value based on navigation agent's properties that's used for AcceptanceRadius when DefaultAcceptanceRadius is requested */
+	float MyDefaultAcceptanceRadius;
+
+	/** min distance to destination to consider request successful */
+	float AcceptanceRadius;
+
+	/** min distance to end of current path segment to consider segment finished */
+	float CurrentAcceptanceRadius;
+
+	/** part of agent radius used as min acceptance radius */
+	float MinAgentRadiusPct;
+
+	/** part of agent height used as min acceptable height difference */
+	float MinAgentHeightPct;
+
+	/** game specific data */
+	FCustomMoveSharedPtr GameData;
+
+	/** current request observer */
+	FRequestCompletedSignature OnRequestFinished;
+
+	/** destination actor. Use SetDestinationActor to set this */
+	TWeakObjectPtr<AActor> DestinationActor;
+
+	/** cached DestinationActor cast to INavAgentInterface. Use SetDestinationActor to set this */
+	const INavAgentInterface* DestinationAgent;
+
+	/** destination for current path segment */
+	FBasedPosition CurrentDestination;
+
+	/** relative offset from goal actor's location to end of path */
+	FVector MoveOffset;
+
+	/** agent location when movement was paused */
+	FVector LocationWhenPaused;
+
+	/** set when paths simplification using visibility tests are needed  (disabled by default because of performance) */
+	UPROPERTY(config)
+	uint32 bUseVisibilityTestsSimplification : 1;
+
+	/** increase acceptance radius with agent's radius */
+	uint32 bStopOnOverlap : 1;
+
+	/** if set, movement block detection will be used */
+	uint32 bUseBlockDetection : 1;
+
+	/** set when agent collides with goal actor */
+	uint32 bCollidedWithGoal : 1;
+
+	/** set when last move request was finished at goal */
+	uint32 bLastMoveReachedGoal : 1;
+
+	/** set when UpdateMove() is called during paused move, will update path's start segment on resuming */
+	uint32 bPendingPathStartUpdate : 1;
+
+	/** detect blocked movement when distance between center of location samples and furthest one (centroid radius) is below threshold */
+	float BlockDetectionDistance;
+
+	/** interval for collecting location samples */
+	float BlockDetectionInterval;
+
+	/** number of samples required for block detection */
+	int32 BlockDetectionSampleCount;
+
+	/** timestamp of last location sample */
+	float LastSampleTime;
+
+	/** index of next location sample in array */
+	int32 NextSampleIdx;
+
+	/** location samples for stuck detection */
+	TArray<FBasedPosition> LocationSamples;
+
+	/** index of path point being current move beginning */
+	int32 MoveSegmentStartIndex;
+
+	/** index of path point being current move target */
+	int32 MoveSegmentEndIndex;
+
+	/** reference of node at segment start */
+	NavNodeRef MoveSegmentStartRef;
+
+	/** reference of node at segment end */
+	NavNodeRef MoveSegmentEndRef;
+
+	/** direction of current move segment */
+	FVector MoveSegmentDirection;
+
+	/** reset path following data */
+	virtual void Reset();
+
+	/** should verify if agent if still on path ater movement has been resumed? */
+	virtual bool ShouldCheckPathOnResume() const;
+
+	/** sets variables related to current move segment */
+	virtual void SetMoveSegment(int32 SegmentStartIndex);
+	
+	/** follow current path segment */
+	virtual void FollowPathSegment(float DeltaTime);
+
+	/** check state of path following, update move segment if needed */
+	virtual void UpdatePathSegment();
+
+	/** switch to using smart link */
+	virtual void SetCurrentSmartLink(USmartNavLinkComponent* InLink, const FVector& DestPoint);
+
+	/** update blocked movement detection, @returns true if new sample was added */
+	virtual bool UpdateBlockDetection();
+
+	/** check if move is completed */
+	bool HasReachedDestination(const FVector& CurrentLocation) const;
+
+	/** check if segment is completed */
+	bool HasReachedCurrentTarget(const FVector& CurrentLocation) const;
+
+	/** check if moving agent has reached goal defined by cylinder */
+	bool HasReachedInternal(const FVector& GoalLocation, float GoalRadius, float GoalHalfHeight, const FVector& AgentLocation, float RadiusThreshold, bool bUseAgentRadius) const;
+
+	/** check if agent is on path */
+	virtual bool IsOnPath() const;
+
+	/** check if movement is blocked */
+	bool IsBlocked() const;
+
+	/** switch to next segment on path */
+	FORCEINLINE void SetNextMoveSegment() { SetMoveSegment(GetNextPathIndex()); }
+
+	FORCEINLINE static uint32 GetNextRequestId() { return NextRequestId++; }
+	FORCEINLINE void StoreRequestId() { CurrentRequestId = UPathFollowingComponent::GetNextRequestId(); }
+
+	/** Checks if this PathFollowingComponent is already on path, and
+	*	if so determines index of next path point
+	*	@return what PathFollowingComponent thinks should be next path point. INDEX_NONE if given path is invalid
+	*	@note this function does not set MoveSegmentEndIndex */
+	virtual int32 DetermineStartingPathPoint(const FNavigationPath* ConsideredPath) const;
+
+	/** @return index of path point, that should be target of current move segment */
+	virtual int32 DetermineCurrentTargetPathPoint(int32 StartIndex, bool bResumedFromSmartLink);
+
+	/** Visibility tests to skip some path points if possible
+	@param NextSegmentStartIndex Selected next segment to follow
+	@return  better path segment to follow, to use as MoveSegmentEndIndex */
+	int32 OptimizeSegmentVisibility(int32 StartIndex);
+
+	/** check if movement component is valid or tries to grab one from owner 
+	 *	@param bForce results in looking for owner's movement component even if pointer to one is already cached */
+	virtual bool UpdateMovementComponent(bool bForce = false);
+
+	/** clears Block Detection stored data effectively resetting the mechanism */
+	void ResetBlockDetectionData();
+
+	/** force creating new location sample for block detection */
+	void ForceBlockDetectionUpdate();
+
+	/** set move focus in AI owner */
+	void UpdateMoveFocus();
+
+	/** debug point reach test values */
+	void DebugReachTest(float& CurrentDot, float& CurrentDistance, float& CurrentHeight, uint8& bDotFailed, uint8& bDistanceFailed, uint8& bHeightFailed) const;
+	
+private:
+
+	/** used for debugging purposes to be able to identify which logged information
+	 *	results from which request, if there was multiple ones during one frame */
+	static uint32 NextRequestId;
+	FAIRequestID CurrentRequestId;
+
+	/** Current location on navigation data.  Lazy-updated, so read this via GetCurrentNavLocation(). 
+	 *	Since it makes conceptual sense for GetCurrentNavLocation() to be const but we may 
+	 *	need to update the cached value, CurrentNavLocation is mutable. */
+	mutable FNavLocation CurrentNavLocation;
+
+	/** used to keep track of which subsystem requested this AI resource be locked */
+	FAIResourceLock ResourceLock;
+
+	/** empty delegate for RequestMove */
+	static FRequestCompletedSignature UnboundRequestDelegate;
+
+public:
+	/** special float constant to symbolize "use default value". This does not contain 
+	 *	value to be used, it's used to detect the fact that it's requested, and 
+	 *	appropriate value from querier/doer will be pulled */
+	static const float DefaultAcceptanceRadius;
+};
diff --git a/Engine/Source/Runtime/AIModule/Private/AIController.cpp b/Engine/Source/Runtime/AIModule/Private/AIController.cpp
new file mode 100644
--- /dev/null
+++ b/Engine/Source/Runtime/AIModule/Private/AIController.cpp
@@ -0,0 +1,745 @@
+// Copyright 1998-2014 Epic Games, Inc. All Rights Reserved.
+
+#include "AIModulePrivate.h"
+#include "VisualLog.h"
+#include "Kismet/GameplayStatics.h"
+#include "DisplayDebugHelpers.h"
+#include "BrainComponent.h"
+#include "BehaviorTree/BehaviorTree.h"
+#include "BehaviorTree/BehaviorTreeManager.h"
+#include "BehaviorTree/BlackboardComponent.h"
+#include "Navigation/PathFollowingComponent.h"
+#include "Navigation/NavigationComponent.h"
+#include "AIController.h"
+
+// mz@todo these need to be removed, legacy code
+#define CLOSEPROXIMITY					500.f
+#define NEARSIGHTTHRESHOLD				2000.f
+#define MEDSIGHTTHRESHOLD				3162.f
+#define FARSIGHTTHRESHOLD				8000.f
+#define CLOSEPROXIMITYSQUARED			(CLOSEPROXIMITY*CLOSEPROXIMITY)
+#define NEARSIGHTTHRESHOLDSQUARED		(NEARSIGHTTHRESHOLD*NEARSIGHTTHRESHOLD)
+#define MEDSIGHTTHRESHOLDSQUARED		(MEDSIGHTTHRESHOLD*MEDSIGHTTHRESHOLD)
+#define FARSIGHTTHRESHOLDSQUARED		(FARSIGHTTHRESHOLD*FARSIGHTTHRESHOLD)
+
+
+bool AAIController::bAIIgnorePlayers = false;
+
+DECLARE_CYCLE_STAT(TEXT("MoveToLocation"),STAT_MoveToLocation,STATGROUP_AI);
+DECLARE_CYCLE_STAT(TEXT("MoveToActor"),STAT_MoveToActor,STATGROUP_AI);
+
+DEFINE_LOG_CATEGORY(LogAINavigation);
+
+AAIController::AAIController(const class FPostConstructInitializeProperties& PCIP)
+	: Super(PCIP)
+{
+
+	// set up navigation component
+	NavComponent = PCIP.CreateDefaultSubobject<UNavigationComponent>(this, TEXT("NavComponent"));
+
+	PathFollowingComponent = PCIP.CreateDefaultSubobject<UPathFollowingComponent>(this, TEXT("PathFollowingComponent"));
+	PathFollowingComponent->OnMoveFinished.AddUObject(this, &AAIController::OnMoveCompleted);
+
+	bSkipExtraLOSChecks = true;
+	bWantsPlayerState = false;
+}
+
+void AAIController::Tick(float DeltaTime)
+{
+	Super::Tick(DeltaTime);
+
+	UpdateControlRotation(DeltaTime);
+}
+
+void AAIController::PostInitializeComponents()
+{
+	Super::PostInitializeComponents();
+
+	if (bWantsPlayerState && !IsPendingKill() && (GetNetMode() != NM_Client))
+	{
+		InitPlayerState();
+	}
+}
+
+void AAIController::Reset()
+{
+	Super::Reset();
+
+	if (PathFollowingComponent)
+	{
+		PathFollowingComponent->AbortMove(TEXT("controller reset"));
+	}
+}
+
+void AAIController::DisplayDebug(class UCanvas* Canvas, const FDebugDisplayInfo& DebugDisplay, float& YL, float& YPos)
+{
+	Super::DisplayDebug(Canvas, DebugDisplay, YL, YPos);
+
+	static FName NAME_AI = FName(TEXT("AI"));
+	if (DebugDisplay.IsDisplayOn(NAME_AI))
+	{
+		if (PathFollowingComponent)
+		{
+			PathFollowingComponent->DisplayDebug(Canvas, DebugDisplay, YL, YPos);
+		}
+		
+		if ( GetFocusActor() )
+		{
+			Canvas->DrawText(GEngine->GetSmallFont(), FString::Printf(TEXT("      Focus %s"), *GetFocusActor()->GetName()), 4.0f, YPos);
+			YPos += YL;
+		}
+	}
+}
+
+#if ENABLE_VISUAL_LOG
+
+void AAIController::GrabDebugSnapshot(FVisLogEntry* Snapshot) const
+{
+	FVisLogEntry::FStatusCategory MyCategory;
+	MyCategory.Category = TEXT("AI Contoller");
+	MyCategory.Add(TEXT("Pawn"), GetNameSafe(GetPawn()));
+	AActor* FocusActor = GetFocusActor();
+	MyCategory.Add(TEXT("Focus"), GetDebugName(FocusActor));
+	if (FocusActor == NULL)
+	{
+		MyCategory.Add(TEXT("Focus Location"), GetFocalPoint().ToString());
+	}
+	Snapshot->Status.Add(MyCategory);
+
+	if (GetPawn())
+	{
+		Snapshot->Location = GetPawn()->GetActorLocation();
+	}
+
+	if (NavComponent.IsValid())
+	{
+		NavComponent->DescribeSelfToVisLog(Snapshot);
+	}
+
+	if (PathFollowingComponent.IsValid())
+	{
+		PathFollowingComponent->DescribeSelfToVisLog(Snapshot);
+	}
+
+	if (BrainComponent != NULL)
+	{
+		BrainComponent->DescribeSelfToVisLog(Snapshot);
+	}
+}
+#endif // ENABLE_VISUAL_LOG
+
+void AAIController::GetPlayerViewPoint(FVector& out_Location, FRotator& out_Rotation) const
+{
+	// AI does things from the Pawn
+	if (GetPawn() != NULL)
+	{
+		out_Location = GetPawn()->GetActorLocation();
+		out_Rotation = GetPawn()->GetActorRotation();
+	}
+	else
+	{
+		Super::GetPlayerViewPoint(out_Location, out_Rotation);
+	}
+}
+
+void AAIController::SetFocalPoint( FVector FP, bool bOffsetFromBase, uint8 InPriority)
+{
+	if (InPriority >= FocusInformation.Priorities.Num())
+	{
+		FocusInformation.Priorities.SetNum(InPriority+1);
+	}
+	FFocusKnowledge::FFocusItem& Focusitem = FocusInformation.Priorities[InPriority];
+
+	{
+		AActor* FocalActor = NULL;
+		if (bOffsetFromBase)
+		{
+			APawn *Pawn = GetPawn();
+			if (Pawn && Pawn->GetMovementBase())
+			{
+				FocalActor = Pawn->GetMovementBase()->GetOwner();
+			}
+		}
+		Focusitem.Position.Set( FocalActor, FP );
+	}
+
+	Focusitem.Actor = NULL;
+}
+
+FVector AAIController::GetFocalPoint() const
+{
+	FBasedPosition FinalFocus;
+	// find focus with highest priority
+	for( int32 Index = FocusInformation.Priorities.Num()-1; Index >= 0; --Index)
+	{
+		const FFocusKnowledge::FFocusItem& Focusitem = FocusInformation.Priorities[Index];
+		if ( Focusitem.Actor.IsValid() )
+		{
+			const AActor* Focus = Focusitem.Actor.Get();
+			UPrimitiveComponent* MyBase = GetPawn() ? GetPawn()->GetMovementBase() : NULL;
+			const bool bRequestedFocusIsBased = MyBase && Cast<const APawn>(Focus) && (Cast<const APawn>(Focus)->GetMovementBase() == MyBase);
+			FinalFocus.Set(bRequestedFocusIsBased && MyBase ? MyBase->GetOwner() : NULL, Focus->GetActorLocation());
+			break;
+		}
+		else if( !(*Focusitem.Position).IsZero() )
+		{
+			FinalFocus = Focusitem.Position;
+			break;
+		}
+	}
+
+	return *FinalFocus;
+}
+
+AActor* AAIController::GetFocusActor() const
+{
+	AActor* FocusActor = NULL;
+	for( int32 Index = FocusInformation.Priorities.Num()-1; Index >= 0; --Index)
+	{
+		const FFocusKnowledge::FFocusItem& Focusitem = FocusInformation.Priorities[Index];
+		if ( Focusitem.Actor.IsValid() )
+		{
+			FocusActor = Focusitem.Actor.Get();
+			break;
+		}
+		else if( !(*Focusitem.Position).IsZero() )
+		{
+			break;
+		}
+	}
+
+	return FocusActor;
+}
+
+void AAIController::K2_SetFocus(AActor* NewFocus)
+{ 
+	SetFocus(NewFocus, EAIFocusPriority::Gameplay); 
+}
+
+void AAIController::K2_SetFocalPoint(FVector FP, bool bOffsetFromBase)
+{
+	SetFocalPoint(FP, bOffsetFromBase, EAIFocusPriority::Gameplay);
+}
+
+void AAIController::K2_ClearFocus()
+{
+	ClearFocus(EAIFocusPriority::Gameplay);
+}
+
+void AAIController::SetFocus(AActor* NewFocus, uint8 InPriority)
+{
+	if( NewFocus )
+	{
+		if (InPriority >= FocusInformation.Priorities.Num())
+		{
+			FocusInformation.Priorities.SetNum(InPriority+1);
+		}
+		FocusInformation.Priorities[InPriority].Actor = NewFocus;
+	}
+	else
+	{
+		ClearFocus(InPriority);
+	}
+}
+
+void AAIController::ClearFocus(uint8 InPriority)
+{
+	if (InPriority < FocusInformation.Priorities.Num())
+	{
+		FocusInformation.Priorities[InPriority].Actor = NULL;
+		FocusInformation.Priorities[InPriority].Position.Clear();
+	}
+}
+
+bool AAIController::LineOfSightTo(const AActor* Other, FVector ViewPoint, bool bAlternateChecks)
+{
+	if( !Other )
+	{
+		return false;
+	}
+		
+	if ( ViewPoint.IsZero() )
+	{
+		AActor*	ViewTarg = GetViewTarget();
+		ViewPoint = ViewTarg->GetActorLocation();
+		if( ViewTarg == GetPawn() )
+		{
+			ViewPoint.Z += GetPawn()->BaseEyeHeight; //look from eyes
+		}
+	}
+
+	static FName NAME_LineOfSight = FName(TEXT("LineOfSight"));
+	FVector TargetLocation = Other->GetTargetLocation(GetPawn());
+
+	FCollisionQueryParams CollisionParams(NAME_LineOfSight, true, this->GetPawn());
+	CollisionParams.AddIgnoredActor(Other);
+
+	bool bHit = GetWorld()->LineTraceTest(ViewPoint, TargetLocation, ECC_Visibility, CollisionParams);
+	if( !bHit )
+	{
+		return true;
+	}
+
+	// if other isn't using a cylinder for collision and isn't a Pawn (which already requires an accurate cylinder for AI)
+	// then don't go any further as it likely will not be tracing to the correct location
+	const APawn * OtherPawn = Cast<const APawn>(Other);
+	if (!OtherPawn && Cast<UCapsuleComponent>(Other->GetRootComponent()) == NULL)
+	{
+		return false;
+	}
+
+	const FVector OtherActorLocation = Other->GetActorLocation();
+	float distSq = (OtherActorLocation - ViewPoint).SizeSquared();
+	if ( distSq > FARSIGHTTHRESHOLDSQUARED )
+	{
+		return false;
+	}
+
+	if ( !OtherPawn && (distSq > NEARSIGHTTHRESHOLDSQUARED) ) 
+	{
+		return false;
+	}
+
+	float OtherRadius, OtherHeight;
+	Other->GetSimpleCollisionCylinder(OtherRadius, OtherHeight);
+
+	if ( !bAlternateChecks || !bLOSflag )
+	{
+		//try viewpoint to head
+		bHit = GetWorld()->LineTraceTest(ViewPoint,  OtherActorLocation + FVector(0.f,0.f,OtherHeight), ECC_Visibility, CollisionParams);
+		if ( !bHit )
+		{
+			return true;
+		}
+	}
+
+	if( !bSkipExtraLOSChecks && (!bAlternateChecks || bLOSflag) )
+	{
+		// only check sides if width of other is significant compared to distance
+		if( OtherRadius * OtherRadius/(OtherActorLocation - ViewPoint).SizeSquared() < 0.0001f )
+		{
+			return false;
+		}
+		//try checking sides - look at dist to four side points, and cull furthest and closest
+		FVector Points[4];
+		Points[0] = OtherActorLocation - FVector(OtherRadius, -1 * OtherRadius, 0);
+		Points[1] = OtherActorLocation + FVector(OtherRadius, OtherRadius, 0);
+		Points[2] = OtherActorLocation - FVector(OtherRadius, OtherRadius, 0);
+		Points[3] = OtherActorLocation + FVector(OtherRadius, -1 * OtherRadius, 0);
+		int32 imin = 0;
+		int32 imax = 0;
+		float currentmin = (Points[0] - ViewPoint).SizeSquared();
+		float currentmax = currentmin;
+		for ( int32 i=1; i<4; i++ )
+		{
+			float nextsize = (Points[i] - ViewPoint).SizeSquared();
+			if (nextsize > currentmax)
+			{
+				currentmax = nextsize;
+				imax = i;
+			}
+			else if (nextsize < currentmin)
+			{
+				currentmin = nextsize;
+				imin = i;
+			}
+		}
+
+		for ( int32 i=0; i<4; i++ )
+		{
+			if	( (i != imin) && (i != imax) )
+			{
+				bHit = GetWorld()->LineTraceTest(ViewPoint,  Points[i], ECC_Visibility, CollisionParams);
+				if ( !bHit )
+				{
+					return true;
+				}
+			}
+		}
+	}
+	return false;
+}
+
+DEFINE_LOG_CATEGORY_STATIC(LogTestAI, All, All);
+
+void AAIController::UpdateControlRotation(float DeltaTime, bool bUpdatePawn)
+{
+	// Look toward focus
+	FVector FocalPoint = GetFocalPoint();
+	if( !FocalPoint.IsZero() && GetPawn())
+	{
+		FVector Direction = FocalPoint - GetPawn()->GetActorLocation();
+		FRotator NewControlRotation = Direction.Rotation();
+
+		// Don't pitch view of walking pawns unless looking at another pawn
+		if ( GetPawn()->GetMovementComponent() && GetPawn()->GetMovementComponent()->IsMovingOnGround() &&
+			PathFollowingComponent && (!PathFollowingComponent->GetMoveGoal() || !Cast<APawn>(PathFollowingComponent->GetMoveGoal()) ) )
+		{
+			NewControlRotation.Pitch = 0.f;
+		}
+		NewControlRotation.Yaw = FRotator::ClampAxis(NewControlRotation.Yaw);
+
+		if (!GetControlRotation().Equals(NewControlRotation, 1e-3f))
+		{
+			SetControlRotation(NewControlRotation);
+
+			APawn* const P = GetPawn();
+			if (P && bUpdatePawn)
+			{
+				P->FaceRotation(NewControlRotation, DeltaTime);
+			}
+		}
+	}
+}
+
+
+void AAIController::Possess(APawn* InPawn)
+{
+	Super::Possess(InPawn);
+
+	if (!GetPawn())
+	{
+		return;
+	}
+	
+	// no point in doing navigation setup if pawn has no movement component
+	const UPawnMovementComponent* MovementComp = InPawn->GetMovementComponent();
+	if (MovementComp != NULL)
+	{
+		UpdateNavigationComponents();
+	}
+
+	if (PathFollowingComponent)
+	{
+		PathFollowingComponent->Initialize();
+	}
+
+	if (bWantsPlayerState)
+	{
+		ChangeState(NAME_Playing);
+	}
+
+	OnPossess(InPawn);
+}
+
+void AAIController::UnPossess()
+{
+	Super::UnPossess();
+
+	if (PathFollowingComponent)
+	{
+		PathFollowingComponent->Cleanup();
+	}
+}
+
+void AAIController::InitNavigationControl(UNavigationComponent*& PathFindingComp, UPathFollowingComponent*& PathFollowingComp)
+{
+	PathFindingComp = NavComponent;
+	PathFollowingComp = PathFollowingComponent;
+}
+
+EPathFollowingRequestResult::Type AAIController::MoveToActor(class AActor* Goal, float AcceptanceRadius, bool bStopOnOverlap, bool bUsePathfinding, bool bCanStrafe, TSubclassOf<UNavigationQueryFilter> FilterClass)
+{
+	SCOPE_CYCLE_COUNTER(STAT_MoveToActor);
+	EPathFollowingRequestResult::Type Result = EPathFollowingRequestResult::Failed;
+
+	UE_VLOG(this, LogAINavigation, Log, TEXT("MoveToActor: Goal(%s) AcceptRadius(%.1f%s) bUsePathfinding(%d) bCanStrafe(%d)"),
+		*GetNameSafe(Goal), AcceptanceRadius, bStopOnOverlap ? TEXT(" + agent") : TEXT(""), bUsePathfinding, bCanStrafe);
+
+	if (Goal)
+	{
+		if (PathFollowingComponent && PathFollowingComponent->HasReached(Goal, AcceptanceRadius, !bStopOnOverlap))
+		{
+			UE_VLOG(this, LogAINavigation, Log, TEXT("MoveToActor: already at goal!"));
+
+			// make sure previous move request gets aborted
+			PathFollowingComponent->AbortMove(TEXT("Aborting move due to new move request finishing with AlreadyAtGoal"), FAIRequestID::AnyRequest);
+
+			PathFollowingComponent->SetLastMoveAtGoal(true);
+
+			OnMoveCompleted(FAIRequestID::CurrentRequest, EPathFollowingResult::Success);
+			Result = EPathFollowingRequestResult::AlreadyAtGoal;
+		}
+		else
+		{
+			FNavPathSharedPtr Path = FindPath(Goal, bUsePathfinding, UNavigationQueryFilter::GetQueryFilter(NavComponent.Get() ? NavComponent->GetNavData() : NULL, FilterClass));
+			if (Path.IsValid())
+			{
+				bAllowStrafe = bCanStrafe;
+
+				const uint32 RequestID = RequestMove(Path, Goal, AcceptanceRadius, bStopOnOverlap);
+				Result = (RequestID != FAIRequestID::InvalidRequest) ? EPathFollowingRequestResult::RequestSuccessful : EPathFollowingRequestResult::Failed;
+			}
+		}
+	}
+
+	if (Result == EPathFollowingRequestResult::Failed)
+	{
+		if (PathFollowingComponent)
+		{
+			PathFollowingComponent->SetLastMoveAtGoal(false);
+		}
+
+		OnMoveCompleted(FAIRequestID::InvalidRequest, EPathFollowingResult::Invalid);
+	}
+
+	return Result;
+}
+
+EPathFollowingRequestResult::Type AAIController::MoveToLocation(const FVector& Dest, float AcceptanceRadius, bool bStopOnOverlap, bool bUsePathfinding, bool bProjectDestinationToNavigation, bool bCanStrafe, TSubclassOf<UNavigationQueryFilter> FilterClass)
+{
+	SCOPE_CYCLE_COUNTER(STAT_MoveToLocation);
+
+	EPathFollowingRequestResult::Type Result = EPathFollowingRequestResult::Failed;
+	bool bCanRequestMove = true;
+
+	UE_VLOG(this, LogAINavigation, Log, TEXT("MoveToLocation: Goal(%s) AcceptRadius(%.1f%s) bUsePathfinding(%d) bCanStrafe(%d)"),
+		*Dest.ToString(), AcceptanceRadius, bStopOnOverlap ? TEXT(" + agent") : TEXT(""), bUsePathfinding, bCanStrafe);
+
+	// Check input is valid
+	if( Dest.ContainsNaN() )
+	{
+		UE_VLOG(this, LogAINavigation, Error, TEXT("AAIController::MoveToLocation: Destination is not valid! Goal(%s) AcceptRadius(%.1f%s) bUsePathfinding(%d) bCanStrafe(%d)"),
+			*Dest.ToString(), AcceptanceRadius, bStopOnOverlap ? TEXT(" + agent") : TEXT(""), bUsePathfinding, bCanStrafe);
+		
+		ensure( !Dest.ContainsNaN() );
+		bCanRequestMove = false;
+	}
+
+	FVector GoalLocation = Dest;
+
+	// fail if projection to navigation is required but it either failed or there's not navigation component
+	if (bCanRequestMove && bProjectDestinationToNavigation && (NavComponent.Get() == NULL || !NavComponent->ProjectPointToNavigation(Dest, GoalLocation)))
+	{
+		UE_VLOG_LOCATION(this, LogAINavigation, Error, Dest, 30.f, FLinearColor::Red, TEXT("AAIController::MoveToLocation failed to project destination location to navmesh"));
+		bCanRequestMove = false;
+	}
+
+	if (bCanRequestMove && PathFollowingComponent && PathFollowingComponent->HasReached(GoalLocation, AcceptanceRadius, !bStopOnOverlap))
+	{
+		UE_VLOG(this, LogAINavigation, Log, TEXT("MoveToLocation: already at goal!"));
+
+		// make sure previous move request gets aborted
+		PathFollowingComponent->AbortMove(TEXT("Aborting move due to new move request finishing with AlreadyAtGoal"), FAIRequestID::AnyRequest);
+
+		PathFollowingComponent->SetLastMoveAtGoal(true);
+
+		OnMoveCompleted(FAIRequestID::CurrentRequest, EPathFollowingResult::Success);
+		Result = EPathFollowingRequestResult::AlreadyAtGoal;
+		bCanRequestMove = false;
+	}
+
+	if (bCanRequestMove)
+	{
+		FNavPathSharedPtr Path = FindPath(GoalLocation, bUsePathfinding, UNavigationQueryFilter::GetQueryFilter(NavComponent.Get() != NULL ? NavComponent->GetNavData() : NULL, FilterClass));
+		if (Path.IsValid())
+		{
+			bAllowStrafe = bCanStrafe;
+
+			const FAIRequestID RequestID = RequestMove(Path, NULL, AcceptanceRadius, bStopOnOverlap);
+			Result = RequestID.IsValid() ? EPathFollowingRequestResult::RequestSuccessful : EPathFollowingRequestResult::Failed;
+		}
+	}
+
+	if (Result == EPathFollowingRequestResult::Failed)
+	{
+		if (PathFollowingComponent)
+		{
+			PathFollowingComponent->SetLastMoveAtGoal(false);
+		}
+
+		OnMoveCompleted(FAIRequestID::InvalidRequest, EPathFollowingResult::Invalid);
+	}
+
+	return Result;
+}
+
+FAIRequestID AAIController::RequestMove(FNavPathSharedPtr Path, class AActor* Goal, float AcceptanceRadius, bool bStopOnOverlap, FCustomMoveSharedPtr CustomData)
+{
+	uint32 RequestID = FAIRequestID::InvalidRequest;
+	if (PathFollowingComponent)
+	{
+		RequestID = PathFollowingComponent->RequestMove(Path, Goal, AcceptanceRadius, bStopOnOverlap, CustomData);
+	}
+
+	return RequestID;
+}
+
+bool AAIController::PauseMove(FAIRequestID RequestToPause)
+{
+	if (PathFollowingComponent != NULL && RequestToPause.IsEquivalent(PathFollowingComponent->GetCurrentRequestId()))
+	{
+		PathFollowingComponent->PauseMove(RequestToPause);
+		return true;
+	}
+	return false;
+}
+
+bool AAIController::ResumeMove(FAIRequestID RequestToResume)
+{
+	if (PathFollowingComponent != NULL && RequestToResume.IsEquivalent(PathFollowingComponent->GetCurrentRequestId()))
+	{
+		PathFollowingComponent->ResumeMove(RequestToResume);
+		return true;
+	}
+	return false;
+}
+
+void AAIController::StopMovement()
+{
+	UE_VLOG(this, LogAINavigation, Log, TEXT("AAIController::StopMovement: %s STOP MOVEMENT"), *GetNameSafe(GetPawn()));
+	PathFollowingComponent->AbortMove(TEXT("StopMovement"));
+}
+
+FNavPathSharedPtr AAIController::FindPath(class AActor* Goal, bool bUsePathfinding, TSharedPtr<const FNavigationQueryFilter> QueryFilter)
+{
+	if (NavComponent)
+	{
+		const bool bFoundPath = bUsePathfinding ? NavComponent->FindPathToActor(Goal, QueryFilter) : NavComponent->FindSimplePathToActor(Goal);
+		if (bFoundPath)
+		{
+			return NavComponent->GetPath();
+		}
+	}
+
+	return NULL;
+}
+
+FNavPathSharedPtr AAIController::FindPath(const FVector& Dest, bool bUsePathfinding, TSharedPtr<const FNavigationQueryFilter> QueryFilter)
+{
+	if (NavComponent)
+	{
+		const bool bFoundPath = bUsePathfinding ? NavComponent->FindPathToLocation(Dest, QueryFilter) : NavComponent->FindSimplePathToLocation(Dest);
+		if (bFoundPath) 
+		{
+			return NavComponent->GetPath();
+		}
+	}
+
+	return NULL;
+}
+
+EPathFollowingStatus::Type AAIController::GetMoveStatus() const
+{
+	return (PathFollowingComponent) ? PathFollowingComponent->GetStatus() : EPathFollowingStatus::Idle;
+}
+
+bool AAIController::HasPartialPath() const
+{
+	return (PathFollowingComponent.Get() != NULL) && (PathFollowingComponent->HasPartialPath());
+}
+
+FVector AAIController::GetImmediateMoveDestination() const
+{
+	return (PathFollowingComponent) ? PathFollowingComponent->GetCurrentTargetLocation() : FVector::ZeroVector;
+}
+
+void AAIController::SetMoveBlockDetection(bool bEnable)
+{
+	if (PathFollowingComponent)
+	{
+		PathFollowingComponent->SetBlockDetectionState(bEnable);
+	}
+}
+
+void AAIController::OnMoveCompleted(FAIRequestID RequestID, EPathFollowingResult::Type Result)
+{
+	ReceiveMoveCompleted.Broadcast(RequestID, Result);
+}
+
+bool AAIController::RunBehaviorTree(UBehaviorTree* BTAsset)
+{
+	// @todo: find BrainComponent and see if it's BehaviorTreeComponent
+	// Also check if BTAsset requires BlackBoardComponent, and if so 
+	// check if BB type is accepted by BTAsset.
+	// Spawn BehaviorTreeComponent if none present. 
+	// Spawn BlackBoardComponent if none present, but fail if one is present but is not of compatible class
+	if (BTAsset == NULL)
+	{
+		UE_VLOG(this, LogBehaviorTree, Warning, TEXT("RunBehaviorTree: Unable to run NULL behavior tree"));
+		return false;
+	}
+
+	bool bSuccess = true;
+	bool bShouldInitializeBlackboard = false;
+
+	// see if need a blackboard component at all
+	UBlackboardComponent* BlackboardComp = NULL;
+	if (BTAsset->BlackboardAsset)
+	{
+		bSuccess = UseBlackboard(BTAsset->BlackboardAsset);
+		BlackboardComp = FindComponentByClass<UBlackboardComponent>();
+	}
+	
+	if (bSuccess)
+	{
+		UBehaviorTreeComponent* BTComp = Cast<UBehaviorTreeComponent>(BrainComponent);
+		if (BTComp == NULL)
+		{
+			UE_VLOG(this, LogBehaviorTree, Log, TEXT("RunBehaviorTree: spawning BehaviorTreeComponent.."));
+
+			BrainComponent = BTComp = ConstructObject<UBehaviorTreeComponent>(UBehaviorTreeComponent::StaticClass(), this, TEXT("BTComponent"));
+			BrainComponent->RegisterComponent();
+
+			if (BrainComponent->bWantsInitializeComponent)
+			{
+				// make sure that newly created component is initialized before running BT
+				// both blackboard and BT to must exist before calling it!
+				BrainComponent->InitializeComponent();
+			}
+		}
+
+		check(BTComp != NULL);
+		BTComp->StartTree(BTAsset, EBTExecutionMode::Looped);
+	}
+
+	return bSuccess;
+}
+
+bool AAIController::UseBlackboard(UBlackboardData* BlackboardAsset)
+{
+	if (BlackboardAsset == NULL)
+	{
+		UE_VLOG(this, LogBehaviorTree, Log, TEXT("UseBlackboard: trying to use NULL Blackboard asset. Ignoring"));
+		return false;
+	}
+
+	bool bSuccess = true;
+	UBlackboardComponent* BlackboardComp = FindComponentByClass<UBlackboardComponent>();
+
+	if (BlackboardComp == NULL)
+	{
+		BlackboardComp = ConstructObject<UBlackboardComponent>(UBlackboardComponent::StaticClass(), this, TEXT("BlackboardComponent"));
+		if (BlackboardComp != NULL)
+		{
+			BlackboardComp->InitializeBlackboard(BlackboardAsset);
+			BlackboardComp->RegisterComponent();
+			BlackboardComp->InitializeComponent();
+		}
+
+	}
+	else if (BlackboardComp->GetBlackboardAsset() == NULL)
+	{
+		BlackboardComp->InitializeBlackboard(BlackboardAsset);
+	}
+	else if (BlackboardComp->GetBlackboardAsset() != BlackboardAsset)
+	{
+		UE_VLOG(this, LogBehaviorTree, Log, TEXT("UseBlackboard: requested blackboard %s while already has %s instantiated. Forcing new BB.")
+			, *GetNameSafe(BlackboardAsset), *GetNameSafe(BlackboardComp->GetBlackboardAsset()));
+		BlackboardComp->InitializeBlackboard(BlackboardAsset);
+	}
+
+	return bSuccess;
+}
+
+bool AAIController::SuggestTossVelocity(FVector& OutTossVelocity, FVector Start, FVector End, float TossSpeed, bool bPreferHighArc, float CollisionRadius, bool bOnlyTraceUp)
+{
+	// pawn's physics volume gets 2nd priority
+	APhysicsVolume const* const PhysicsVolume = GetPawn() ? GetPawn()->GetPawnPhysicsVolume() : NULL;
+	float const GravityOverride = PhysicsVolume ? PhysicsVolume->GetGravityZ() : 0.f;
+	ESuggestProjVelocityTraceOption::Type const TraceOption = bOnlyTraceUp ? ESuggestProjVelocityTraceOption::OnlyTraceWhileAsceding : ESuggestProjVelocityTraceOption::TraceFullPath;
+
+	return UGameplayStatics::SuggestProjectileVelocity(this, OutTossVelocity, Start, End, TossSpeed, bPreferHighArc, CollisionRadius, GravityOverride, TraceOption);
+}
+
+
diff --git a/Engine/Source/Runtime/AIModule/Private/AIInterfaces.cpp b/Engine/Source/Runtime/AIModule/Private/AIInterfaces.cpp
new file mode 100644
--- /dev/null
+++ b/Engine/Source/Runtime/AIModule/Private/AIInterfaces.cpp
@@ -0,0 +1,9 @@
+// Copyright 1998-2014 Epic Games, Inc. All Rights Reserved.
+
+#include "AIModulePrivate.h"
+#include "AIResourceInterface.h"
+
+UAIResourceInterface::UAIResourceInterface(const class FPostConstructInitializeProperties& PCIP)
+	: Super(PCIP)
+{
+}
diff --git a/Engine/Source/Runtime/AIModule/Private/AIModule.cpp b/Engine/Source/Runtime/AIModule/Private/AIModule.cpp
--- a/Engine/Source/Runtime/AIModule/Private/AIModule.cpp
+++ b/Engine/Source/Runtime/AIModule/Private/AIModule.cpp
@@ -1,6 +1,6 @@
 // Copyright 1998-2014 Epic Games, Inc. All Rights Reserved.
 
-#include "AIModulePrivatePCH.h"
+#include "AIModulePrivate.h"
 
 class FAIModule : public IAIModule
 {
diff --git a/Engine/Source/Runtime/AIModule/Private/AIModulePrivate.h b/Engine/Source/Runtime/AIModule/Private/AIModulePrivate.h
new file mode 100644
--- /dev/null
+++ b/Engine/Source/Runtime/AIModule/Private/AIModulePrivate.h
@@ -0,0 +1,29 @@
+// Copyright 1998-2014 Epic Games, Inc. All Rights Reserved.
+
+#pragma once
+
+
+
+#if WITH_EDITOR
+#include "UnrealEd.h"
+#endif
+
+#include "Core.h"
+#include "Engine.h"
+
+#include "AIModule.h"
+
+#include "AITypes.h"
+#include "EnvironmentQuery/EnvQueryTypes.h"
+#include "BehaviorTree/Blackboard/BlackboardKeyType.h"
+#include "AIController.h"
+#include "AI/Navigation/NavigationSystem.h"
+
+#include "BehaviorTree/BTTaskNode.h"
+#include "BehaviorTree/Tasks/BTTask_BlackboardBase.h"
+#include "BehaviorTree/Decorators/BTDecorator_BlackboardBase.h"
+#include "BehaviorTree/Services/BTService_BlackboardBase.h"
+
+DECLARE_LOG_CATEGORY_EXTERN(LogAINavigation, Warning, All);
+DECLARE_LOG_CATEGORY_EXTERN(LogBehaviorTree, Display, All);
+DECLARE_LOG_CATEGORY_EXTERN(LogCrowdFollowing, Warning, All);
\ No newline at end of file
diff --git a/Engine/Source/Runtime/AIModule/Private/AISystem.cpp b/Engine/Source/Runtime/AIModule/Private/AISystem.cpp
new file mode 100644
--- /dev/null
+++ b/Engine/Source/Runtime/AIModule/Private/AISystem.cpp
@@ -0,0 +1,8 @@
+// Copyright 1998-2014 Epic Games, Inc. All Rights Reserved.
+
+#include "AIModulePrivate.h"
+#include "AISystem.h"
+
+UAISystem::UAISystem(const class FPostConstructInitializeProperties& PCIP) : Super(PCIP)
+{
+}
\ No newline at end of file
diff --git a/Engine/Source/Runtime/AIModule/Private/AITypes.cpp b/Engine/Source/Runtime/AIModule/Private/AITypes.cpp
new file mode 100644
--- /dev/null
+++ b/Engine/Source/Runtime/AIModule/Private/AITypes.cpp
@@ -0,0 +1,42 @@
+// Copyright 1998-2014 Epic Games, Inc. All Rights Reserved.
+
+#include "AIModulePrivate.h"
+#include "AITypes.h"
+
+//----------------------------------------------------------------------//
+// FAIResourceLock
+//----------------------------------------------------------------------//
+FAIResourceLock::FAIResourceLock()
+{
+	FMemory::MemZero(Locks);
+}
+
+void FAIResourceLock::ForceClearAllLocks()
+{
+	FMemory::MemZero(Locks);
+}
+
+FString FAIResourceLock::GetLockSourceName() const
+{
+	const static UEnum* SourceEnum = FindObject<UEnum>(ANY_PACKAGE, TEXT("EAILockSource"));
+
+	FString LockNames;
+
+	for (int32 LockLevel = 0; LockLevel < int32(EAILockSource::MAX); ++LockLevel)
+	{
+		if (Locks[LockLevel])
+		{
+			LockNames += FString::Printf(TEXT("%s, "), *SourceEnum->GetEnumName(LockLevel));
+		}
+	}
+
+	return LockNames;
+}
+
+//----------------------------------------------------------------------//
+// FAIRequestID
+//----------------------------------------------------------------------//
+const FAIRequestID FAIRequestID::InvalidRequest(FAIRequestID::InvalidRequestID);
+const FAIRequestID FAIRequestID::AnyRequest(FAIRequestID::AnyRequestID);
+// same as AnyRequest on purpose. Just a readability thing
+const FAIRequestID FAIRequestID::CurrentRequest(FAIRequestID::AnyRequestID);
diff --git a/Engine/Source/Runtime/AIModule/Private/BehaviorTree/BTAuxiliaryNode.cpp b/Engine/Source/Runtime/AIModule/Private/BehaviorTree/BTAuxiliaryNode.cpp
new file mode 100644
--- /dev/null
+++ b/Engine/Source/Runtime/AIModule/Private/BehaviorTree/BTAuxiliaryNode.cpp
@@ -0,0 +1,98 @@
+// Copyright 1998-2014 Epic Games, Inc. All Rights Reserved.
+
+#include "AIModulePrivate.h"
+#include "BehaviorTree/BTAuxiliaryNode.h"
+
+UBTAuxiliaryNode::UBTAuxiliaryNode(const class FPostConstructInitializeProperties& PCIP) : Super(PCIP)
+{
+	bNotifyBecomeRelevant = false;
+	bNotifyCeaseRelevant = false;
+	bNotifyTick = false;
+	bTickIntervals = false;
+}
+
+void UBTAuxiliaryNode::WrappedOnBecomeRelevant(class UBehaviorTreeComponent* OwnerComp, uint8* NodeMemory) const
+{
+	if (bNotifyBecomeRelevant)
+	{
+		const UBTNode* NodeOb = bCreateNodeInstance ? GetNodeInstance(OwnerComp, NodeMemory) : this;
+		if (NodeOb)
+		{
+			((UBTAuxiliaryNode*)NodeOb)->OnBecomeRelevant(OwnerComp, NodeMemory);
+		}
+	}
+}
+
+void UBTAuxiliaryNode::WrappedOnCeaseRelevant(class UBehaviorTreeComponent* OwnerComp, uint8* NodeMemory) const
+{
+	if (bNotifyCeaseRelevant)
+	{
+		const UBTNode* NodeOb = bCreateNodeInstance ? GetNodeInstance(OwnerComp, NodeMemory) : this;
+		if (NodeOb)
+		{
+			((UBTAuxiliaryNode*)NodeOb)->OnCeaseRelevant(OwnerComp, NodeMemory);
+		}
+	}
+}
+
+void UBTAuxiliaryNode::WrappedTickNode(class UBehaviorTreeComponent* OwnerComp, uint8* NodeMemory, float DeltaSeconds) const
+{
+	if (bNotifyTick)
+	{
+		if (bTickIntervals)
+		{
+			FBTAuxiliaryMemory* AuxMemory = GetSpecialNodeMemory<FBTAuxiliaryMemory>(NodeMemory);
+			AuxMemory->NextTickRemainingTime -= DeltaSeconds;
+			if (AuxMemory->NextTickRemainingTime > 0.0f)
+			{
+				return;
+			}	
+		}
+
+		const UBTNode* NodeOb = bCreateNodeInstance ? GetNodeInstance(OwnerComp, NodeMemory) : this;
+		if (NodeOb)
+		{
+			((UBTAuxiliaryNode*)NodeOb)->TickNode(OwnerComp, NodeMemory, DeltaSeconds);
+		}
+	}
+}
+
+void UBTAuxiliaryNode::SetNextTickTime(uint8* NodeMemory, float RemainingTime) const
+{
+	if (bTickIntervals)
+	{
+		FBTAuxiliaryMemory* AuxMemory = GetSpecialNodeMemory<FBTAuxiliaryMemory>(NodeMemory);
+		AuxMemory->NextTickRemainingTime = RemainingTime;
+	}
+}
+
+void UBTAuxiliaryNode::DescribeRuntimeValues(const class UBehaviorTreeComponent* OwnerComp, uint8* NodeMemory, EBTDescriptionVerbosity::Type Verbosity, TArray<FString>& Values) const
+{
+	Super::DescribeRuntimeValues(OwnerComp, NodeMemory, Verbosity, Values);
+
+	if (Verbosity == EBTDescriptionVerbosity::Detailed && bTickIntervals)
+	{
+		FBTAuxiliaryMemory* AuxMemory = GetSpecialNodeMemory<FBTAuxiliaryMemory>(NodeMemory);
+		Values.Add(FString::Printf(TEXT("next tick: %ss"), *FString::SanitizeFloat(AuxMemory->NextTickRemainingTime)));
+	}
+}
+
+void UBTAuxiliaryNode::OnBecomeRelevant(class UBehaviorTreeComponent* OwnerComp, uint8* NodeMemory)
+{
+	// empty in base class
+}
+
+void UBTAuxiliaryNode::OnCeaseRelevant(class UBehaviorTreeComponent* OwnerComp, uint8* NodeMemory)
+{
+	// empty in base class
+}
+
+void UBTAuxiliaryNode::TickNode(class UBehaviorTreeComponent* OwnerComp, uint8* NodeMemory, float DeltaSeconds)
+{
+	// empty in base class
+}
+
+uint16 UBTAuxiliaryNode::GetSpecialMemorySize() const
+{
+	return bTickIntervals ? sizeof(FBTAuxiliaryMemory) : Super::GetSpecialMemorySize();
+}
diff --git a/Engine/Source/Runtime/AIModule/Private/BehaviorTree/BTCompositeNode.cpp b/Engine/Source/Runtime/AIModule/Private/BehaviorTree/BTCompositeNode.cpp
new file mode 100644
--- /dev/null
+++ b/Engine/Source/Runtime/AIModule/Private/BehaviorTree/BTCompositeNode.cpp
@@ -0,0 +1,558 @@
+// Copyright 1998-2014 Epic Games, Inc. All Rights Reserved.
+
+#include "AIModulePrivate.h"
+#include "BehaviorTree/BTCompositeNode.h"
+
+UBTCompositeNode::UBTCompositeNode(const class FPostConstructInitializeProperties& PCIP) : Super(PCIP)
+{
+	NodeName = "UnknownComposite";
+	bUseChildExecutionNotify = false;
+	bUseNodeActivationNotify = false;
+	bUseNodeDeactivationNotify = false;
+	OptionalDeactivationResult = EBTNodeResult::Optional;
+}
+
+void UBTCompositeNode::InitializeComposite(uint16 InLastExecutionIndex)
+{
+	LastExecutionIndex = InLastExecutionIndex;
+}
+
+int32 UBTCompositeNode::FindChildToExecute(struct FBehaviorTreeSearchData& SearchData, EBTNodeResult::Type& LastResult) const
+{
+	FBTCompositeMemory* NodeMemory = GetNodeMemory<FBTCompositeMemory>(SearchData);
+	int32 RetIdx = BTSpecialChild::ReturnToParent;
+
+	if (Children.Num())
+	{
+		for (int32 ChildIdx = GetNextChild(SearchData, NodeMemory->CurrentChild, LastResult);
+			ChildIdx >= 0;
+			ChildIdx = GetNextChild(SearchData, ChildIdx, LastResult))
+		{
+			// check decorators
+			if (DoDecoratorsAllowExecution(SearchData.OwnerComp, SearchData.OwnerComp->ActiveInstanceIdx, ChildIdx))
+			{
+				OnChildActivation(SearchData, ChildIdx);
+				RetIdx = ChildIdx;
+				break;
+			}
+			else
+			{
+				LastResult = EBTNodeResult::Failed;
+				NotifyDecoratorsOnFailedActivation(SearchData, ChildIdx, LastResult);
+			}
+		}
+	}
+
+	return RetIdx;
+}
+
+int32 UBTCompositeNode::GetChildIndex(struct FBehaviorTreeSearchData& SearchData, const class UBTNode* ChildNode) const
+{
+	if (ChildNode->GetParentNode() != this)
+	{
+		FBTCompositeMemory* NodeMemory = GetNodeMemory<FBTCompositeMemory>(SearchData);
+		return NodeMemory->CurrentChild;		
+	}
+
+	return GetChildIndex(ChildNode);
+}
+
+int32 UBTCompositeNode::GetChildIndex(const class UBTNode* ChildNode) const
+{
+	for (int32 i = 0; i < Children.Num(); i++)
+	{
+		if (Children[i].ChildComposite == ChildNode ||
+			Children[i].ChildTask == ChildNode)
+		{
+			return i;
+		}
+	}
+
+	return BTSpecialChild::ReturnToParent;
+}
+
+void UBTCompositeNode::OnChildActivation(struct FBehaviorTreeSearchData& SearchData, const class UBTNode* ChildNode) const
+{
+	OnChildActivation(SearchData, GetChildIndex(SearchData, ChildNode));
+}
+
+void UBTCompositeNode::OnChildActivation(struct FBehaviorTreeSearchData& SearchData, int32 ChildIndex) const
+{
+	const FBTCompositeChild& ChildInfo = Children[ChildIndex];
+	FBTCompositeMemory* NodeMemory = GetNodeMemory<FBTCompositeMemory>(SearchData);
+
+	// pass to decorators before changing current child in node memory
+	// so they can access previously executed one if needed
+	NotifyDecoratorsOnActivation(SearchData, ChildIndex);
+
+	// pass to child composite
+	if (ChildInfo.ChildComposite)
+	{
+		ChildInfo.ChildComposite->OnNodeActivation(SearchData);
+	}
+
+	// update active node in current context: child node
+	NodeMemory->CurrentChild = ChildIndex;
+}
+
+void UBTCompositeNode::OnChildDeactivation(struct FBehaviorTreeSearchData& SearchData, const class UBTNode* ChildNode, EBTNodeResult::Type& NodeResult) const
+{
+	OnChildDeactivation(SearchData, GetChildIndex(SearchData, ChildNode), NodeResult);
+}
+
+void UBTCompositeNode::OnChildDeactivation(struct FBehaviorTreeSearchData& SearchData, int32 ChildIndex, EBTNodeResult::Type& NodeResult) const
+{
+	const FBTCompositeChild& ChildInfo = Children[ChildIndex];
+
+	// pass to decorators, unless it's from aborting, which will be handled separately
+	// (needs to remove all pending aux nodes from all branches, not only closest children)
+	if (NodeResult != EBTNodeResult::Aborted)
+	{
+		NotifyDecoratorsOnDeactivation(SearchData, ChildIndex, NodeResult);
+	}
+
+	// pass to child composite
+	if (ChildInfo.ChildComposite)
+	{
+		ChildInfo.ChildComposite->OnNodeDeactivation(SearchData, NodeResult);
+	}
+}
+
+void UBTCompositeNode::OnNodeActivation(struct FBehaviorTreeSearchData& SearchData) const
+{
+	OnNodeRestart(SearchData);
+
+	if (bUseNodeActivationNotify)
+	{
+		NotifyNodeActivation(SearchData);
+	}
+
+	// add services when execution flow enters this composite
+	for (int32 i = 0; i < Services.Num(); i++)
+	{
+		SearchData.AddUniqueUpdate(FBehaviorTreeSearchUpdate(Services[i], SearchData.OwnerComp->GetActiveInstanceIdx(), EBTNodeUpdateMode::Add));
+	}
+}
+
+void UBTCompositeNode::OnNodeDeactivation(struct FBehaviorTreeSearchData& SearchData, EBTNodeResult::Type& NodeResult) const
+{
+	// replace optional result before calling handler
+	if (NodeResult == EBTNodeResult::Optional)
+	{
+		NodeResult = OptionalDeactivationResult;
+	}
+
+	if (bUseNodeDeactivationNotify)
+	{
+		NotifyNodeDeactivation(SearchData, NodeResult);
+	}
+
+	// remove all services if execution flow leaves this composite
+	for (int32 i = 0; i < Services.Num(); i++)
+	{
+		SearchData.AddUniqueUpdate(FBehaviorTreeSearchUpdate(Services[i], SearchData.OwnerComp->GetActiveInstanceIdx(), EBTNodeUpdateMode::Remove));
+	}
+}
+
+void UBTCompositeNode::OnNodeRestart(struct FBehaviorTreeSearchData& SearchData) const
+{
+	FBTCompositeMemory* NodeMemory = GetNodeMemory<FBTCompositeMemory>(SearchData);
+	NodeMemory->CurrentChild = BTSpecialChild::NotInitialized;
+	NodeMemory->OverrideChild = BTSpecialChild::NotInitialized;
+}
+
+void UBTCompositeNode::NotifyDecoratorsOnActivation(struct FBehaviorTreeSearchData& SearchData, int32 ChildIdx) const
+{
+	const FBTCompositeChild& ChildInfo = Children[ChildIdx];
+	for (int32 i = 0; i < ChildInfo.Decorators.Num(); i++)
+	{
+		const UBTDecorator* DecoratorOb = ChildInfo.Decorators[i];
+		DecoratorOb->WrappedOnNodeActivation(SearchData);
+
+		switch (DecoratorOb->GetFlowAbortMode())
+		{
+			case EBTFlowAbortMode::LowerPriority:
+				SearchData.AddUniqueUpdate(FBehaviorTreeSearchUpdate(DecoratorOb, SearchData.OwnerComp->GetActiveInstanceIdx(), EBTNodeUpdateMode::Remove));
+				break;
+
+			case EBTFlowAbortMode::Self:
+			case EBTFlowAbortMode::Both:
+				SearchData.AddUniqueUpdate(FBehaviorTreeSearchUpdate(DecoratorOb, SearchData.OwnerComp->GetActiveInstanceIdx(), EBTNodeUpdateMode::Add));
+				break;
+
+			default:
+				break;
+		}
+	}
+}
+
+void UBTCompositeNode::NotifyDecoratorsOnDeactivation(struct FBehaviorTreeSearchData& SearchData, int32 ChildIdx, EBTNodeResult::Type& NodeResult) const
+{
+	const FBTCompositeChild& ChildInfo = Children[ChildIdx];
+	for (int32 i = 0; i < ChildInfo.Decorators.Num(); i++)
+	{
+		const UBTDecorator* DecoratorOb = ChildInfo.Decorators[i];
+		DecoratorOb->WrappedOnNodeProcessed(SearchData, NodeResult);
+		DecoratorOb->WrappedOnNodeDeactivation(SearchData, NodeResult);
+
+		// leaving child branch: 
+		if (DecoratorOb->GetFlowAbortMode() == EBTFlowAbortMode::Self)
+		{
+			// - observers with mode "Self" are now out of scope, remove them
+			SearchData.AddUniqueUpdate(FBehaviorTreeSearchUpdate(DecoratorOb, SearchData.OwnerComp->GetActiveInstanceIdx(), EBTNodeUpdateMode::Remove));
+		}
+		else if (DecoratorOb->GetFlowAbortMode() == EBTFlowAbortMode::LowerPriority)
+		{
+			// - observers with mode "Lower Priority" will try to reactivate themselves ("Both" is not removed on node activation)
+			SearchData.AddUniqueUpdate(FBehaviorTreeSearchUpdate(DecoratorOb, SearchData.OwnerComp->GetActiveInstanceIdx(), EBTNodeUpdateMode::AddForLowerPri));
+		}
+	}
+}
+
+void UBTCompositeNode::NotifyDecoratorsOnFailedActivation(struct FBehaviorTreeSearchData& SearchData, int32 ChildIdx, EBTNodeResult::Type& NodeResult) const
+{
+	const FBTCompositeChild& ChildInfo = Children[ChildIdx];
+	const uint16 ActiveInstanceIdx = SearchData.OwnerComp->GetActiveInstanceIdx();
+
+	for (int32 i = 0; i < ChildInfo.Decorators.Num(); i++)
+	{
+		const UBTDecorator* DecoratorOb = ChildInfo.Decorators[i];
+		DecoratorOb->WrappedOnNodeProcessed(SearchData, NodeResult);
+
+		if (DecoratorOb->GetFlowAbortMode() == EBTFlowAbortMode::LowerPriority ||
+			DecoratorOb->GetFlowAbortMode() == EBTFlowAbortMode::Both)
+		{
+			SearchData.AddUniqueUpdate(FBehaviorTreeSearchUpdate(DecoratorOb, ActiveInstanceIdx, EBTNodeUpdateMode::Add));
+		}
+	}
+
+	SearchData.OwnerComp->StoreDebuggerSearchStep(GetChildNode(ChildIdx), ActiveInstanceIdx, NodeResult);
+}
+
+void UBTCompositeNode::NotifyChildExecution(class UBehaviorTreeComponent* OwnerComp, uint8* NodeMemory, int32 ChildIdx, EBTNodeResult::Type& NodeResult) const
+{
+}
+
+void UBTCompositeNode::NotifyNodeActivation(struct FBehaviorTreeSearchData& SearchData) const
+{
+}
+
+void UBTCompositeNode::NotifyNodeDeactivation(struct FBehaviorTreeSearchData& SearchData, EBTNodeResult::Type& NodeResult) const
+{
+}
+
+void UBTCompositeNode::ConditionalNotifyChildExecution(class UBehaviorTreeComponent* OwnerComp, uint8* NodeMemory, const class UBTNode* ChildNode, EBTNodeResult::Type& NodeResult) const
+{
+	if (bUseChildExecutionNotify)
+	{
+		for (int32 i = 0; i < Children.Num(); i++)
+		{
+			if (Children[i].ChildComposite == ChildNode || Children[i].ChildTask == ChildNode)
+			{
+				NotifyChildExecution(OwnerComp, NodeMemory, i, NodeResult);
+				break;
+			}
+		}
+	}
+}
+
+static bool IsLogicOp(const FBTDecoratorLogic& Info)
+{
+	return (Info.Operation != EBTDecoratorLogic::Test) && (Info.Operation != EBTDecoratorLogic::Invalid);
+}
+
+static FString DescribeLogicOp(const TEnumAsByte<EBTDecoratorLogic::Type>& Op)
+{
+	static FString LogicDesc[] = { TEXT("Invalid"), TEXT("Test"), TEXT("AND"), TEXT("OR"), TEXT("NOT") };
+	return LogicDesc[Op];
+}
+
+struct FOperationStackInfo
+{
+	uint16 NumLeft;
+	TEnumAsByte<EBTDecoratorLogic::Type> Op;
+	uint8 bHasForcedResult : 1;
+	uint8 bForcedResult : 1;
+
+	FOperationStackInfo() {}
+	FOperationStackInfo(const FBTDecoratorLogic& DecoratorOp) :
+		NumLeft(DecoratorOp.Number), Op(DecoratorOp.Operation), bHasForcedResult(0) {};
+};
+
+static bool UpdateOperationStack(const class UBehaviorTreeComponent* OwnerComp, FString& Indent,
+								 TArray<FOperationStackInfo>& Stack, bool bTestResult,
+								 int32& FailedDecoratorIdx, int32& NodeDecoratorIdx, bool& bShouldStoreNodeIndex)
+{
+	if (Stack.Num() == 0)
+	{
+		return bTestResult;
+	}
+
+	FOperationStackInfo& CurrentOp = Stack.Last();
+	CurrentOp.NumLeft--;
+
+	if (CurrentOp.Op == EBTDecoratorLogic::And)
+	{
+		if (!CurrentOp.bHasForcedResult && !bTestResult)
+		{
+			CurrentOp.bHasForcedResult = true;
+			CurrentOp.bForcedResult = bTestResult;
+		}
+	}
+	else if (CurrentOp.Op == EBTDecoratorLogic::Or)
+	{
+		if (!CurrentOp.bHasForcedResult && bTestResult)
+		{
+			CurrentOp.bHasForcedResult = true;
+			CurrentOp.bForcedResult = bTestResult;
+		}	
+	}
+	else if (CurrentOp.Op == EBTDecoratorLogic::Not)
+	{
+		bTestResult = !bTestResult;
+	}
+
+	// update debugger while processing top level stack
+	if (Stack.Num() == 1)
+	{
+		// reset node flag and grab next decorator index
+		bShouldStoreNodeIndex = true;
+
+		// store first failed node
+		if (!bTestResult && FailedDecoratorIdx == INDEX_NONE)
+		{
+			FailedDecoratorIdx = NodeDecoratorIdx;
+		}
+	}
+
+	if (CurrentOp.bHasForcedResult)
+	{
+		bTestResult = CurrentOp.bForcedResult;
+	}
+
+	if (CurrentOp.NumLeft == 0)
+	{
+		UE_VLOG(OwnerComp->GetOwner(), LogBehaviorTree, Verbose, TEXT("%s%s finished: %s"), *Indent,
+			*DescribeLogicOp(CurrentOp.Op),
+			bTestResult ? TEXT("allowed") : TEXT("forbidden"));
+		Indent = Indent.LeftChop(2);
+
+		Stack.RemoveAt(Stack.Num() - 1);
+		return UpdateOperationStack(OwnerComp, Indent, Stack, bTestResult, FailedDecoratorIdx, NodeDecoratorIdx, bShouldStoreNodeIndex);
+	}
+
+	return bTestResult;
+}
+
+bool UBTCompositeNode::DoDecoratorsAllowExecution(class UBehaviorTreeComponent* OwnerComp, int32 InstanceIdx, int32 ChildIdx) const
+{
+	const FBTCompositeChild& ChildInfo = Children[ChildIdx];
+	bool bResult = true;
+
+	if (ChildInfo.Decorators.Num() == 0)
+	{
+		return bResult;
+	}
+
+	FBehaviorTreeInstance& MyInstance = OwnerComp->InstanceStack[InstanceIdx];
+
+	if (ChildInfo.DecoratorOps.Num() == 0)
+	{
+		// simple check: all decorators must agree
+		for (int32 i = 0; i < ChildInfo.Decorators.Num(); i++)
+		{
+			const UBTDecorator* TestDecorator = ChildInfo.Decorators[i];
+			const bool bIsAllowed = TestDecorator ? TestDecorator->WrappedCanExecute(OwnerComp, TestDecorator->GetNodeMemory<uint8>(MyInstance)) : false;
+			OwnerComp->StoreDebuggerSearchStep(TestDecorator, InstanceIdx, bIsAllowed);
+
+			if (!bIsAllowed)
+			{
+				UE_VLOG(OwnerComp->GetOwner(), LogBehaviorTree, Verbose, TEXT("Child[%d] execution forbidden by %s"),
+					ChildIdx, *UBehaviorTreeTypes::DescribeNodeHelper(TestDecorator));
+
+				bResult = false;
+				break;
+			}
+		}
+	}
+	else
+	{
+		// advanced check: follow decorator logic operations (composite decorator on child link)
+		UE_VLOG(OwnerComp->GetOwner(), LogBehaviorTree, Verbose, TEXT("Child[%d] execution test with logic operations"), ChildIdx);
+
+		TArray<FOperationStackInfo> OperationStack;
+		FString Indent;
+
+		// debugger data collection:
+		// - get index of each decorator from main AND test, they will match graph nodes
+		// - if first operator is not AND it means, that there's only single composite decorator on line
+		// - while updating top level stack, grab index of first failed node
+
+		int32 NodeDecoratorIdx = INDEX_NONE;
+		int32 FailedDecoratorIdx = INDEX_NONE;
+		bool bShouldStoreNodeIndex = true;
+
+		for (int32 i = 0; i < ChildInfo.DecoratorOps.Num(); i++)
+		{
+			const FBTDecoratorLogic& DecoratorOp = ChildInfo.DecoratorOps[i];
+			if (IsLogicOp(DecoratorOp))
+			{
+				OperationStack.Add(FOperationStackInfo(DecoratorOp));
+				Indent += TEXT("  ");
+				UE_VLOG(OwnerComp->GetOwner(), LogBehaviorTree, Verbose, TEXT("%spushed %s:%d"), *Indent,
+					*DescribeLogicOp(DecoratorOp.Operation), DecoratorOp.Number);
+			}
+			else if (DecoratorOp.Operation == EBTDecoratorLogic::Test)
+			{
+				const bool bHasOverride = OperationStack.Num() ? OperationStack.Last().bHasForcedResult : false;
+				const bool bCurrentOverride = OperationStack.Num() ? OperationStack.Last().bForcedResult : false;
+
+				// debugger: store first decorator of graph node
+				if (bShouldStoreNodeIndex)
+				{
+					bShouldStoreNodeIndex = false;
+					NodeDecoratorIdx = DecoratorOp.Number;
+				}
+
+				UBTDecorator* TestDecorator = ChildInfo.Decorators[DecoratorOp.Number];
+				const bool bIsAllowed = bHasOverride ? bCurrentOverride : TestDecorator->WrappedCanExecute(OwnerComp, TestDecorator->GetNodeMemory<uint8>(MyInstance));
+				UE_VLOG(OwnerComp->GetOwner(), LogBehaviorTree, Verbose, TEXT("%s%s %s: %s"), *Indent,
+					bHasOverride ? TEXT("skipping") : TEXT("testing"),
+					*UBehaviorTreeTypes::DescribeNodeHelper(TestDecorator),
+					bIsAllowed ? TEXT("allowed") : TEXT("forbidden"));
+
+				bResult = UpdateOperationStack(OwnerComp, Indent, OperationStack, bIsAllowed, FailedDecoratorIdx, NodeDecoratorIdx, bShouldStoreNodeIndex);
+				if (OperationStack.Num() == 0)
+				{
+					UE_VLOG(OwnerComp->GetOwner(), LogBehaviorTree, Verbose, TEXT("finished execution test: %s"),
+						bResult ? TEXT("allowed") : TEXT("forbidden"));
+
+					OwnerComp->StoreDebuggerSearchStep(ChildInfo.Decorators[FMath::Max(0, FailedDecoratorIdx)], InstanceIdx, bResult);
+					break;
+				}
+			}
+		}
+	}
+
+	return bResult;
+}
+
+int32 UBTCompositeNode::GetMatchingChildIndex(int32 ActiveInstanceIdx, struct FBTNodeIndex& NodeIdx) const
+{
+	const int32 OutsideRange = BTSpecialChild::ReturnToParent;
+	const int32 UnlimitedRange = Children.Num() - 1;
+
+	// search ends at the same instance level: use execution index to determine branch containing node index
+	if (ActiveInstanceIdx == NodeIdx.InstanceIndex)
+	{
+		// is composite even in range of search?
+		if (GetExecutionIndex() > NodeIdx.ExecutionIndex)
+		{
+			return OutsideRange;
+		}
+
+		// find child outside range
+		for (int32 i = 0; i < Children.Num(); i++)
+		{
+			const UBTNode* ChildNode = GetChildNode(i);
+			if (ChildNode->GetExecutionIndex() > NodeIdx.ExecutionIndex)
+			{
+				return i ? (i - 1) : 0;
+			}
+		}
+
+		return UnlimitedRange;
+	}
+
+	// search ends at higher level: allow every node
+	// search ends at lower level: outside allowed range
+	return (ActiveInstanceIdx > NodeIdx.InstanceIndex) ? UnlimitedRange : OutsideRange;
+}
+
+int32 UBTCompositeNode::GetNextChild(struct FBehaviorTreeSearchData& SearchData, int32 LastChildIdx, EBTNodeResult::Type LastResult) const
+{
+	FBTCompositeMemory* NodeMemory = GetNodeMemory<FBTCompositeMemory>(SearchData);
+	int32 NextChildIndex = BTSpecialChild::ReturnToParent;
+	uint16 ActiveInstanceIdx = SearchData.OwnerComp->GetActiveInstanceIdx();
+
+	// newly activated node, search range not reached yet: select search branch for decorator test
+	if (LastChildIdx == BTSpecialChild::NotInitialized && SearchData.SearchStart.IsSet() &&
+		FBTNodeIndex(ActiveInstanceIdx, GetExecutionIndex()).TakesPriorityOver(SearchData.SearchStart))
+	{
+		NextChildIndex = GetMatchingChildIndex(ActiveInstanceIdx, SearchData.SearchStart);
+	}
+	else if (NodeMemory->OverrideChild != BTSpecialChild::NotInitialized && !SearchData.OwnerComp->IsRestartPending())
+	{
+		NextChildIndex = NodeMemory->OverrideChild;
+		NodeMemory->OverrideChild = BTSpecialChild::NotInitialized;
+	}
+	// or use composite's logic
+	else if (OnNextChild.IsBound())
+	{
+		NextChildIndex = OnNextChild.Execute(SearchData, LastChildIdx, LastResult);
+	}
+
+	return NextChildIndex;
+}
+
+void UBTCompositeNode::SetChildOverride(FBehaviorTreeSearchData& SearchData, int8 Index) const
+{
+	if (Children.IsValidIndex(Index) || Index == BTSpecialChild::ReturnToParent)
+	{
+		FBTCompositeMemory* MyMemory = GetNodeMemory<FBTCompositeMemory>(SearchData);
+		MyMemory->OverrideChild = Index;
+	}
+}
+
+void UBTCompositeNode::RequestDelayedExecution(class UBehaviorTreeComponent* OwnerComp, EBTNodeResult::Type LastResult) const
+{
+	OwnerComp->RequestExecution(LastResult);
+}
+
+uint16 UBTCompositeNode::GetChildExecutionIndex(int32 Index, EBTChildIndex::Type ChildMode) const
+{
+	const UBTNode* ChildNode = GetChildNode(Index);
+	if (ChildNode)
+	{
+		const int32 Offset = (ChildMode == EBTChildIndex::FirstNode) ? Children[Index].Decorators.Num() : 0;
+		return ChildNode->GetExecutionIndex() - Offset;
+	}
+	
+	return (LastExecutionIndex + 1);
+}
+
+bool UBTCompositeNode::CanPushSubtree(class UBehaviorTreeComponent* OwnerComp, uint8* NodeMemory, int32 ChildIdx) const
+{
+	return true;
+}
+
+void UBTCompositeNode::DescribeRuntimeValues(const class UBehaviorTreeComponent* OwnerComp, uint8* NodeMemory, EBTDescriptionVerbosity::Type Verbosity, TArray<FString>& Values) const
+{
+	Super::DescribeRuntimeValues(OwnerComp, NodeMemory, Verbosity, Values);
+
+	if (Verbosity == EBTDescriptionVerbosity::Detailed)
+	{
+		FBTCompositeMemory* MyMemory = (FBTCompositeMemory*)NodeMemory;
+		Values.Add(FString::Printf(TEXT("current child: %d"), MyMemory->CurrentChild));
+		Values.Add(FString::Printf(TEXT("override child: %d"), MyMemory->OverrideChild));
+	}
+}
+
+#if WITH_EDITOR
+
+bool UBTCompositeNode::CanAbortSelf() const
+{
+	return true;
+}
+
+bool UBTCompositeNode::CanAbortLowerPriority() const
+{
+	return true;
+}
+
+#endif
+
+uint16 UBTCompositeNode::GetInstanceMemorySize() const
+{
+	return sizeof(FBTCompositeMemory);
+}
diff --git a/Engine/Source/Runtime/AIModule/Private/BehaviorTree/BTFunctionLibrary.cpp b/Engine/Source/Runtime/AIModule/Private/BehaviorTree/BTFunctionLibrary.cpp
new file mode 100644
--- /dev/null
+++ b/Engine/Source/Runtime/AIModule/Private/BehaviorTree/BTFunctionLibrary.cpp
@@ -0,0 +1,177 @@
+// Copyright 1998-2014 Epic Games, Inc. All Rights Reserved.
+
+#include "AIModulePrivate.h"
+#include "BehaviorTree/BTFunctionLibrary.h"
+#include "BehaviorTree/BlackboardComponent.h"
+#include "BehaviorTree/Blackboard/BlackboardKeyAllTypes.h"
+#include "BehaviorTree/BTFunctionLibrary.h"
+
+UBTFunctionLibrary::UBTFunctionLibrary(const class FPostConstructInitializeProperties& PCIP) : Super(PCIP)
+{
+}
+
+UBlackboardComponent* UBTFunctionLibrary::GetBlackboard(UBTNode* NodeOwner)
+{
+	UBehaviorTreeComponent* BTComp = NodeOwner ? Cast<UBehaviorTreeComponent>(NodeOwner->GetOuter()) : NULL;
+	return BTComp ? BTComp->GetBlackboardComponent() : NULL;
+}
+
+UObject* UBTFunctionLibrary::GetBlackboardValueAsObject(UBTNode* NodeOwner, const struct FBlackboardKeySelector& Key)
+{
+	const UBehaviorTreeComponent* BTComp = NodeOwner ? Cast<UBehaviorTreeComponent>(NodeOwner->GetOuter()) : NULL;
+	return BTComp ? BTComp->GetBlackboardComponent()->GetValueAsObject(Key.SelectedKeyName) : NULL;
+}
+
+AActor* UBTFunctionLibrary::GetBlackboardValueAsActor(UBTNode* NodeOwner, const struct FBlackboardKeySelector& Key)
+{
+	const UBehaviorTreeComponent* BTComp = NodeOwner ? Cast<UBehaviorTreeComponent>(NodeOwner->GetOuter()) : NULL;
+	return BTComp ? Cast<AActor>(BTComp->GetBlackboardComponent()->GetValueAsObject(Key.SelectedKeyName)) : NULL;
+}
+
+UClass* UBTFunctionLibrary::GetBlackboardValueAsClass(UBTNode* NodeOwner, const struct FBlackboardKeySelector& Key)
+{
+	const UBehaviorTreeComponent* BTComp = NodeOwner ? Cast<UBehaviorTreeComponent>(NodeOwner->GetOuter()) : NULL;
+	return BTComp ? BTComp->GetBlackboardComponent()->GetValueAsClass(Key.SelectedKeyName) : NULL;
+}
+
+uint8 UBTFunctionLibrary::GetBlackboardValueAsEnum(UBTNode* NodeOwner, const struct FBlackboardKeySelector& Key)
+{
+	const UBehaviorTreeComponent* BTComp = NodeOwner ? Cast<UBehaviorTreeComponent>(NodeOwner->GetOuter()) : NULL;
+	return BTComp ? BTComp->GetBlackboardComponent()->GetValueAsEnum(Key.SelectedKeyName) : 0;
+}
+
+int32 UBTFunctionLibrary::GetBlackboardValueAsInt(UBTNode* NodeOwner, const struct FBlackboardKeySelector& Key)
+{
+	const UBehaviorTreeComponent* BTComp = NodeOwner ? Cast<UBehaviorTreeComponent>(NodeOwner->GetOuter()) : NULL;
+	return BTComp ? BTComp->GetBlackboardComponent()->GetValueAsInt(Key.SelectedKeyName) : 0;
+}
+
+float UBTFunctionLibrary::GetBlackboardValueAsFloat(UBTNode* NodeOwner, const struct FBlackboardKeySelector& Key)
+{
+	const UBehaviorTreeComponent* BTComp = NodeOwner ? Cast<UBehaviorTreeComponent>(NodeOwner->GetOuter()) : NULL;
+	return BTComp ? BTComp->GetBlackboardComponent()->GetValueAsFloat(Key.SelectedKeyName) : 0.0f;
+}
+
+bool UBTFunctionLibrary::GetBlackboardValueAsBool(UBTNode* NodeOwner, const struct FBlackboardKeySelector& Key)
+{
+	const UBehaviorTreeComponent* BTComp = NodeOwner ? Cast<UBehaviorTreeComponent>(NodeOwner->GetOuter()) : NULL;
+	return BTComp ? BTComp->GetBlackboardComponent()->GetValueAsBool(Key.SelectedKeyName) : false;
+}
+
+FString UBTFunctionLibrary::GetBlackboardValueAsString(UBTNode* NodeOwner, const struct FBlackboardKeySelector& Key)
+{
+	const UBehaviorTreeComponent* BTComp = NodeOwner ? Cast<UBehaviorTreeComponent>(NodeOwner->GetOuter()) : NULL;
+	return BTComp ? BTComp->GetBlackboardComponent()->GetValueAsString(Key.SelectedKeyName) : FString();
+}
+
+FName UBTFunctionLibrary::GetBlackboardValueAsName(UBTNode* NodeOwner, const struct FBlackboardKeySelector& Key)
+{
+	const UBehaviorTreeComponent* BTComp = NodeOwner ? Cast<UBehaviorTreeComponent>(NodeOwner->GetOuter()) : NULL;
+	return BTComp ? BTComp->GetBlackboardComponent()->GetValueAsName(Key.SelectedKeyName) : NAME_None;
+}
+
+FVector UBTFunctionLibrary::GetBlackboardValueAsVector(UBTNode* NodeOwner, const struct FBlackboardKeySelector& Key)
+{
+	const UBehaviorTreeComponent* BTComp = NodeOwner ? Cast<UBehaviorTreeComponent>(NodeOwner->GetOuter()) : NULL;
+	return BTComp ? BTComp->GetBlackboardComponent()->GetValueAsVector(Key.SelectedKeyName) : FVector::ZeroVector;
+}
+
+void UBTFunctionLibrary::SetBlackboardValueAsObject(UBTNode* NodeOwner, const struct FBlackboardKeySelector& Key, UObject* Value)
+{
+	UBehaviorTreeComponent* BTComp = NodeOwner ? Cast<UBehaviorTreeComponent>(NodeOwner->GetOuter()) : NULL;
+	if (BTComp != NULL)
+	{
+		BTComp->GetBlackboardComponent()->SetValueAsObject(Key.SelectedKeyName, Value);
+	}
+}
+
+void UBTFunctionLibrary::SetBlackboardValueAsClass(UBTNode* NodeOwner, const struct FBlackboardKeySelector& Key, UClass* Value)
+{
+	UBehaviorTreeComponent* BTComp = NodeOwner ? Cast<UBehaviorTreeComponent>(NodeOwner->GetOuter()) : NULL;
+	if (BTComp != NULL)
+	{
+		BTComp->GetBlackboardComponent()->SetValueAsClass(Key.SelectedKeyName, Value);
+	}
+}
+
+void UBTFunctionLibrary::SetBlackboardValueAsEnum(UBTNode* NodeOwner, const struct FBlackboardKeySelector& Key, uint8 Value)
+{
+	UBehaviorTreeComponent* BTComp = NodeOwner ? Cast<UBehaviorTreeComponent>(NodeOwner->GetOuter()) : NULL;
+	if (BTComp != NULL)
+	{
+		BTComp->GetBlackboardComponent()->SetValueAsEnum(Key.SelectedKeyName, Value);
+	}
+}
+
+void UBTFunctionLibrary::SetBlackboardValueAsInt(UBTNode* NodeOwner, const struct FBlackboardKeySelector& Key, int32 Value)
+{
+	UBehaviorTreeComponent* BTComp = NodeOwner ? Cast<UBehaviorTreeComponent>(NodeOwner->GetOuter()) : NULL;
+	if (BTComp != NULL)
+	{
+		BTComp->GetBlackboardComponent()->SetValueAsInt(Key.SelectedKeyName, Value);
+	}
+}
+
+void UBTFunctionLibrary::SetBlackboardValueAsFloat(UBTNode* NodeOwner, const struct FBlackboardKeySelector& Key, float Value)
+{
+	UBehaviorTreeComponent* BTComp = NodeOwner ? Cast<UBehaviorTreeComponent>(NodeOwner->GetOuter()) : NULL;
+	if (BTComp != NULL)
+	{
+		BTComp->GetBlackboardComponent()->SetValueAsFloat(Key.SelectedKeyName, Value);
+	}
+}
+
+void UBTFunctionLibrary::SetBlackboardValueAsBool(UBTNode* NodeOwner, const struct FBlackboardKeySelector& Key, bool Value)
+{
+	UBehaviorTreeComponent* BTComp = NodeOwner ? Cast<UBehaviorTreeComponent>(NodeOwner->GetOuter()) : NULL;
+	if (BTComp != NULL)
+	{
+		BTComp->GetBlackboardComponent()->SetValueAsBool(Key.SelectedKeyName, Value);
+	}
+}
+
+void UBTFunctionLibrary::SetBlackboardValueAsString(UBTNode* NodeOwner, const struct FBlackboardKeySelector& Key, FString Value)
+{
+	UBehaviorTreeComponent* BTComp = NodeOwner ? Cast<UBehaviorTreeComponent>(NodeOwner->GetOuter()) : NULL;
+	if (BTComp != NULL)
+	{
+		BTComp->GetBlackboardComponent()->SetValueAsString(Key.SelectedKeyName, Value);
+	}
+}
+
+void UBTFunctionLibrary::SetBlackboardValueAsName(UBTNode* NodeOwner, const struct FBlackboardKeySelector& Key, FName Value)
+{
+	UBehaviorTreeComponent* BTComp = NodeOwner ? Cast<UBehaviorTreeComponent>(NodeOwner->GetOuter()) : NULL;
+	if (BTComp != NULL)
+	{
+		BTComp->GetBlackboardComponent()->SetValueAsName(Key.SelectedKeyName, Value);
+	}
+}
+
+void UBTFunctionLibrary::SetBlackboardValueAsVector(UBTNode* NodeOwner, const struct FBlackboardKeySelector& Key, FVector Value)
+{
+	UBehaviorTreeComponent* BTComp = NodeOwner ? Cast<UBehaviorTreeComponent>(NodeOwner->GetOuter()) : NULL;
+	if (BTComp != NULL)
+	{
+		BTComp->GetBlackboardComponent()->SetValueAsVector(Key.SelectedKeyName, Value);
+	}
+}
+
+void UBTFunctionLibrary::ClearBlackboardValueAsVector(UBTNode* NodeOwner, const struct FBlackboardKeySelector& Key)
+{
+	UBehaviorTreeComponent* BTComp = NodeOwner ? Cast<UBehaviorTreeComponent>(NodeOwner->GetOuter()) : NULL;
+	if (BTComp != NULL)
+	{
+		BTComp->GetBlackboardComponent()->ClearValueAsVector(Key.SelectedKeyName);
+	}
+}
+
+void UBTFunctionLibrary::StartUsingExternalEvent(UBTNode* NodeOwner, AActor* OwningActor)
+{
+	// deprecated, not removed yet
+}
+
+void UBTFunctionLibrary::StopUsingExternalEvent(UBTNode* NodeOwner)
+{
+	// deprecated, not removed yet
+}
