Extracting AI code from the engine into a separate module, pass 1 #UE4

diff --git a/Engine/Source/Runtime/AIModule/Private/BehaviorTree/BTNode.cpp b/Engine/Source/Runtime/AIModule/Private/BehaviorTree/BTNode.cpp
new file mode 100644
--- /dev/null
+++ b/Engine/Source/Runtime/AIModule/Private/BehaviorTree/BTNode.cpp
@@ -0,0 +1,159 @@
+// Copyright 1998-2014 Epic Games, Inc. All Rights Reserved.
+
+#include "AIModulePrivate.h"
+#include "BehaviorTree/BehaviorTree.h"
+#include "BehaviorTree/BTNode.h"
+#include "BehaviorTree/BTCompositeNode.h"
+
+//----------------------------------------------------------------------//
+// UBTNode
+//----------------------------------------------------------------------//
+
+UBTNode::UBTNode(const class FPostConstructInitializeProperties& PCIP) : Super(PCIP)
+{
+	NodeName = "";
+	ParentNode = NULL;
+	TreeAsset = NULL;
+	ExecutionIndex = 0;
+	MemoryOffset = 0;
+	TreeDepth = 0;
+	bCreateNodeInstance = false;
+	bIsInstanced = false;
+	bIsInjected = false;
+
+#if USE_BEHAVIORTREE_DEBUGGER
+	NextExecutionNode = NULL;
+#endif
+}
+
+UWorld* UBTNode::GetWorld() const
+{
+	// instanced nodes are created for behavior tree component owning that instance
+	// template nodes are created for behavior tree manager, which is located directly in UWorld
+
+	return GetOuter() == NULL ? NULL :
+		IsInstanced() ? (Cast<UBehaviorTreeComponent>(GetOuter()))->GetWorld() :
+		Cast<UWorld>(GetOuter()->GetOuter());
+}
+
+void UBTNode::InitializeNode(class UBTCompositeNode* InParentNode, uint16 InExecutionIndex, uint16 InMemoryOffset, uint8 InTreeDepth)
+{
+	ParentNode = InParentNode;
+	ExecutionIndex = InExecutionIndex;
+	MemoryOffset = InMemoryOffset;
+	TreeDepth = InTreeDepth;
+}
+
+void UBTNode::InitializeMemory(class UBehaviorTreeComponent* OwnerComp, uint8* NodeMemory) const
+{
+	// empty in base class 
+}
+
+void UBTNode::OnInstanceCreated(class UBehaviorTreeComponent* OwnerComp)
+{
+	// empty in base class
+}
+
+void UBTNode::OnInstanceDestroyed(class UBehaviorTreeComponent* OwnerComp)
+{
+	// empty in base class
+}
+
+void UBTNode::InitializeForInstance(class UBehaviorTreeComponent* OwnerComp, uint8* NodeMemory, int32& NextInstancedIndex) const
+{
+	if (bCreateNodeInstance)
+	{
+		// composite nodes can't be instanced!
+		check(IsA(UBTCompositeNode::StaticClass()) == false);
+
+		UBTNode* NodeInstance = OwnerComp->NodeInstances.IsValidIndex(NextInstancedIndex) ? OwnerComp->NodeInstances[NextInstancedIndex] : NULL;
+		if (NodeInstance == NULL)
+		{
+			NodeInstance = DuplicateObject<UBTNode>(this, OwnerComp);
+			NodeInstance->InitializeNode(GetParentNode(), GetExecutionIndex(), GetMemoryOffset(), GetTreeDepth());
+			NodeInstance->bIsInstanced = true;
+
+			OwnerComp->NodeInstances.Add(NodeInstance);
+		}
+
+		FBTInstancedNodeMemory* MyMemory = GetSpecialNodeMemory<FBTInstancedNodeMemory>(NodeMemory);
+		MyMemory->NodeIdx = NextInstancedIndex;
+
+		NodeInstance->OnInstanceCreated(OwnerComp);
+		NextInstancedIndex++;
+	}
+	else
+	{
+		InitializeMemory(OwnerComp, NodeMemory);
+	}
+}
+
+#if USE_BEHAVIORTREE_DEBUGGER
+void UBTNode::InitializeExecutionOrder(class UBTNode* NextNode)
+{
+	NextExecutionNode = NextNode;
+}
+#endif
+
+void UBTNode::InitializeFromAsset(class UBehaviorTree* Asset)
+{
+	TreeAsset = Asset;
+}
+
+UBlackboardData* UBTNode::GetBlackboardAsset() const
+{
+	return TreeAsset ? TreeAsset->BlackboardAsset : NULL;
+}
+
+uint16 UBTNode::GetInstanceMemorySize() const
+{
+	return 0;
+}
+
+uint16 UBTNode::GetSpecialMemorySize() const
+{
+	return bCreateNodeInstance ? sizeof(FBTInstancedNodeMemory) : 0;
+}
+
+class UBTNode* UBTNode::GetNodeInstance(const class UBehaviorTreeComponent* OwnerComp, uint8* NodeMemory) const
+{
+	FBTInstancedNodeMemory* MyMemory = GetSpecialNodeMemory<FBTInstancedNodeMemory>(NodeMemory);
+	return OwnerComp && MyMemory && OwnerComp->NodeInstances.IsValidIndex(MyMemory->NodeIdx) ?
+		OwnerComp->NodeInstances[MyMemory->NodeIdx] : NULL;
+}
+
+class UBTNode* UBTNode::GetNodeInstance(struct FBehaviorTreeSearchData& SearchData) const
+{
+	return GetNodeInstance(SearchData.OwnerComp, GetNodeMemory<uint8>(SearchData));
+}
+
+FString UBTNode::GetRuntimeDescription(const class UBehaviorTreeComponent* OwnerComp, uint8* NodeMemory, EBTDescriptionVerbosity::Type Verbosity) const
+{
+	FString Description = NodeName.Len() ? FString::Printf(TEXT("%s [%s]"), *NodeName, *GetStaticDescription()) : GetStaticDescription();
+	TArray<FString> RuntimeValues;
+
+	const UBTNode* NodeOb = bCreateNodeInstance ? GetNodeInstance(OwnerComp, NodeMemory) : this;
+	if (NodeOb)
+	{
+		NodeOb->DescribeRuntimeValues(OwnerComp, NodeMemory, Verbosity, RuntimeValues);
+	}
+
+	for (int32 i = 0; i < RuntimeValues.Num(); i++)
+	{
+		Description += TEXT(", ");
+		Description += RuntimeValues[i];
+	}
+
+	return Description;
+}
+
+FString UBTNode::GetStaticDescription() const
+{
+	// short type name
+	return UBehaviorTreeTypes::GetShortTypeName(this);
+}
+
+void UBTNode::DescribeRuntimeValues(const class UBehaviorTreeComponent* OwnerComp, uint8* NodeMemory, EBTDescriptionVerbosity::Type Verbosity, TArray<FString>& Values) const
+{
+	// nothing stored in memory for base class
+}
\ No newline at end of file
diff --git a/Engine/Source/Runtime/AIModule/Private/BehaviorTree/BTService.cpp b/Engine/Source/Runtime/AIModule/Private/BehaviorTree/BTService.cpp
new file mode 100644
--- /dev/null
+++ b/Engine/Source/Runtime/AIModule/Private/BehaviorTree/BTService.cpp
@@ -0,0 +1,36 @@
+// Copyright 1998-2014 Epic Games, Inc. All Rights Reserved.
+
+#include "AIModulePrivate.h"
+#include "BehaviorTree/BTService.h"
+
+UBTService::UBTService(const class FPostConstructInitializeProperties& PCIP) : Super(PCIP)
+{
+	bNotifyBecomeRelevant = true;
+	bNotifyTick = true;
+	bTickIntervals = true;
+
+	Interval = 0.5f;
+	RandomDeviation = 0.1f;
+}
+
+void UBTService::TickNode(UBehaviorTreeComponent* OwnerComp, uint8* NodeMemory, float DeltaSeconds)
+{
+	const float NextTickTime = FMath::FRandRange(FMath::Max(0.0f, Interval - RandomDeviation), (Interval + RandomDeviation));
+	SetNextTickTime(NodeMemory, NextTickTime);
+}
+
+void UBTService::OnBecomeRelevant(UBehaviorTreeComponent* OwnerComp, uint8* NodeMemory)
+{
+	// force it, without checking delays in ConditionalTickNode
+	TickNode(OwnerComp, NodeMemory, 0.0f);
+}
+
+FString UBTService::GetStaticDescription() const
+{
+	FString IntervalDesc = (RandomDeviation > 0.0f) ?
+		FString::Printf(TEXT("%.2fs..%.2fs"), FMath::Max(0.0f, Interval - RandomDeviation), (Interval + RandomDeviation)) :
+		FString::Printf(TEXT("%.2fs"), Interval);
+
+	return FString::Printf(TEXT("%s: tick every %s"), *UBehaviorTreeTypes::GetShortTypeName(this), *IntervalDesc);
+}
+
diff --git a/Engine/Source/Runtime/AIModule/Private/BehaviorTree/BTTaskNode.cpp b/Engine/Source/Runtime/AIModule/Private/BehaviorTree/BTTaskNode.cpp
new file mode 100644
--- /dev/null
+++ b/Engine/Source/Runtime/AIModule/Private/BehaviorTree/BTTaskNode.cpp
@@ -0,0 +1,103 @@
+// Copyright 1998-2014 Epic Games, Inc. All Rights Reserved.
+
+#include "AIModulePrivate.h"
+#include "BehaviorTree/BTTaskNode.h"
+
+UBTTaskNode::UBTTaskNode(const class FPostConstructInitializeProperties& PCIP) : Super(PCIP)
+{
+	NodeName = "UnknownTask";
+	bNotifyTick = false;
+}
+
+EBTNodeResult::Type UBTTaskNode::ExecuteTask(class UBehaviorTreeComponent* OwnerComp, uint8* NodeMemory)
+{
+	return EBTNodeResult::Succeeded;
+}
+
+EBTNodeResult::Type UBTTaskNode::AbortTask(class UBehaviorTreeComponent* OwnerComp, uint8* NodeMemory)
+{
+	return EBTNodeResult::Aborted;
+}
+
+EBTNodeResult::Type UBTTaskNode::WrappedExecuteTask(class UBehaviorTreeComponent* OwnerComp, uint8* NodeMemory) const
+{
+	const UBTNode* NodeOb = bCreateNodeInstance ? GetNodeInstance(OwnerComp, NodeMemory) : this;
+	return NodeOb ? ((UBTTaskNode*)NodeOb)->ExecuteTask(OwnerComp, NodeMemory) : EBTNodeResult::Failed;
+}
+
+EBTNodeResult::Type UBTTaskNode::WrappedAbortTask(class UBehaviorTreeComponent* OwnerComp, uint8* NodeMemory) const
+{
+	const UBTNode* NodeOb = bCreateNodeInstance ? GetNodeInstance(OwnerComp, NodeMemory) : this;
+	return NodeOb ? ((UBTTaskNode*)NodeOb)->AbortTask(OwnerComp, NodeMemory) : EBTNodeResult::Aborted;
+}
+
+void UBTTaskNode::WrappedTickTask(class UBehaviorTreeComponent* OwnerComp, uint8* NodeMemory, float DeltaSeconds) const
+{
+	if (bNotifyTick)
+	{
+		const UBTNode* NodeOb = bCreateNodeInstance ? GetNodeInstance(OwnerComp, NodeMemory) : this;
+		if (NodeOb)
+		{
+			((UBTTaskNode*)NodeOb)->TickTask(OwnerComp, NodeMemory, DeltaSeconds);
+		}
+	}
+}
+
+void UBTTaskNode::ReceivedMessage(UBrainComponent* BrainComp, const struct FAIMessage& Message)
+{
+	UBehaviorTreeComponent* OwnerComp = (UBehaviorTreeComponent*)BrainComp;
+	uint16 InstanceIdx = OwnerComp->FindInstanceContainingNode(this);
+	uint8* NodeMemory = GetNodeMemory<uint8>(OwnerComp->InstanceStack[InstanceIdx]);
+
+	OnMessage(OwnerComp, NodeMemory, Message.MessageName, Message.RequestID, Message.Status == FAIMessage::Success);
+}
+
+void UBTTaskNode::TickTask(class UBehaviorTreeComponent* OwnerComp, uint8* NodeMemory, float DeltaSeconds)
+{
+	// empty in base class
+}
+
+void UBTTaskNode::OnMessage(class UBehaviorTreeComponent* OwnerComp, uint8* NodeMemory, FName Message, int32 RequestID, bool bSuccess)
+{
+	const EBTTaskStatus::Type Status = OwnerComp->GetTaskStatus(this);
+	if (Status == EBTTaskStatus::Active)
+	{
+		FinishLatentTask(OwnerComp, bSuccess ? EBTNodeResult::Succeeded : EBTNodeResult::Failed);
+	}
+	else if (Status == EBTTaskStatus::Aborting)
+	{
+		FinishLatentAbort(OwnerComp);
+	}
+}
+
+void UBTTaskNode::FinishLatentTask(class UBehaviorTreeComponent* OwnerComp, EBTNodeResult::Type TaskResult) const
+{
+	// OnTaskFinished must receive valid template node
+	UBTTaskNode* TemplateNode = (UBTTaskNode*)OwnerComp->FindTemplateNode(this);
+	OwnerComp->OnTaskFinished(TemplateNode, TaskResult);
+}
+
+void UBTTaskNode::FinishLatentAbort(class UBehaviorTreeComponent* OwnerComp) const
+{
+	// OnTaskFinished must receive valid template node
+	UBTTaskNode* TemplateNode = (UBTTaskNode*)OwnerComp->FindTemplateNode(this);
+	OwnerComp->OnTaskFinished(TemplateNode, EBTNodeResult::Aborted);
+}
+
+void UBTTaskNode::WaitForMessage(class UBehaviorTreeComponent* OwnerComp, FName MessageType) const
+{
+	// messages delegates should be called on node instances (if they exists)
+	OwnerComp->RegisterMessageObserver(this, MessageType);
+}
+
+void UBTTaskNode::WaitForMessage(class UBehaviorTreeComponent* OwnerComp, FName MessageType, int32 RequestID) const
+{
+	// messages delegates should be called on node instances (if they exists)
+	OwnerComp->RegisterMessageObserver(this, MessageType, RequestID);
+}
+	
+void UBTTaskNode::StopWaitingForMessages(class UBehaviorTreeComponent* OwnerComp) const
+{
+	// messages delegates should be called on node instances (if they exists)
+	OwnerComp->UnregisterMessageObserversFrom(this);
+}
diff --git a/Engine/Source/Runtime/AIModule/Private/BehaviorTree/BehaviorTree.cpp b/Engine/Source/Runtime/AIModule/Private/BehaviorTree/BehaviorTree.cpp
new file mode 100644
--- /dev/null
+++ b/Engine/Source/Runtime/AIModule/Private/BehaviorTree/BehaviorTree.cpp
@@ -0,0 +1,10 @@
+// Copyright 1998-2014 Epic Games, Inc. All Rights Reserved.
+
+#include "AIModulePrivate.h"
+#include "BehaviorTree/BehaviorTree.h"
+
+DEFINE_LOG_CATEGORY(LogBehaviorTree);
+
+UBehaviorTree::UBehaviorTree(const class FPostConstructInitializeProperties& PCIP) : Super(PCIP)
+{
+}
diff --git a/Engine/Source/Runtime/AIModule/Private/BehaviorTree/BehaviorTreeComponent.cpp b/Engine/Source/Runtime/AIModule/Private/BehaviorTree/BehaviorTreeComponent.cpp
new file mode 100644
--- /dev/null
+++ b/Engine/Source/Runtime/AIModule/Private/BehaviorTree/BehaviorTreeComponent.cpp
@@ -0,0 +1,1761 @@
+// Copyright 1998-2014 Epic Games, Inc. All Rights Reserved.
+
+#include "AIModulePrivate.h"
+#include "BehaviorTree/BehaviorTree.h"
+#include "BehaviorTree/BTService.h"
+#include "BehaviorTreeDelegates.h"
+#include "BehaviorTree/BTCompositeNode.h"
+#include "BehaviorTree/BehaviorTreeManager.h"
+#include "BehaviorTree/BlackboardComponent.h"
+#include "BehaviorTree/BehaviorTreeComponent.h"
+
+//----------------------------------------------------------------------//
+// UBehaviorTreeComponent
+//----------------------------------------------------------------------//
+
+UBehaviorTreeComponent::UBehaviorTreeComponent(const class FPostConstructInitializeProperties& PCIP) : Super(PCIP)
+{
+	ActiveInstanceIdx = 0;
+	bLoopExecution = false;
+	bWaitingForAbortingTasks = false;
+	bRequestedFlowUpdate = false;
+	bAutoActivate = true;
+	bWantsInitializeComponent = true; 
+	bIsRunning = false;
+	bIsPaused = false;
+	
+	SearchData.OwnerComp = this;
+}
+
+void UBehaviorTreeComponent::RestartLogic()
+{
+	UE_VLOG(GetOwner(), LogBehaviorTree, Log, TEXT("UBehaviorTreeComponent::RestartLogic"));
+
+	bIsRunning = true;
+	RestartTree();
+}
+
+void UBehaviorTreeComponent::StopLogic(const FString& Reason) 
+{
+	UE_VLOG(GetOwner(), LogBehaviorTree, Log, TEXT("Stopping BT, reason: \'%s\'"), *Reason);
+	bIsRunning = false;
+}
+
+void UBehaviorTreeComponent::PauseLogic(const FString& Reason)
+{
+	UE_VLOG(GetOwner(), LogBehaviorTree, Log, TEXT("Execution updates: PAUSED (%s)"), *Reason);
+	bIsPaused = true;
+
+	if (BlackboardComp)
+	{
+		BlackboardComp->PauseUpdates();
+	}
+}
+
+EAILogicResuming::Type UBehaviorTreeComponent::ResumeLogic(const FString& Reason)
+{
+	const EAILogicResuming::Type SuperResumeResult = Super::ResumeLogic(Reason);
+	if (!!bIsPaused)
+	{
+		bIsPaused = false;
+
+		if (SuperResumeResult == EAILogicResuming::Continue)
+		{
+			if (BlackboardComp)
+			{
+				BlackboardComp->ResumeUpdates();
+			}
+
+			if (ExecutionRequest.ExecuteNode)
+			{
+				ScheduleExecutionUpdate();
+			}
+
+			return EAILogicResuming::Continue;
+		}
+	}
+
+	return SuperResumeResult;
+}
+
+bool UBehaviorTreeComponent::TreeHasBeenStarted() const
+{
+	return bIsRunning && InstanceStack.Num();
+}
+
+bool UBehaviorTreeComponent::IsRunning() const
+{ 
+	return bIsPaused == false && TreeHasBeenStarted() == true;
+}
+
+bool UBehaviorTreeComponent::IsPaused() const
+{
+	return bIsPaused;
+}
+
+bool UBehaviorTreeComponent::StartTree(class UBehaviorTree* TreeAsset, EBTExecutionMode::Type ExecuteMode)
+{
+	if (UBehaviorTreeManager::IsBehaviorTreeUsageEnabled() == false)
+	{
+		UE_VLOG(GetOwner(), LogBehaviorTree, Error, TEXT("Behavior tree usage is not enabled"));
+		return false;
+	}
+
+	// clear instance stack, start should always run new tree from root
+	class UBehaviorTree* CurrentRoot = GetRootTree();
+	
+	if (CurrentRoot == TreeAsset && TreeHasBeenStarted())
+	{
+		UE_VLOG(GetOwner(), LogBehaviorTree, Log, TEXT("Skipping behavior start request - it's already running"));
+		return true;
+	}
+	else if (CurrentRoot)
+	{
+		UE_VLOG(GetOwner(), LogBehaviorTree, Log, TEXT("Abandoning behavior %s to start new one (%s)"),
+			*GetNameSafe(CurrentRoot), *GetNameSafe(TreeAsset));
+	}
+
+	StopTree();
+	bLoopExecution = (ExecuteMode == EBTExecutionMode::Looped);
+	bIsRunning = true;
+
+#if USE_BEHAVIORTREE_DEBUGGER
+	DebuggerSteps.Reset();
+#endif
+
+	// push new instance
+	const bool bPushed = PushInstance(TreeAsset);
+	return bPushed;
+}
+
+void UBehaviorTreeComponent::StopTree()
+{
+	// clear current state, don't touch debugger data
+	for (int32 i = 0; i < InstanceStack.Num(); i++)
+	{
+		InstanceStack[i].Cleanup(this);
+	}
+
+	InstanceStack.Reset();
+	TaskMessageObservers.Reset();
+	ExecutionRequest = FBTNodeExecutionInfo();
+	ActiveInstanceIdx = 0;
+	GetOwner()->GetWorldTimerManager().ClearTimer(this, &UBehaviorTreeComponent::ProcessExecutionRequest);
+	// make sure to allow new execution requests since we just removed last request timer
+	bRequestedFlowUpdate = false;
+}
+
+void UBehaviorTreeComponent::RestartTree()
+{
+	UE_VLOG(GetOwner(), LogBehaviorTree, Log, TEXT("UBehaviorTreeComponent::RestartTree"));
+	
+	if (InstanceStack.Num())
+	{
+		FBehaviorTreeInstance& TopInstance = InstanceStack[0];
+		RequestExecution(TopInstance.RootNode, 0, TopInstance.RootNode, -1, EBTNodeResult::Aborted);
+	}
+}
+
+void UBehaviorTreeComponent::OnTaskFinished(const class UBTTaskNode* TaskNode, EBTNodeResult::Type TaskResult)
+{
+	if (TaskNode == NULL || InstanceStack.Num() == 0 || IsPendingKill())
+	{
+		return;
+	}
+
+	UE_VLOG(GetOwner(), LogBehaviorTree, Log, TEXT("Task %s finished: %s"), 
+		*UBehaviorTreeTypes::DescribeNodeHelper(TaskNode), *UBehaviorTreeTypes::DescribeNodeResult(TaskResult));
+
+	// notify parent node
+	UBTCompositeNode* ParentNode = TaskNode->GetParentNode();
+	const int32 TaskInstanceIdx = FindInstanceContainingNode(TaskNode);
+	uint8* ParentMemory = ParentNode->GetNodeMemory<uint8>(InstanceStack[TaskInstanceIdx]);
+
+	const bool bWasWaitingForAbort = bWaitingForAbortingTasks;
+	ParentNode->ConditionalNotifyChildExecution(this, ParentMemory, TaskNode, TaskResult);
+	
+	if (TaskResult != EBTNodeResult::InProgress)
+	{
+		StoreDebuggerSearchStep(TaskNode, TaskInstanceIdx, TaskResult);
+
+		// cleanup task observers
+		UnregisterMessageObserversFrom(TaskNode);
+
+		// update execution when active task is finished
+		FBehaviorTreeInstance& ActiveInstance = InstanceStack[ActiveInstanceIdx];
+		if (ActiveInstance.ActiveNode == TaskNode)
+		{
+			const bool bWasAborting = (ActiveInstance.ActiveNodeType == EBTActiveNode::AbortingTask);
+			ActiveInstance.ActiveNodeType = EBTActiveNode::InactiveTask;
+
+			// request execution from parent
+			if (!bWasAborting)
+			{
+				RequestExecution(TaskResult);
+			}
+		}
+
+		// update state of aborting tasks after currently finished one was set to Inactive
+		UpdateAbortingTasks();
+
+		// make sure that we continue execution after all pending latent aborts finished
+		if (!bWaitingForAbortingTasks && bWasWaitingForAbort)
+		{
+			ScheduleExecutionUpdate();
+		}
+	}
+	else
+	{
+		// always update state of aborting tasks
+		UpdateAbortingTasks();
+	}
+}
+
+void UBehaviorTreeComponent::OnTreeFinished()
+{
+	UE_VLOG(GetOwner(), LogBehaviorTree, Verbose, TEXT("Ran out of nodes to check, %s tree."),
+		bLoopExecution ? TEXT("looping") : TEXT("stopping"));
+
+	ActiveInstanceIdx = 0;
+	StoreDebuggerExecutionStep(EBTExecutionSnap::OutOfNodes);
+
+	if (bLoopExecution)
+	{
+		// it should be already deactivated (including root)
+		// set active node to initial state: root activation
+		FBehaviorTreeInstance& TopInstance = InstanceStack[0];
+		TopInstance.ActiveNode = NULL;
+		TopInstance.ActiveNodeType = EBTActiveNode::Composite;
+
+		// result doesn't really matter, root node will be reset and start iterating child nodes from scratch
+		// although it shouldn't be set to Aborted, as it has special meaning in RequestExecution (switch to higher priority)
+		RequestExecution(TopInstance.RootNode, 0, TopInstance.RootNode, 0, EBTNodeResult::InProgress);
+	}
+	else
+	{
+		StopTree();
+	}
+}
+
+bool UBehaviorTreeComponent::IsExecutingBranch(const class UBTNode* Node, int32 ChildIndex) const
+{
+	const int32 TestInstanceIdx = FindInstanceContainingNode(Node);
+	if (TestInstanceIdx == INDEX_NONE || InstanceStack[TestInstanceIdx].ActiveNode == NULL)
+	{
+		return false;
+	}
+
+	// check if branch is currently being restarted:
+	const FBTNodeIndex TestIndex(TestInstanceIdx, Node->GetExecutionIndex());
+	if (ExecutionRequest.ExecuteNode)
+	{
+		// - search takes priority (won't work on decorators above)
+		if (ExecutionRequest.SearchStart.TakesPriorityOver(TestIndex) ||
+			ExecutionRequest.SearchStart == TestIndex)
+		{
+			return false;
+		}
+
+		// - special case for decorators skipped by previous check
+		const UBTDecorator* TestDecorator = Cast<const UBTDecorator>(Node);
+		if (TestDecorator && TestDecorator->GetParentNode() == ExecutionRequest.ExecuteNode)
+		{
+			return false;
+		}
+	}
+
+	// is it active node or root of tree?
+	const FBehaviorTreeInstance& TestInstance = InstanceStack[TestInstanceIdx];
+	if (Node == TestInstance.RootNode || Node == TestInstance.ActiveNode)
+	{
+		return true;
+	}
+
+	// compare with index of next child
+	const uint16 ActiveExecutionIndex = TestInstance.ActiveNode->GetExecutionIndex();
+	const uint16 NextChildExecutionIndex = Node->GetParentNode()->GetChildExecutionIndex(ChildIndex + 1);
+	return (ActiveExecutionIndex >= Node->GetExecutionIndex()) && (ActiveExecutionIndex < NextChildExecutionIndex);
+}
+
+bool UBehaviorTreeComponent::IsAuxNodeActive(const class UBTAuxiliaryNode* AuxNode) const
+{
+	if (AuxNode == NULL)
+	{
+		return false;
+	}
+
+	const uint16 AuxExecutionIndex = AuxNode->GetExecutionIndex();
+	for (int32 i = 0; i < InstanceStack.Num(); i++)
+	{
+		const FBehaviorTreeInstance& InstanceInfo = InstanceStack[i];
+		for (int32 iAux = 0; iAux < InstanceInfo.ActiveAuxNodes.Num(); iAux++)
+		{
+			const UBTAuxiliaryNode* TestAuxNode = InstanceInfo.ActiveAuxNodes[iAux];
+
+			// check template version
+			if (TestAuxNode == AuxNode)
+			{
+				return true;
+			}
+
+			// check instanced version
+			if (AuxNode->IsInstanced() && TestAuxNode && TestAuxNode->GetExecutionIndex() == AuxExecutionIndex)
+			{
+				const uint8* NodeMemory = TestAuxNode->GetNodeMemory<uint8>(InstanceInfo);
+				UBTNode* NodeInstance = TestAuxNode->GetNodeInstance(this, (uint8*)NodeMemory);
+
+				if (NodeInstance == AuxNode)
+				{
+					return true;
+				}
+			}
+		}
+	}
+
+	return false;
+}
+
+EBTTaskStatus::Type UBehaviorTreeComponent::GetTaskStatus(const class UBTTaskNode* TaskNode) const
+{
+	EBTTaskStatus::Type Status = EBTTaskStatus::Inactive;
+	const int32 InstanceIdx = FindInstanceContainingNode(TaskNode);
+
+	if (InstanceIdx != INDEX_NONE)
+	{
+		const uint16 ExecutionIndex = TaskNode->GetExecutionIndex();
+		const FBehaviorTreeInstance& InstanceInfo = InstanceStack[InstanceIdx];
+
+		// always check parallel execution first, it takes priority over ActiveNodeType
+		for (int32 i = 0; i < InstanceInfo.ParallelTasks.Num(); i++)
+		{
+			const FBehaviorTreeParallelTask& ParallelInfo = InstanceInfo.ParallelTasks[i];
+			if (ParallelInfo.TaskNode == TaskNode ||
+				(TaskNode->IsInstanced() && ParallelInfo.TaskNode && ParallelInfo.TaskNode->GetExecutionIndex() == ExecutionIndex))
+			{
+				Status = ParallelInfo.Status;
+				break;
+			}
+		}
+
+		if (Status == EBTTaskStatus::Inactive)
+		{
+			if (InstanceInfo.ActiveNode == TaskNode ||
+				(TaskNode->IsInstanced() && InstanceInfo.ActiveNode && InstanceInfo.ActiveNode->GetExecutionIndex() == ExecutionIndex))
+			{
+				Status =
+					(InstanceInfo.ActiveNodeType == EBTActiveNode::ActiveTask) ? EBTTaskStatus::Active :
+					(InstanceInfo.ActiveNodeType == EBTActiveNode::AbortingTask) ? EBTTaskStatus::Aborting :
+					EBTTaskStatus::Inactive;
+			}
+		}
+	}
+
+	return Status;
+}
+
+void UBehaviorTreeComponent::RequestExecution(const class UBTDecorator* RequestedBy)
+{
+	// search range depends on decorator's FlowAbortMode:
+	//
+	// - LowerPri: try entering branch = search only nodes under decorator
+	//
+	// - Self: leave execution = from node under decorator to end of tree
+	//         UNLESS branch is no longer executing (pending restart), in which case it will restart parent
+	//
+	// - Both: check if active node is within inner child nodes and choose Self or LowerPri
+	//
+
+	EBTFlowAbortMode::Type AbortMode = RequestedBy->GetFlowAbortMode();
+	if (AbortMode == EBTFlowAbortMode::None)
+	{
+		return;
+	}
+
+	const int32 InstanceIdx = FindInstanceContainingNode(RequestedBy->GetParentNode());
+	if (InstanceIdx == INDEX_NONE)
+	{
+		return;
+	}
+
+	if (AbortMode == EBTFlowAbortMode::Both || AbortMode == EBTFlowAbortMode::Self)
+	{
+		const bool bIsExecutingChildNodes = IsExecutingBranch(RequestedBy, RequestedBy->GetChildIndex());
+		AbortMode = bIsExecutingChildNodes ? EBTFlowAbortMode::Self : EBTFlowAbortMode::LowerPriority;
+	}
+
+	EBTNodeResult::Type ContinueResult = (AbortMode == EBTFlowAbortMode::Self) ? EBTNodeResult::Failed : EBTNodeResult::Aborted;
+	RequestExecution(RequestedBy->GetParentNode(), InstanceIdx, RequestedBy, RequestedBy->GetChildIndex(), ContinueResult);
+}
+
+void UBehaviorTreeComponent::RequestExecution(EBTNodeResult::Type LastResult)
+{
+	// task helpers can't continue with InProgress or Aborted result, it should be handled 
+	// either by decorator helper or regular RequestExecution() (6 param version)
+
+	if (LastResult != EBTNodeResult::Aborted && LastResult != EBTNodeResult::InProgress)
+	{
+		const FBehaviorTreeInstance& ActiveInstance = InstanceStack[ActiveInstanceIdx];
+		UBTCompositeNode* ExecuteParent = (ActiveInstance.ActiveNode == NULL) ? ActiveInstance.RootNode :
+			(ActiveInstance.ActiveNodeType == EBTActiveNode::Composite) ? (UBTCompositeNode*)ActiveInstance.ActiveNode :
+			ActiveInstance.ActiveNode->GetParentNode();
+
+		RequestExecution(ExecuteParent, InstanceStack.Num() - 1,
+			ActiveInstance.ActiveNode ? ActiveInstance.ActiveNode : ActiveInstance.RootNode, -1,
+			LastResult, false);
+	}
+}
+
+static void FindCommonParent(const TArray<FBehaviorTreeInstance>& Instances, const TArray<FBehaviorTreeInstanceId>& KnownInstances,
+							 class UBTCompositeNode* InNodeA, uint16 InstanceIdxA,
+							 class UBTCompositeNode* InNodeB, uint16 InstanceIdxB,
+							 class UBTCompositeNode*& CommonParentNode, uint16& CommonInstanceIdx)
+{
+	// find two nodes in the same instance (choose lower index = closer to root)
+	CommonInstanceIdx = (InstanceIdxA <= InstanceIdxB) ? InstanceIdxA : InstanceIdxB;
+
+	class UBTCompositeNode* NodeA = (CommonInstanceIdx == InstanceIdxA) ? InNodeA : Instances[CommonInstanceIdx].ActiveNode->GetParentNode();
+	class UBTCompositeNode* NodeB = (CommonInstanceIdx == InstanceIdxB) ? InNodeB : Instances[CommonInstanceIdx].ActiveNode->GetParentNode();
+
+	// special case: node was taken from CommonInstanceIdx, but it had ActiveNode set to root (no parent)
+	if (NodeA == NULL && CommonInstanceIdx != InstanceIdxA)
+	{
+		NodeA = Instances[CommonInstanceIdx].RootNode;
+	}
+	if (NodeB == NULL && CommonInstanceIdx != InstanceIdxB)
+	{
+		NodeB = Instances[CommonInstanceIdx].RootNode;
+	}
+
+	// if one of nodes is still empty, we have serious problem with execution flow - crash and log details
+	if (NodeA == NULL || NodeB == NULL)
+	{
+		FString AssetAName = Instances.IsValidIndex(InstanceIdxA) && KnownInstances.IsValidIndex(Instances[InstanceIdxA].InstanceIdIndex) ? GetNameSafe(KnownInstances[Instances[InstanceIdxA].InstanceIdIndex].TreeAsset) : TEXT("unknown");
+		FString AssetBName = Instances.IsValidIndex(InstanceIdxB) && KnownInstances.IsValidIndex(Instances[InstanceIdxB].InstanceIdIndex) ? GetNameSafe(KnownInstances[Instances[InstanceIdxB].InstanceIdIndex].TreeAsset) : TEXT("unknown");
+		FString AssetCName = Instances.IsValidIndex(CommonInstanceIdx) && KnownInstances.IsValidIndex(Instances[CommonInstanceIdx].InstanceIdIndex) ? GetNameSafe(KnownInstances[Instances[CommonInstanceIdx].InstanceIdIndex].TreeAsset) : TEXT("unknown");
+
+		UE_LOG(LogBehaviorTree, Fatal, TEXT("Fatal error in FindCommonParent() call.\nInNodeA: %s, InstanceIdxA: %d (%s), NodeA: %s\nInNodeB: %s, InstanceIdxB: %d (%s), NodeB: %s\nCommonInstanceIdx: %d (%s), ActiveNode: %s%s"),
+			*UBehaviorTreeTypes::DescribeNodeHelper(InNodeA), InstanceIdxA, *AssetAName, *UBehaviorTreeTypes::DescribeNodeHelper(NodeA),
+			*UBehaviorTreeTypes::DescribeNodeHelper(InNodeB), InstanceIdxB, *AssetBName, *UBehaviorTreeTypes::DescribeNodeHelper(NodeB),
+			CommonInstanceIdx, *AssetCName, *UBehaviorTreeTypes::DescribeNodeHelper(Instances[CommonInstanceIdx].ActiveNode),
+			(Instances[CommonInstanceIdx].ActiveNode == Instances[CommonInstanceIdx].RootNode) ? TEXT(" (root)") : TEXT(""));
+	}
+
+	// find common parent of two nodes
+	int32 NodeADepth = NodeA->GetTreeDepth();
+	int32 NodeBDepth = NodeB->GetTreeDepth();
+
+	while (NodeADepth > NodeBDepth)
+	{
+		NodeA = NodeA->GetParentNode();
+		NodeADepth = NodeA->GetTreeDepth();
+	}
+
+	while (NodeBDepth > NodeADepth)
+	{
+		NodeB = NodeB->GetParentNode();
+		NodeBDepth = NodeB->GetTreeDepth();
+	}
+
+	while (NodeA != NodeB)
+	{
+		NodeA = NodeA->GetParentNode();
+		NodeB = NodeB->GetParentNode();
+	}
+
+	CommonParentNode = NodeA;
+}
+
+void UBehaviorTreeComponent::ScheduleExecutionUpdate()
+{
+	if (!bRequestedFlowUpdate)
+	{
+		bRequestedFlowUpdate = true;
+
+		GetOwner()->GetWorldTimerManager().SetTimer(this, &UBehaviorTreeComponent::ProcessExecutionRequest, 0.001f, false);
+	}
+}
+
+void UBehaviorTreeComponent::RequestExecution(class UBTCompositeNode* RequestedOn, int32 InstanceIdx, const class UBTNode* RequestedBy,
+											  int32 RequestedByChildIndex, EBTNodeResult::Type ContinueWithResult, bool bStoreForDebugger)
+{
+	SCOPE_CYCLE_COUNTER(STAT_AI_BehaviorTree_SearchTime);
+
+	UE_VLOG(GetOwner(), LogBehaviorTree, Log, TEXT("Execution request by %s (result: %s)"),
+		*UBehaviorTreeTypes::DescribeNodeHelper(RequestedBy),
+		*UBehaviorTreeTypes::DescribeNodeResult(ContinueWithResult));
+
+	if (!bIsRunning || GetOwner() == NULL || GetOwner()->IsPendingKillPending())
+	{
+		UE_VLOG(GetOwner(), LogBehaviorTree, Log, TEXT("> skip: tree is not running"));
+		return;
+	}
+
+	const bool bSwitchToHigherPriority = (ContinueWithResult == EBTNodeResult::Aborted);
+	const UBTNode* DebuggerNode = bStoreForDebugger ? RequestedBy : NULL;
+
+	FBTNodeIndex ExecutionIdx;
+	ExecutionIdx.InstanceIndex = InstanceIdx;
+	ExecutionIdx.ExecutionIndex = RequestedBy->GetExecutionIndex();
+
+	if (bSwitchToHigherPriority)
+	{
+		RequestedByChildIndex = FMath::Max(0, RequestedByChildIndex);
+		ExecutionIdx.ExecutionIndex = RequestedOn->GetChildExecutionIndex(RequestedByChildIndex, EBTChildIndex::FirstNode);
+	}
+
+	// check if it's more important than currently requested
+	if (ExecutionRequest.ExecuteNode && ExecutionRequest.SearchStart.TakesPriorityOver(ExecutionIdx))
+	{
+		UE_VLOG(GetOwner(), LogBehaviorTree, Log, TEXT("> skip: already has request with higher priority"));
+		StoreDebuggerRestart(DebuggerNode, InstanceIdx, true);
+		return;
+	}
+
+	// when it's aborting and moving to higher priority node:
+	if (bSwitchToHigherPriority)
+	{
+		// check if decorators allow execution on requesting link
+		const bool bShouldCheckDecorators = (RequestedByChildIndex >= 0);
+		const bool bCanExecute = !bShouldCheckDecorators || RequestedOn->DoDecoratorsAllowExecution(this, InstanceIdx, RequestedByChildIndex);
+		if (!bCanExecute)
+		{
+			UE_VLOG(GetOwner(), LogBehaviorTree, Log, TEXT("> skip: decorators are not allowing execution"));
+			StoreDebuggerRestart(DebuggerNode, InstanceIdx, false);
+			return;
+		}
+
+		// update common parent: requesting node with prev common/active node
+		UBTCompositeNode* CurrentNode = ExecutionRequest.ExecuteNode;
+		uint16 CurrentInstanceIdx = ExecutionRequest.ExecuteInstanceIdx;
+		if (ExecutionRequest.ExecuteNode == NULL)
+		{
+			FBehaviorTreeInstance& ActiveInstance = InstanceStack[ActiveInstanceIdx];
+			CurrentNode = (ActiveInstance.ActiveNode == NULL) ? ActiveInstance.RootNode :
+				(ActiveInstance.ActiveNodeType == EBTActiveNode::Composite) ? (UBTCompositeNode*)ActiveInstance.ActiveNode :
+				ActiveInstance.ActiveNode->GetParentNode();
+
+			CurrentInstanceIdx = InstanceStack.Num() - 1;
+		}
+
+		if (ExecutionRequest.ExecuteNode != RequestedOn)
+		{
+			UBTCompositeNode* CommonParent = NULL;
+			uint16 CommonInstanceIdx = MAX_uint16;
+
+			FindCommonParent(InstanceStack, KnownInstances, RequestedOn, InstanceIdx, CurrentNode, CurrentInstanceIdx, CommonParent, CommonInstanceIdx);
+
+			// check decorators between common parent and restart parent
+			// it's always on the same stack level, because only tasks can push new subtrees
+			for (UBTCompositeNode* It = RequestedOn; It && It != CommonParent;)
+			{
+				UBTCompositeNode* ParentNode = It->GetParentNode();
+				const int32 ChildIdx = ParentNode->GetChildIndex(It);
+				const bool bCanExecuteTest = ParentNode->DoDecoratorsAllowExecution(this, CommonInstanceIdx, ChildIdx);
+				if (!bCanExecuteTest)
+				{
+					UE_VLOG(GetOwner(), LogBehaviorTree, Log, TEXT("> skip: decorators are not allowing execution"));
+					StoreDebuggerRestart(DebuggerNode, InstanceIdx, false);
+					return;
+				}
+
+				It = ParentNode;
+			}
+
+			ExecutionRequest.ExecuteNode = CommonParent;
+			ExecutionRequest.ExecuteInstanceIdx = CommonInstanceIdx;
+		}
+	}
+	else
+	{
+		// check if decorators allow execution on requesting link (only when restart comes from composite decorator)
+		const bool bShouldCheckDecorators = RequestedOn->Children.IsValidIndex(RequestedByChildIndex) &&
+			(RequestedOn->Children[RequestedByChildIndex].DecoratorOps.Num() > 0) &&
+			RequestedBy->IsA(UBTDecorator::StaticClass());
+
+		const bool bCanExecute = bShouldCheckDecorators && RequestedOn->DoDecoratorsAllowExecution(this, InstanceIdx, RequestedByChildIndex);
+		if (bCanExecute)
+		{
+			UE_VLOG(GetOwner(), LogBehaviorTree, Log, TEXT("> skip: decorators are still allowing execution"));
+			StoreDebuggerRestart(DebuggerNode, InstanceIdx, false);
+			return;
+		}
+
+		ExecutionRequest.ExecuteNode = RequestedOn;
+		ExecutionRequest.ExecuteInstanceIdx = InstanceIdx;
+	}
+
+	// store it
+	StoreDebuggerRestart(DebuggerNode, InstanceIdx, true);
+
+	ExecutionRequest.ContinueWithResult = ContinueWithResult;
+	ExecutionRequest.SearchStart = ExecutionIdx;
+	ExecutionRequest.bTryNextChild = !bSwitchToHigherPriority;
+
+	ScheduleExecutionUpdate();
+}
+
+void UBehaviorTreeComponent::ApplySearchUpdates(const TArray<struct FBehaviorTreeSearchUpdate>& UpdateList, int32 UpToIdx, int32 NewNodeExecutionIndex, bool bPostUpdate)
+{
+	for (int32 i = 0; i < UpToIdx; i++)
+	{
+		const FBehaviorTreeSearchUpdate& UpdateInfo = UpdateList[i];
+		const UBTNode* UpdateNode = UpdateInfo.AuxNode ? (const UBTNode*)UpdateInfo.AuxNode : (const UBTNode*)UpdateInfo.TaskNode;
+		FBehaviorTreeInstance& UpdateInstance = InstanceStack[UpdateInfo.InstanceIndex];
+		int32 ParallelTaskIdx = INDEX_NONE;
+		bool bIsActive = false;
+
+		if (UpdateInfo.AuxNode)
+		{
+			bIsActive = UpdateInstance.ActiveAuxNodes.Contains(UpdateInfo.AuxNode);
+		}
+		else if (UpdateInfo.TaskNode)
+		{
+			ParallelTaskIdx = UpdateInstance.ParallelTasks.IndexOfByKey(UpdateInfo.TaskNode);
+			bIsActive = (ParallelTaskIdx != INDEX_NONE && UpdateInstance.ParallelTasks[ParallelTaskIdx].Status == EBTTaskStatus::Active);
+		}
+
+		if ((UpdateInfo.Mode == EBTNodeUpdateMode::Remove && !bIsActive) ||
+			(UpdateInfo.Mode == EBTNodeUpdateMode::Add && bIsActive) ||
+			(UpdateInfo.Mode == EBTNodeUpdateMode::AddForLowerPri && (bIsActive || UpdateNode->GetExecutionIndex() > NewNodeExecutionIndex)) ||
+			(UpdateInfo.bPostUpdate != bPostUpdate))
+		{
+			continue;
+		}
+
+		UE_VLOG(GetOwner(), LogBehaviorTree, Log, TEXT("Update: %s for %s: %s"),
+			*UBehaviorTreeTypes::DescribeNodeUpdateMode(UpdateInfo.Mode),
+			UpdateInfo.AuxNode ? TEXT("auxiliary node") : TEXT("parallel's main task"),
+			*UBehaviorTreeTypes::DescribeNodeHelper(UpdateNode));
+
+		if (UpdateInfo.AuxNode)
+		{
+			// special case: service node at root of top most subtree - don't remove/re-add them when tree is in looping mode
+			// don't bother with decorators parent == root means that they are on child branches
+			if (bLoopExecution && UpdateInfo.AuxNode->GetParentNode() == InstanceStack[0].RootNode &&
+				UpdateInfo.AuxNode->IsA(UBTService::StaticClass()))
+			{
+				if (UpdateInfo.Mode == EBTNodeUpdateMode::Remove ||
+					InstanceStack[0].ActiveAuxNodes.Contains(UpdateInfo.AuxNode))
+				{
+					UE_VLOG(GetOwner(), LogBehaviorTree, Verbose, TEXT("> skip [looped execution]"));
+					continue;
+				}
+			}
+
+			uint8* NodeMemory = (uint8*)UpdateNode->GetNodeMemory<uint8>(UpdateInstance);
+			if (UpdateInfo.Mode == EBTNodeUpdateMode::Remove)
+			{
+				UpdateInstance.ActiveAuxNodes.RemoveSingleSwap(UpdateInfo.AuxNode);
+				UpdateInfo.AuxNode->WrappedOnCeaseRelevant(this, NodeMemory);
+			}
+			else
+			{
+				UpdateInstance.ActiveAuxNodes.Add(UpdateInfo.AuxNode);
+				UpdateInfo.AuxNode->WrappedOnBecomeRelevant(this, NodeMemory);
+			}
+		}
+		else if (UpdateInfo.TaskNode)
+		{
+			if (UpdateInfo.Mode == EBTNodeUpdateMode::Remove)
+			{
+				// remove all message observers from node to abort to avoid calling OnTaskFinished from AbortTask
+				UnregisterMessageObserversFrom(UpdateInfo.TaskNode);
+
+				uint8* NodeMemory = (uint8*)UpdateNode->GetNodeMemory<uint8>(UpdateInstance);
+				EBTNodeResult::Type NodeResult = UpdateInfo.TaskNode->AbortTask(this, NodeMemory);
+
+				UE_VLOG(GetOwner(), LogBehaviorTree, Log, TEXT("Parallel task aborted: %s (%s)"),
+					*UBehaviorTreeTypes::DescribeNodeHelper(UpdateInfo.TaskNode),
+					(NodeResult == EBTNodeResult::InProgress) ? TEXT("in progress") : TEXT("instant"));
+
+				// mark as pending abort
+				if (NodeResult == EBTNodeResult::InProgress)
+				{
+					UpdateInstance.ParallelTasks[ParallelTaskIdx].Status = EBTTaskStatus::Aborting;
+					bWaitingForAbortingTasks = true;
+				}
+
+				OnTaskFinished(UpdateInfo.TaskNode, NodeResult);
+			}
+			else
+			{
+				UE_VLOG(GetOwner(), LogBehaviorTree, Verbose, TEXT("Parallel task: %s added to active list"),
+					*UBehaviorTreeTypes::DescribeNodeHelper(UpdateInfo.TaskNode));
+
+				UpdateInstance.ParallelTasks.Add(FBehaviorTreeParallelTask(UpdateInfo.TaskNode, EBTTaskStatus::Active));
+			}
+		}
+	}
+}
+
+void UBehaviorTreeComponent::ApplySearchData(class UBTNode* NewActiveNode, int32 UpToIdx)
+{
+	const bool bFullUpdate = (UpToIdx < 0);
+	const int32 MaxIdx = bFullUpdate ? SearchData.PendingUpdates.Num() : (UpToIdx + 1);
+	const int32 NewNodeExecutionIndex = NewActiveNode ? NewActiveNode->GetExecutionIndex() : 0;
+	ApplySearchUpdates(SearchData.PendingUpdates, MaxIdx, NewNodeExecutionIndex);
+
+	// go though post updates (services) and clear both arrays if function was called without limit
+	if (bFullUpdate)
+	{
+		ApplySearchUpdates(SearchData.PendingUpdates, SearchData.PendingUpdates.Num(), NewNodeExecutionIndex, true);
+		SearchData.PendingUpdates.Reset();
+	}
+	else
+	{
+		SearchData.PendingUpdates.RemoveAt(0, UpToIdx + 1, false);
+	}
+}
+
+void UBehaviorTreeComponent::TickComponent(float DeltaTime, enum ELevelTick TickType, FActorComponentTickFunction *ThisTickFunction)
+{
+	Super::TickComponent(DeltaTime, TickType, ThisTickFunction);
+	SCOPE_CYCLE_COUNTER(STAT_AI_BehaviorTree_Tick);
+
+	check(this && this->IsPendingKill() == false);
+		
+	if (InstanceStack.Num() == 0 || !bIsRunning)
+	{
+		return;
+	}
+	
+	// tick active auxiliary nodes and parallel tasks (in execution order, before task)
+	for (int32 i = 0; i < InstanceStack.Num(); i++)
+	{
+		FBehaviorTreeInstance& InstanceInfo = InstanceStack[i];		
+		for (int32 iAux = 0; iAux < InstanceInfo.ActiveAuxNodes.Num(); iAux++)
+		{
+			const UBTAuxiliaryNode* AuxNode = InstanceInfo.ActiveAuxNodes[iAux];
+			uint8* NodeMemory = AuxNode->GetNodeMemory<uint8>(InstanceInfo);
+			AuxNode->WrappedTickNode(this, NodeMemory, DeltaTime);
+		}
+
+		for (int32 iTask = 0; iTask < InstanceInfo.ParallelTasks.Num(); iTask++)
+		{
+			const UBTTaskNode* ParallelTask = InstanceInfo.ParallelTasks[iTask].TaskNode;
+			uint8* NodeMemory = ParallelTask->GetNodeMemory<uint8>(InstanceInfo);
+			ParallelTask->WrappedTickTask(this, NodeMemory, DeltaTime);
+		}
+	}
+
+	// tick active task
+	FBehaviorTreeInstance& ActiveInstance = InstanceStack[ActiveInstanceIdx];
+	if (ActiveInstance.ActiveNodeType == EBTActiveNode::ActiveTask ||
+		ActiveInstance.ActiveNodeType == EBTActiveNode::AbortingTask)
+	{
+		UBTTaskNode* ActiveTask = (UBTTaskNode*)ActiveInstance.ActiveNode;
+		uint8* NodeMemory = ActiveTask->GetNodeMemory<uint8>(ActiveInstance);
+		ActiveTask->WrappedTickTask(this, NodeMemory, DeltaTime);
+	}
+}
+
+void UBehaviorTreeComponent::ProcessExecutionRequest()
+{
+	bRequestedFlowUpdate = false;
+	if (bIsPaused)
+	{
+		UE_VLOG(GetOwner(), LogBehaviorTree, Verbose, TEXT("Ignoring ProcessExecutionRequest call due to BTComponent still being paused"));
+		return;
+	}
+
+	EBTNodeResult::Type NodeResult = ExecutionRequest.ContinueWithResult;
+	UBTTaskNode* NextTask = NULL;
+
+	{
+	SCOPE_CYCLE_COUNTER(STAT_AI_BehaviorTree_SearchTime);
+	
+	// abort task if needed
+	if (InstanceStack[ActiveInstanceIdx].ActiveNodeType == EBTActiveNode::ActiveTask)
+	{
+		AbortCurrentTask();
+	}
+
+	// can't continue if current task is still aborting
+	// it means, that all services and main tasks of parallels will be active!
+	if (bWaitingForAbortingTasks)
+	{
+		return;
+	}
+
+	// deactivate up to ExecuteNode
+	if (InstanceStack[ActiveInstanceIdx].ActiveNode != ExecutionRequest.ExecuteNode)
+	{
+		// DeactivateUpTo() call modifies ActiveNode & ActiveNodeType data, make sure that active task finished aborting first!
+		const bool bDeactivated = DeactivateUpTo(ExecutionRequest.ExecuteNode, ExecutionRequest.ExecuteInstanceIdx, NodeResult);
+		if (!bDeactivated)
+		{
+			return;
+		}
+
+		// find pending abort request with highest priority (search from end to optimize)
+		// and apply all changes below, to be sure that tree is not executing
+		// inactive tasks before searching for new one
+		for (int32 i = SearchData.PendingUpdates.Num() - 1; i >= 0; i--)
+		{
+			FBehaviorTreeSearchUpdate& UpdateInfo = SearchData.PendingUpdates[i];
+			if (UpdateInfo.TaskNode && UpdateInfo.Mode == EBTNodeUpdateMode::Remove)
+			{
+				// NewActiveNode param doesn't matter here, it's only for AddForLowerPri mode
+				ApplySearchData(NULL, i);
+				break;
+			}
+		}
+	}
+
+	// can't continue if parallel tasks are still aborting
+	if (bWaitingForAbortingTasks)
+	{
+		return;
+	}
+
+	// make sure that we don't have any additional instances on stack
+	// (ApplySearchData doesn't remove them, so parallels can abort latently)
+	if (InstanceStack.Num() > (ActiveInstanceIdx + 1))
+	{
+		for (int32 i = ActiveInstanceIdx + 1; i < InstanceStack.Num(); i++)
+		{
+			InstanceStack[i].Cleanup(this);
+		}
+
+		InstanceStack.SetNum(ActiveInstanceIdx + 1);
+	}
+
+	FBehaviorTreeInstance& ActiveInstance = InstanceStack[ActiveInstanceIdx];
+	UBTCompositeNode* TestNode = ExecutionRequest.ExecuteNode;
+
+	// activate root node if needed (can't be handled by parent composite...)
+	if (ActiveInstance.ActiveNode == NULL)
+	{
+		ActiveInstance.ActiveNode = InstanceStack[ActiveInstanceIdx].RootNode;
+		ActiveInstance.RootNode->OnNodeActivation(SearchData);
+		BT_SEARCHLOG(SearchData, Verbose, TEXT("Activated root node: %s"), *UBehaviorTreeTypes::DescribeNodeHelper(ActiveInstance.RootNode));
+	}
+
+	// additional operations for restarting:
+	if (!ExecutionRequest.bTryNextChild)
+	{
+		// mark all decorators less important than current search start for removal
+		UnregisterAuxNodesUpTo(ExecutionRequest.SearchStart);
+
+		// reactivate top search node, so it could use search range correctly
+		BT_SEARCHLOG(SearchData, Verbose, TEXT("Reactivate node: %s [restart]"), *UBehaviorTreeTypes::DescribeNodeHelper(TestNode));
+		ExecutionRequest.ExecuteNode->OnNodeRestart(SearchData);
+
+		SearchData.SearchStart = ExecutionRequest.SearchStart;
+	}
+	else
+	{
+		// make sure it's reset before starting new search
+		SearchData.SearchStart = FBTNodeIndex();
+	}
+
+	// store blackboard values from search start (can be changed by aux node removal/adding)
+#if USE_BEHAVIORTREE_DEBUGGER
+	StoreDebuggerBlackboard(SearchStartBlackboard);
+#endif
+
+	// start looking for next task
+	while (TestNode && NextTask == NULL)
+	{
+		BT_SEARCHLOG(SearchData, Verbose, TEXT("Testing node: %s"),	*UBehaviorTreeTypes::DescribeNodeHelper(TestNode));
+		const int32 ChildBranchIdx = TestNode->FindChildToExecute(SearchData, NodeResult);
+		UBTNode* StoreNode = TestNode;
+
+		if (ChildBranchIdx == BTSpecialChild::ReturnToParent)
+		{
+			UBTCompositeNode* ChildNode = TestNode;
+			TestNode = TestNode->GetParentNode();
+
+			// does it want to move up the tree?
+			if (TestNode == NULL)
+			{
+				// special case for leaving instance: deactivate root manually
+				ChildNode->OnNodeDeactivation(SearchData, NodeResult);
+
+				// don't remove top instance from stack, so it could be looped
+				if (ActiveInstanceIdx > 0)
+				{
+					StoreDebuggerSearchStep(InstanceStack[ActiveInstanceIdx].ActiveNode, ActiveInstanceIdx, NodeResult);
+					StoreDebuggerRemovedInstance(ActiveInstanceIdx);
+
+					// apply pending aux changes (don't give any NewActiveNode param, AddForLowerPri updates should be skipped)
+					ApplySearchData(NULL);
+
+					// and leave subtree
+					InstanceStack[ActiveInstanceIdx].Cleanup(this);
+					InstanceStack.Pop();
+					ActiveInstanceIdx--;
+
+					StoreDebuggerSearchStep(InstanceStack[ActiveInstanceIdx].ActiveNode, ActiveInstanceIdx, NodeResult);
+					TestNode = InstanceStack[ActiveInstanceIdx].ActiveNode->GetParentNode();
+				}
+			}
+
+			if (TestNode)
+			{
+				TestNode->OnChildDeactivation(SearchData, ChildNode, NodeResult);
+			}
+		}
+		else
+		{
+			// was new task found?
+			NextTask = TestNode->Children[ChildBranchIdx].ChildTask;
+
+			// or it wants to move down the tree?
+			TestNode = TestNode->Children[ChildBranchIdx].ChildComposite;
+		}
+
+		// store after node deactivation had chance to modify result
+		StoreDebuggerSearchStep(StoreNode, ActiveInstanceIdx, NodeResult);
+	}
+
+	// clear pending execution variables (applying search data can request new abort)
+	ExecutionRequest = FBTNodeExecutionInfo();
+
+	ApplySearchData(NextTask);
+
+	// finish timer scope
+	}
+
+	// execute next task / notify out of nodes
+	if (NextTask)
+	{
+		ExecuteTask(NextTask);
+	}
+	else
+	{
+		OnTreeFinished();
+	}
+}
+
+bool UBehaviorTreeComponent::DeactivateUpTo(class UBTCompositeNode* Node, uint16 NodeInstanceIdx, EBTNodeResult::Type& NodeResult)
+{
+	UBTNode* DeactivatedChild = InstanceStack[ActiveInstanceIdx].ActiveNode;
+	bool bDeactivateRoot = true;
+
+	if (DeactivatedChild == NULL && ActiveInstanceIdx > NodeInstanceIdx)
+	{
+		// use tree's root node if instance didn't activated itself yet
+		DeactivatedChild = InstanceStack[ActiveInstanceIdx].RootNode;
+		bDeactivateRoot = false;
+	}
+
+	while (DeactivatedChild)
+	{
+		UBTCompositeNode* NotifyParent = DeactivatedChild->GetParentNode();
+		if (NotifyParent)
+		{
+			NotifyParent->OnChildDeactivation(SearchData, DeactivatedChild, NodeResult);
+
+			BT_SEARCHLOG(SearchData, Verbose, TEXT("Deactivate node: %s"), *UBehaviorTreeTypes::DescribeNodeHelper(DeactivatedChild));
+			StoreDebuggerSearchStep(DeactivatedChild, ActiveInstanceIdx, NodeResult);
+			DeactivatedChild = NotifyParent;
+
+			FBehaviorTreeInstance& InstanceInfo = InstanceStack[ActiveInstanceIdx];
+			InstanceInfo.ActiveNode = NotifyParent;
+			InstanceInfo.ActiveNodeType = EBTActiveNode::Composite;
+		}
+		else
+		{
+			// special case for leaving instance: deactivate root manually
+			if (bDeactivateRoot)
+			{
+				InstanceStack[ActiveInstanceIdx].RootNode->OnNodeDeactivation(SearchData, NodeResult);
+			}
+
+			BT_SEARCHLOG(SearchData, Verbose, TEXT("%s node: %s [leave subtree]"),
+				bDeactivateRoot ? TEXT("Deactivate") : TEXT("Skip over"),
+				*UBehaviorTreeTypes::DescribeNodeHelper(InstanceStack[ActiveInstanceIdx].RootNode));
+
+			// clear flag, it's valid only for newest instance
+			bDeactivateRoot = true;
+
+			// shouldn't happen, but it's better to have built in failsafe just in case
+			if (ActiveInstanceIdx == 0)
+			{
+				BT_SEARCHLOG(SearchData, Error, TEXT("Execution path does NOT contain common parent node, restarting tree! AI:%s"),
+					*GetNameSafe(SearchData.OwnerComp->GetOwner()));
+
+				RestartTree();
+				return false;
+			}
+
+			ActiveInstanceIdx--;
+			DeactivatedChild = InstanceStack[ActiveInstanceIdx].ActiveNode;
+
+			// apply aux node changes from removed subtree
+			// (don't give any NewActiveNode param, AddForLowerPri updates should be skipped)
+			ApplySearchData(NULL);
+		}
+
+		if (DeactivatedChild == Node)
+		{
+			break;
+		}
+	}
+
+	return true;
+}
+
+void UBehaviorTreeComponent::UnregisterAuxNodesUpTo(const FBTNodeIndex& Index)
+{
+	for (int32 i = 0; i < InstanceStack.Num(); i++)
+	{
+		FBehaviorTreeInstance& InstanceInfo = InstanceStack[i];
+		for (int32 iAux = 0; iAux < InstanceInfo.ActiveAuxNodes.Num(); iAux++)
+		{
+			FBTNodeIndex AuxIdx(i, InstanceInfo.ActiveAuxNodes[iAux]->GetExecutionIndex());
+			if (Index.TakesPriorityOver(AuxIdx))
+			{
+				SearchData.AddUniqueUpdate(FBehaviorTreeSearchUpdate(InstanceInfo.ActiveAuxNodes[iAux], i, EBTNodeUpdateMode::Remove));
+			}
+		}
+	}
+}
+
+void UBehaviorTreeComponent::ExecuteTask(UBTTaskNode* TaskNode)
+{
+	SCOPE_CYCLE_COUNTER(STAT_AI_BehaviorTree_ExecutionTime);
+
+	FBehaviorTreeInstance& ActiveInstance = InstanceStack[ActiveInstanceIdx];
+	ActiveInstance.ActiveNode = TaskNode;
+	ActiveInstance.ActiveNodeType = EBTActiveNode::ActiveTask;
+
+	// make a snapshot for debugger
+	StoreDebuggerExecutionStep(EBTExecutionSnap::Regular);
+
+	UE_VLOG(GetOwner(), LogBehaviorTree, Log, TEXT("Execute task: %s"), *UBehaviorTreeTypes::DescribeNodeHelper(TaskNode));
+
+	// store instance before execution, it could result in pushing a subtree
+	uint16 InstanceIdx = ActiveInstanceIdx;
+
+	uint8* NodeMemory = (uint8*)(TaskNode->GetNodeMemory<uint8>(ActiveInstance));
+	EBTNodeResult::Type TaskResult = TaskNode->WrappedExecuteTask(this, NodeMemory);
+
+	// update task's runtime values after it had a chance to initialize memory
+	UpdateDebuggerAfterExecution(TaskNode, InstanceIdx);
+
+	OnTaskFinished(TaskNode, TaskResult);
+}
+
+void UBehaviorTreeComponent::AbortCurrentTask()
+{
+	FBehaviorTreeInstance& ActiveInstance = InstanceStack[ActiveInstanceIdx];
+	UBTTaskNode* ActiveTask = (UBTTaskNode*)ActiveInstance.ActiveNode;
+
+	// remove all observers before requesting abort
+	UnregisterMessageObserversFrom(ActiveTask);
+
+	UE_VLOG(GetOwner(), LogBehaviorTree, Log, TEXT("Abort task: %s"), *UBehaviorTreeTypes::DescribeNodeHelper(ActiveTask));
+
+	// abort task using current state of tree
+	uint8* NodeMemory = (uint8*)(ActiveTask->GetNodeMemory<uint8>(ActiveInstance));
+	EBTNodeResult::Type TaskResult = ActiveTask->WrappedAbortTask(this, NodeMemory);
+
+	ActiveInstance.ActiveNodeType = EBTActiveNode::AbortingTask;
+
+	OnTaskFinished(ActiveTask, TaskResult);
+}
+
+void UBehaviorTreeComponent::RegisterMessageObserver(const class UBTTaskNode* TaskNode, FName MessageType)
+{
+	if (TaskNode)
+	{
+		FBTNodeIndex NodeIdx;
+		NodeIdx.ExecutionIndex = TaskNode->GetExecutionIndex();
+		NodeIdx.InstanceIndex = InstanceStack.Num() - 1;
+
+		TaskMessageObservers.Add(NodeIdx,
+			FAIMessageObserver::Create(this, MessageType, FOnAIMessage::CreateUObject(TaskNode, &UBTTaskNode::ReceivedMessage))
+			);
+
+		UE_VLOG(GetOwner(), LogBehaviorTree, Log, TEXT("Message[%s] observer added for %s"),
+			*MessageType.ToString(), *UBehaviorTreeTypes::DescribeNodeHelper(TaskNode));
+	}
+}
+
+void UBehaviorTreeComponent::RegisterMessageObserver(const class UBTTaskNode* TaskNode, FName MessageType, FAIRequestID RequestID)
+{
+	if (TaskNode)
+	{
+		FBTNodeIndex NodeIdx;
+		NodeIdx.ExecutionIndex = TaskNode->GetExecutionIndex();
+		NodeIdx.InstanceIndex = InstanceStack.Num() - 1;
+
+		TaskMessageObservers.Add(NodeIdx,
+			FAIMessageObserver::Create(this, MessageType, RequestID, FOnAIMessage::CreateUObject(TaskNode, &UBTTaskNode::ReceivedMessage))
+			);
+
+		UE_VLOG(GetOwner(), LogBehaviorTree, Log, TEXT("Message[%s:%d] observer added for %s"),
+			*MessageType.ToString(), RequestID, *UBehaviorTreeTypes::DescribeNodeHelper(TaskNode));
+	}
+}
+
+void UBehaviorTreeComponent::UnregisterMessageObserversFrom(const FBTNodeIndex& TaskIdx)
+{
+	const int32 NumRemoved = TaskMessageObservers.Remove(TaskIdx);
+	if (NumRemoved)
+	{
+		UE_VLOG(GetOwner(), LogBehaviorTree, Log, TEXT("Message observers removed for task[%d:%d] (num:%d)"),
+			TaskIdx.InstanceIndex, TaskIdx.ExecutionIndex, NumRemoved);
+	}
+}
+
+void UBehaviorTreeComponent::UnregisterMessageObserversFrom(const class UBTTaskNode* TaskNode)
+{
+	if (TaskNode)
+	{
+		const FBehaviorTreeInstance& ActiveInstance = InstanceStack.Last();
+
+		FBTNodeIndex NodeIdx;
+		NodeIdx.ExecutionIndex = TaskNode->GetExecutionIndex();
+		NodeIdx.InstanceIndex = FindInstanceContainingNode(TaskNode);
+		
+		UnregisterMessageObserversFrom(NodeIdx);
+	}
+}
+
+#if STATS
+#define AUX_NODE_WRAPPER(cmd) \
+	DEC_MEMORY_STAT_BY(STAT_AI_BehaviorTree_InstanceMemory, InstanceInfo.GetAllocatedSize()); \
+	cmd; \
+	INC_MEMORY_STAT_BY(STAT_AI_BehaviorTree_InstanceMemory, InstanceInfo.GetAllocatedSize());
+
+#else
+#define AUX_NODE_WRAPPER(cmd) cmd;
+#endif // STATS
+
+void UBehaviorTreeComponent::RegisterParallelTask(const class UBTTaskNode* TaskNode)
+{
+	FBehaviorTreeInstance& InstanceInfo = InstanceStack[ActiveInstanceIdx];
+	AUX_NODE_WRAPPER(InstanceInfo.ParallelTasks.Add( FBehaviorTreeParallelTask(TaskNode, EBTTaskStatus::Active) ));
+	
+	UE_VLOG(GetOwner(), LogBehaviorTree, Verbose, TEXT("Parallel task: %s added to active list"),
+		*UBehaviorTreeTypes::DescribeNodeHelper(TaskNode));
+
+	if (InstanceInfo.ActiveNode == TaskNode)
+	{
+		// switch to inactive state, so it could start background tree
+		InstanceInfo.ActiveNodeType = EBTActiveNode::InactiveTask;
+	}
+}
+
+void UBehaviorTreeComponent::UnregisterParallelTask(const class UBTTaskNode* TaskNode, uint16 InstanceIdx)
+{
+	bool bShouldUpdate = false;
+	if (InstanceStack.IsValidIndex(InstanceIdx))
+	{
+		FBehaviorTreeInstance& InstanceInfo = InstanceStack[InstanceIdx];
+		for (int32 iTask = InstanceInfo.ParallelTasks.Num() - 1; iTask >= 0; iTask--)
+		{
+			FBehaviorTreeParallelTask& ParallelInfo = InstanceInfo.ParallelTasks[iTask];
+			if (ParallelInfo.TaskNode == TaskNode)
+			{
+				UE_VLOG(GetOwner(), LogBehaviorTree, Verbose, TEXT("Parallel task: %s removed from active list"),
+					*UBehaviorTreeTypes::DescribeNodeHelper(TaskNode));
+
+				InstanceInfo.ParallelTasks.RemoveAt(iTask);
+				bShouldUpdate = true;
+				break;
+			}
+		}
+	}
+
+	if (bShouldUpdate)
+	{
+		UpdateAbortingTasks();
+	}
+}
+
+#undef AUX_NODE_WRAPPER
+
+void UBehaviorTreeComponent::UpdateAbortingTasks()
+{
+	bWaitingForAbortingTasks = InstanceStack.Num() ? (InstanceStack.Last().ActiveNodeType == EBTActiveNode::AbortingTask) : false;
+
+	for (int32 i = 0; i < InstanceStack.Num() && !bWaitingForAbortingTasks; i++)
+	{
+		FBehaviorTreeInstance& InstanceInfo = InstanceStack[i];
+		for (int32 iTask = InstanceInfo.ParallelTasks.Num() - 1; iTask >= 0; iTask--)
+		{
+			FBehaviorTreeParallelTask& ParallelInfo = InstanceInfo.ParallelTasks[iTask];
+			if (ParallelInfo.Status == EBTTaskStatus::Aborting)
+			{
+				bWaitingForAbortingTasks = true;
+				break;
+			}
+		}
+	}
+}
+
+bool UBehaviorTreeComponent::PushInstance(class UBehaviorTree* TreeAsset)
+{
+	if (TreeAsset == NULL)
+	{
+		return false;
+	}
+
+	// check if blackboard class match
+	if (TreeAsset->BlackboardAsset && BlackboardComp && !BlackboardComp->IsCompatibileWith(TreeAsset->BlackboardAsset))
+	{
+		UE_VLOG(GetOwner(), LogBehaviorTree, Warning, TEXT("Failed to execute tree %s: blackboard %s is not compatibile with current: %s!"),
+			*GetNameSafe(TreeAsset), *GetNameSafe(TreeAsset->BlackboardAsset), *GetNameSafe(BlackboardComp->GetBlackboardAsset()));
+
+		return false;
+	}
+
+	UBehaviorTreeManager* BTManager = GetWorld()->GetBehaviorTreeManager();
+	if (BTManager == NULL)
+	{
+		UE_VLOG(GetOwner(), LogBehaviorTree, Warning, TEXT("Failed to execute tree %s: behavior tree manager not found!"), *GetNameSafe(TreeAsset));
+		return false;
+	}
+
+	// check if parent node allows it
+	const UBTNode* ActiveNode = GetActiveNode();
+	const UBTCompositeNode* ActiveParent = ActiveNode ? ActiveNode->GetParentNode() : NULL;
+	if (ActiveParent)
+	{
+		uint8* ParentMemory = GetNodeMemory((UBTNode*)ActiveParent, InstanceStack.Num() - 1);
+		int32 ChildIdx = ActiveParent->GetChildIndex(ActiveNode);
+
+		const bool bIsAllowed = ActiveParent->CanPushSubtree(this, ParentMemory, ChildIdx);
+		if (!bIsAllowed)
+		{
+			UE_VLOG(GetOwner(), LogBehaviorTree, Warning, TEXT("Failed to execute tree %s: parent of active node does not allow it! (%s)"),
+				*GetNameSafe(TreeAsset), *UBehaviorTreeTypes::DescribeNodeHelper(ActiveParent));
+			return false;
+		}
+	}
+
+	UBTCompositeNode* RootNode = NULL;
+	uint16 InstanceMemorySize = 0;
+
+	const bool bLoaded = BTManager->LoadTree(TreeAsset, RootNode, InstanceMemorySize);
+	if (bLoaded)
+	{
+		FBehaviorTreeInstance NewInstance(InstanceMemorySize);
+		NewInstance.InstanceIdIndex = UpdateInstanceId(TreeAsset, ActiveNode, InstanceStack.Num() - 1);
+		NewInstance.RootNode = RootNode;
+		NewInstance.ActiveNode = NULL;
+		NewInstance.ActiveNodeType = EBTActiveNode::Composite;
+
+		// initialize memory and node instances
+		FBehaviorTreeInstanceId& InstanceInfo = KnownInstances[NewInstance.InstanceIdIndex];
+		int32 NodeInstanceIndex = InstanceInfo.FirstNodeInstance;
+		NewInstance.Initialize(this, RootNode, NodeInstanceIndex);
+		NewInstance.InjectNodes(this, RootNode, NodeInstanceIndex);
+
+		INC_DWORD_STAT(STAT_AI_BehaviorTree_NumInstances);
+		InstanceStack.Push(NewInstance);
+		ActiveInstanceIdx = InstanceStack.Num() - 1;
+
+		// start root level services now (they won't be removed on looping tree anyway)
+		for (int32 i = 0; i < RootNode->Services.Num(); i++)
+		{
+			UBTService* ServiceNode = RootNode->Services[i];
+			uint8* NodeMemory = (uint8*)ServiceNode->GetNodeMemory<uint8>(InstanceStack[ActiveInstanceIdx]);
+
+			InstanceStack[ActiveInstanceIdx].ActiveAuxNodes.Add(ServiceNode);
+			ServiceNode->WrappedOnBecomeRelevant(this, NodeMemory);
+		}
+
+		FBehaviorTreeDelegates::OnTreeStarted.Broadcast(this, TreeAsset);
+
+		// start new task
+		RequestExecution(RootNode, ActiveInstanceIdx, RootNode, 0, EBTNodeResult::InProgress);
+		return true;
+	}
+
+	return false;
+}
+
+uint8 UBehaviorTreeComponent::UpdateInstanceId(class UBehaviorTree* TreeAsset, const UBTNode* OriginNode, int32 OriginInstanceIdx)
+{
+	FBehaviorTreeInstanceId InstanceId;
+	InstanceId.TreeAsset = TreeAsset;
+
+	// build path from origin node
+	{
+		const uint16 ExecutionIndex = OriginNode ? OriginNode->GetExecutionIndex() : MAX_uint16;
+		InstanceId.Path.Add(ExecutionIndex);
+	}
+
+	for (int32 i = OriginInstanceIdx - 1; i >= 0; i--)
+	{
+		const uint16 ExecutionIndex = InstanceStack[i].ActiveNode ? InstanceStack[i].ActiveNode->GetExecutionIndex() : MAX_uint16;
+		InstanceId.Path.Add(ExecutionIndex);
+	}
+
+	// try to find matching existing Id
+	for (int32 i = 0; i < KnownInstances.Num(); i++)
+	{
+		if (KnownInstances[i] == InstanceId)
+		{
+			return i;
+		}
+	}
+
+	// add new one
+	InstanceId.FirstNodeInstance = NodeInstances.Num();
+
+	const int32 NewIndex = KnownInstances.Add(InstanceId);
+	check(NewIndex < MAX_uint8);
+	return NewIndex;
+}
+
+int32 UBehaviorTreeComponent::FindInstanceContainingNode(const UBTNode* Node) const
+{
+	int32 InstanceIdx = INDEX_NONE;
+
+	const UBTNode* TemplateNode = FindTemplateNode(Node);
+	if (TemplateNode && InstanceStack.Num())
+	{
+		if (InstanceStack[ActiveInstanceIdx].ActiveNode != TemplateNode)
+		{
+			const UBTNode* RootNode = TemplateNode;
+			while (RootNode->GetParentNode())
+			{
+				RootNode = RootNode->GetParentNode();
+			}
+
+			for (int32 i = 0; i < InstanceStack.Num(); i++)
+			{
+				if (InstanceStack[i].RootNode == RootNode)
+				{
+					InstanceIdx = i;
+					break;
+				}
+			}
+		}
+		else
+		{
+			InstanceIdx = ActiveInstanceIdx;
+		}
+	}
+
+	return InstanceIdx;
+}
+
+class UBTNode* UBehaviorTreeComponent::FindTemplateNode(const class UBTNode* Node) const
+{
+	if (Node == NULL || !Node->IsInstanced() || Node->GetParentNode() == NULL)
+	{
+		return (UBTNode*)Node;
+	}
+
+	UBTCompositeNode* ParentNode = Node->GetParentNode();
+	for (int32 iChild = 0; iChild < ParentNode->Children.Num(); iChild++)
+	{
+		FBTCompositeChild& ChildInfo = ParentNode->Children[iChild];
+
+		if (ChildInfo.ChildTask && ChildInfo.ChildTask->GetExecutionIndex() == Node->GetExecutionIndex())
+		{
+			return ChildInfo.ChildTask;
+		}
+
+		for (int32 i = 0; i < ChildInfo.Decorators.Num(); i++)
+		{
+			if (ChildInfo.Decorators[i]->GetExecutionIndex() == Node->GetExecutionIndex())
+			{
+				return ChildInfo.Decorators[i];
+			}
+		}
+	}
+
+	for (int32 i = 0; i < ParentNode->Services.Num(); i++)
+	{
+		if (ParentNode->Services[i]->GetExecutionIndex() == Node->GetExecutionIndex())
+		{
+			return ParentNode->Services[i];
+		}
+	}
+
+	return NULL;
+}
+
+uint8* UBehaviorTreeComponent::GetNodeMemory(UBTNode* Node, int32 InstanceIdx) const
+{
+	return InstanceStack.IsValidIndex(InstanceIdx) ? (uint8*)Node->GetNodeMemory<uint8>(InstanceStack[InstanceIdx]) : NULL;
+}
+
+FString UBehaviorTreeComponent::GetDebugInfoString() const 
+{ 
+	FString DebugInfo;
+	for (int32 i = 0; i < InstanceStack.Num(); i++)
+	{
+		const FBehaviorTreeInstance& InstanceData = InstanceStack[i];
+		const FBehaviorTreeInstanceId& InstanceInfo = KnownInstances[InstanceData.InstanceIdIndex];
+		DebugInfo += FString::Printf(TEXT("Behavior tree: %s\n"), *GetNameSafe(InstanceInfo.TreeAsset));
+
+		UBTNode* Node = InstanceData.ActiveNode;
+		FString NodeTrace;
+
+		while (Node)
+		{
+			uint8* NodeMemory = (uint8*)(Node->GetNodeMemory<uint8>(InstanceData));
+			NodeTrace = FString::Printf(TEXT("  %s\n"), *Node->GetRuntimeDescription(this, NodeMemory, EBTDescriptionVerbosity::Basic)) + NodeTrace;
+			Node = Node->GetParentNode();
+		}
+
+		DebugInfo += NodeTrace;
+	}
+
+	return DebugInfo;
+}
+
+FString UBehaviorTreeComponent::DescribeActiveTasks() const
+{
+	FString ActiveTask(TEXT("None"));
+	if (InstanceStack.Num())
+	{
+		const FBehaviorTreeInstance& LastInstance = InstanceStack.Last();
+		if (LastInstance.ActiveNodeType == EBTActiveNode::ActiveTask)
+		{
+			ActiveTask = UBehaviorTreeTypes::DescribeNodeHelper(LastInstance.ActiveNode);
+		}
+	}
+
+	FString ParallelTasks;
+	for (int32 i = 0; i < InstanceStack.Num(); i++)
+	{
+		const FBehaviorTreeInstance& InstanceInfo = InstanceStack[i];
+		for (int32 iParallel = 0; iParallel < InstanceInfo.ParallelTasks.Num(); iParallel++)
+		{
+			const FBehaviorTreeParallelTask& ParallelInfo = InstanceInfo.ParallelTasks[iParallel];
+			if (ParallelInfo.Status == EBTTaskStatus::Active)
+			{
+				ParallelTasks += UBehaviorTreeTypes::DescribeNodeHelper(ParallelInfo.TaskNode);
+				ParallelTasks += TEXT(", ");
+			}
+		}
+	}
+
+	if (ParallelTasks.Len() > 0)
+	{
+		ActiveTask += TEXT(" (");
+		ActiveTask += ParallelTasks.LeftChop(2);
+		ActiveTask += TEXT(')');
+	}
+
+	return ActiveTask;
+}
+
+FString UBehaviorTreeComponent::DescribeActiveTrees() const
+{
+	FString Assets;
+	for (int32 i = 0; i < InstanceStack.Num(); i++)
+	{
+		const FBehaviorTreeInstanceId& InstanceInfo = KnownInstances[InstanceStack[i].InstanceIdIndex];
+		Assets += InstanceInfo.TreeAsset->GetName();
+		Assets += TEXT(", ");
+	}
+
+	return Assets.Len() ? Assets.LeftChop(2) : TEXT("None");
+}
+
+#if ENABLE_VISUAL_LOG
+void UBehaviorTreeComponent::DescribeSelfToVisLog(struct FVisLogEntry* Snapshot) const
+{
+	if (IsPendingKill())
+	{
+		return;
+	}
+	
+	Super::DescribeSelfToVisLog(Snapshot);
+
+	for (int32 i = 0; i < InstanceStack.Num(); i++)
+	{
+		const FBehaviorTreeInstance& InstanceInfo = InstanceStack[i];
+		const FBehaviorTreeInstanceId& InstanceId = KnownInstances[InstanceInfo.InstanceIdIndex];
+
+		FVisLogEntry::FStatusCategory StatusCategory;
+		StatusCategory.Category = FString::Printf(TEXT("BehaviorTree %d (asset: %s)"), i, *GetNameSafe(InstanceId.TreeAsset));
+		
+		TArray<FString> Descriptions;
+		UBTNode* Node = InstanceInfo.ActiveNode;
+		while (Node)
+		{
+			uint8* NodeMemory = (uint8*)(Node->GetNodeMemory<uint8>(InstanceInfo));
+			Descriptions.Add(Node->GetRuntimeDescription(this, NodeMemory, EBTDescriptionVerbosity::Detailed));
+		
+			Node = Node->GetParentNode();
+		}
+
+		for (int32 iDesc = Descriptions.Num() - 1; iDesc >= 0; iDesc--)
+		{
+			int32 SplitIdx = INDEX_NONE;
+			if (Descriptions[iDesc].FindChar(TEXT(','), SplitIdx))
+			{
+				const FString KeyDesc = Descriptions[iDesc].Left(SplitIdx);
+				const FString ValueDesc = Descriptions[iDesc].Mid(SplitIdx + 1).Trim();
+
+				StatusCategory.Add(KeyDesc, ValueDesc);
+			}
+			else
+			{
+				StatusCategory.Add(Descriptions[iDesc], TEXT(""));
+			}
+		}
+
+		if (StatusCategory.Data.Num() == 0)
+		{
+			StatusCategory.Add(TEXT("root"), TEXT("not initialized"));
+		}
+
+		Snapshot->Status.Add(StatusCategory);
+	}
+}
+#endif // ENABLE_VISUAL_LOG
+
+void UBehaviorTreeComponent::StoreDebuggerExecutionStep(EBTExecutionSnap::Type SnapType)
+{
+#if USE_BEHAVIORTREE_DEBUGGER
+	FBehaviorTreeExecutionStep CurrentStep;
+	CurrentStep.StepIndex = DebuggerSteps.Num() ? DebuggerSteps.Last().StepIndex + 1 : 0;
+	CurrentStep.TimeStamp = GetWorld()->GetTimeSeconds();
+	CurrentStep.BlackboardValues = SearchStartBlackboard;
+
+	for (int32 i = 0; i < InstanceStack.Num(); i++)
+	{
+		const FBehaviorTreeInstance& ActiveInstance = InstanceStack[i];
+		
+		FBehaviorTreeDebuggerInstance StoreInfo;
+		StoreDebuggerInstance(StoreInfo, i, SnapType);
+		CurrentStep.InstanceStack.Add(StoreInfo);
+	}
+
+	for (int32 i = RemovedInstances.Num() - 1; i >= 0; i--)
+	{
+		CurrentStep.InstanceStack.Add(RemovedInstances[i]);
+	}
+
+	CurrentSearchFlow.Reset();
+	CurrentRestarts.Reset();
+	RemovedInstances.Reset();
+
+	UBehaviorTreeManager* ManagerCDO = (UBehaviorTreeManager*)UBehaviorTreeManager::StaticClass()->GetDefaultObject();
+	while (DebuggerSteps.Num() >= ManagerCDO->MaxDebuggerSteps)
+	{
+		DebuggerSteps.RemoveAt(0);
+	}
+	DebuggerSteps.Add(CurrentStep);
+#endif
+}
+
+void UBehaviorTreeComponent::StoreDebuggerInstance(FBehaviorTreeDebuggerInstance& InstanceInfo, uint16 InstanceIdx, EBTExecutionSnap::Type SnapType) const
+{
+#if USE_BEHAVIORTREE_DEBUGGER
+	if (!InstanceStack.IsValidIndex(InstanceIdx))
+	{
+		return;
+	}
+
+	const FBehaviorTreeInstance& ActiveInstance = InstanceStack[InstanceIdx];
+	const FBehaviorTreeInstanceId& ActiveInstanceInfo = KnownInstances[ActiveInstance.InstanceIdIndex];
+	InstanceInfo.TreeAsset = ActiveInstanceInfo.TreeAsset;
+	InstanceInfo.RootNode = ActiveInstance.RootNode;
+
+	if (SnapType == EBTExecutionSnap::Regular)
+	{
+		// traverse execution path
+		UBTNode* StoreNode = ActiveInstance.ActiveNode ? ActiveInstance.ActiveNode : ActiveInstance.RootNode;
+		while (StoreNode)
+		{
+			InstanceInfo.ActivePath.Add(StoreNode->GetExecutionIndex());
+			StoreNode = StoreNode->GetParentNode();
+		}
+
+		// add aux nodes
+		for (int32 i = 0; i < ActiveInstance.ActiveAuxNodes.Num(); i++)
+		{
+			InstanceInfo.AdditionalActiveNodes.Add(ActiveInstance.ActiveAuxNodes[i]->GetExecutionIndex());
+		}
+
+		// add active parallels
+		for (int32 i = 0; i < ActiveInstance.ParallelTasks.Num(); i++)
+		{
+			const FBehaviorTreeParallelTask& TaskInfo = ActiveInstance.ParallelTasks[i];
+			InstanceInfo.AdditionalActiveNodes.Add(TaskInfo.TaskNode->GetExecutionIndex());
+		}
+
+		// runtime values
+		StoreDebuggerRuntimeValues(InstanceInfo.RuntimeDesc, ActiveInstance.RootNode, InstanceIdx);
+	}
+
+	// handle restart triggers
+	if (CurrentRestarts.IsValidIndex(InstanceIdx))
+	{
+		InstanceInfo.PathFromPrevious = CurrentRestarts[InstanceIdx];
+	}
+
+	// store search flow, but remove nodes on execution path
+	if (CurrentSearchFlow.IsValidIndex(InstanceIdx))
+	{
+		for (int32 i = 0; i < CurrentSearchFlow[InstanceIdx].Num(); i++)
+		{
+			if (!InstanceInfo.ActivePath.Contains(CurrentSearchFlow[InstanceIdx][i].ExecutionIndex))
+			{
+				InstanceInfo.PathFromPrevious.Add(CurrentSearchFlow[InstanceIdx][i]);
+			}
+		}
+	}
+#endif
+}
+
+void UBehaviorTreeComponent::StoreDebuggerRemovedInstance(uint16 InstanceIdx) const
+{
+#if USE_BEHAVIORTREE_DEBUGGER
+	FBehaviorTreeDebuggerInstance StoreInfo;
+	StoreDebuggerInstance(StoreInfo, InstanceIdx, EBTExecutionSnap::OutOfNodes);
+
+	RemovedInstances.Add(StoreInfo);
+#endif
+}
+
+void UBehaviorTreeComponent::StoreDebuggerSearchStep(const class UBTNode* Node, uint16 InstanceIdx, EBTNodeResult::Type NodeResult) const
+{
+#if USE_BEHAVIORTREE_DEBUGGER
+	if (Node && NodeResult != EBTNodeResult::InProgress && NodeResult != EBTNodeResult::Aborted)
+	{
+		FBehaviorTreeDebuggerInstance::FNodeFlowData FlowInfo;
+		FlowInfo.ExecutionIndex = Node->GetExecutionIndex();
+		FlowInfo.bPassed = (NodeResult == EBTNodeResult::Succeeded);
+		FlowInfo.bOptional = (NodeResult == EBTNodeResult::Optional);
+		
+		if (CurrentSearchFlow.Num() < (InstanceIdx + 1))
+		{
+			CurrentSearchFlow.SetNum(InstanceIdx + 1);
+		}
+
+		if (CurrentSearchFlow[InstanceIdx].Num() == 0 || CurrentSearchFlow[InstanceIdx].Last().ExecutionIndex != FlowInfo.ExecutionIndex)
+		{
+			CurrentSearchFlow[InstanceIdx].Add(FlowInfo);
+		}
+	}
+#endif
+}
+
+void UBehaviorTreeComponent::StoreDebuggerSearchStep(const class UBTNode* Node, uint16 InstanceIdx, bool bPassed) const
+{
+#if USE_BEHAVIORTREE_DEBUGGER
+	if (Node && !bPassed)
+	{
+		FBehaviorTreeDebuggerInstance::FNodeFlowData FlowInfo;
+		FlowInfo.ExecutionIndex = Node->GetExecutionIndex();
+		FlowInfo.bPassed = bPassed;
+
+		if (CurrentSearchFlow.Num() < (InstanceIdx + 1))
+		{
+			CurrentSearchFlow.SetNum(InstanceIdx + 1);
+		}
+
+		CurrentSearchFlow[InstanceIdx].Add(FlowInfo);
+	}
+#endif
+}
+
+void UBehaviorTreeComponent::StoreDebuggerRestart(const class UBTNode* Node, uint16 InstanceIdx, bool bAllowed)
+{
+#if USE_BEHAVIORTREE_DEBUGGER
+	if (Node)
+	{
+		FBehaviorTreeDebuggerInstance::FNodeFlowData FlowInfo;
+		FlowInfo.ExecutionIndex = Node->GetExecutionIndex();
+		FlowInfo.bTrigger = bAllowed;
+		FlowInfo.bDiscardedTrigger = !bAllowed;
+
+		if (CurrentRestarts.Num() < (InstanceIdx + 1))
+		{
+			CurrentRestarts.SetNum(InstanceIdx + 1);
+		}
+
+		CurrentRestarts[InstanceIdx].Add(FlowInfo);
+	}
+#endif
+}
+
+void UBehaviorTreeComponent::StoreDebuggerRuntimeValues(TArray<FString>& RuntimeDescriptions, class UBTNode* RootNode, uint16 InstanceIdx) const
+{
+#if USE_BEHAVIORTREE_DEBUGGER
+	if (!InstanceStack.IsValidIndex(InstanceIdx))
+	{
+		return;
+	}
+
+	const FBehaviorTreeInstance& InstanceInfo = InstanceStack[InstanceIdx];
+
+	TArray<FString> RuntimeValues;
+	for (UBTNode* Node = RootNode; Node; Node = Node->GetNextNode())
+	{
+		uint8* NodeMemory = (uint8*)Node->GetNodeMemory<uint8>(InstanceInfo);
+
+		RuntimeValues.Reset();
+		Node->DescribeRuntimeValues(this, NodeMemory, EBTDescriptionVerbosity::Basic, RuntimeValues);
+
+		FString ComposedDesc;
+		for (int32 i = 0; i < RuntimeValues.Num(); i++)
+		{
+			if (ComposedDesc.Len())
+			{
+				ComposedDesc.AppendChar(TEXT('\n'));
+			}
+
+			ComposedDesc += RuntimeValues[i];
+		}
+
+		RuntimeDescriptions.SetNum(Node->GetExecutionIndex() + 1);
+		RuntimeDescriptions[Node->GetExecutionIndex()] = ComposedDesc;
+	}
+#endif
+}
+
+void UBehaviorTreeComponent::UpdateDebuggerAfterExecution(const class UBTTaskNode* TaskNode, uint16 InstanceIdx) const
+{
+#if USE_BEHAVIORTREE_DEBUGGER
+	FBehaviorTreeExecutionStep& CurrentStep = DebuggerSteps.Last();
+
+	// store runtime values
+	TArray<FString> RuntimeValues;
+	const FBehaviorTreeInstance& InstanceToUpdate = InstanceStack[InstanceIdx];
+	uint8* NodeMemory = (uint8*)TaskNode->GetNodeMemory<uint8>(InstanceToUpdate);
+	TaskNode->DescribeRuntimeValues(this, NodeMemory, EBTDescriptionVerbosity::Basic, RuntimeValues);
+
+	FString ComposedDesc;
+	for (int32 i = 0; i < RuntimeValues.Num(); i++)
+	{
+		if (ComposedDesc.Len())
+		{
+			ComposedDesc.AppendChar(TEXT('\n'));
+		}
+
+		ComposedDesc += RuntimeValues[i];
+	}
+
+	// accessing RuntimeDesc should never be out of bounds (active task MUST be part of active instance)
+	const uint16& ExecutionIndex = TaskNode->GetExecutionIndex();
+	CurrentStep.InstanceStack[InstanceIdx].RuntimeDesc[ExecutionIndex] = ComposedDesc;
+#endif
+}
+
+void UBehaviorTreeComponent::StoreDebuggerBlackboard(TArray<FString>& BlackboardValueDesc) const
+{
+#if USE_BEHAVIORTREE_DEBUGGER
+	if (BlackboardComp && BlackboardComp->HasValidAsset())
+	{
+		const int32 NumKeys = BlackboardComp->GetNumKeys();
+		BlackboardValueDesc.Reset();
+		BlackboardValueDesc.AddZeroed(NumKeys);
+
+		for (int32 i = 0; i < NumKeys; i++)
+		{
+			BlackboardValueDesc[i] = BlackboardComp->DescribeKeyValue(i, EBlackboardDescription::OnlyValue);
+			
+			if (BlackboardValueDesc[i].Len() == 0)
+			{
+				BlackboardValueDesc[i] = TEXT("n/a");
+			}
+		}
+	}
+#endif
+}
diff --git a/Engine/Source/Runtime/AIModule/Private/BehaviorTree/BehaviorTreeManager.cpp b/Engine/Source/Runtime/AIModule/Private/BehaviorTree/BehaviorTreeManager.cpp
new file mode 100644
--- /dev/null
+++ b/Engine/Source/Runtime/AIModule/Private/BehaviorTree/BehaviorTreeManager.cpp
@@ -0,0 +1,236 @@
+// Copyright 1998-2014 Epic Games, Inc. All Rights Reserved.
+
+#include "AIModulePrivate.h"
+#include "BehaviorTree/Tasks/BTTask_RunBehavior.h"
+#include "BehaviorTree/BlackboardComponent.h"
+#include "BehaviorTree/BehaviorTreeManager.h"
+
+DEFINE_STAT(STAT_AI_BehaviorTree_Tick);
+DEFINE_STAT(STAT_AI_BehaviorTree_LoadTime);
+DEFINE_STAT(STAT_AI_BehaviorTree_SearchTime);
+DEFINE_STAT(STAT_AI_BehaviorTree_ExecutionTime);
+DEFINE_STAT(STAT_AI_BehaviorTree_AuxUpdateTime);
+DEFINE_STAT(STAT_AI_BehaviorTree_NumTemplates);
+DEFINE_STAT(STAT_AI_BehaviorTree_NumInstances);
+DEFINE_STAT(STAT_AI_BehaviorTree_InstanceMemory);
+
+UBehaviorTreeManager::UBehaviorTreeManager(const class FPostConstructInitializeProperties& PCIP) : Super(PCIP)
+{
+	MaxDebuggerSteps = 100;
+}
+
+void UBehaviorTreeManager::FinishDestroy()
+{
+	SET_DWORD_STAT(STAT_AI_BehaviorTree_NumTemplates, 0);
+
+	Super::FinishDestroy();
+}
+
+bool UBehaviorTreeManager::IsBehaviorTreeUsageEnabled()
+{
+	struct FCheckMeStaticly
+	{
+		bool bBehaviorTreeEditorEnabled;
+		bool bBehaviorTreeUsageEnabled;
+		FCheckMeStaticly() : bBehaviorTreeEditorEnabled(false), bBehaviorTreeUsageEnabled(true)
+		{
+#if WITH_EDITOR
+			bool bBehaviorTreeEditorEnabledFromUserSettings = false;
+			GConfig->GetBool(TEXT("/Script/UnrealEd.EditorExperimentalSettings"), TEXT("bBehaviorTreeEditor"), bBehaviorTreeEditorEnabledFromUserSettings, GEditorUserSettingsIni);
+
+			GConfig->GetBool(TEXT("BehaviorTreesEd"), TEXT("BehaviorTreeEditorEnabled"), bBehaviorTreeEditorEnabled, GEngineIni);
+			bBehaviorTreeEditorEnabled = bBehaviorTreeEditorEnabled || bBehaviorTreeEditorEnabledFromUserSettings;
+			if (bBehaviorTreeUsageEnabled && bBehaviorTreeEditorEnabled && !FModuleManager::Get().IsModuleLoaded(TEXT("BehaviorTreeEditor")))
+			{
+				// load this module earlier, before any access to BehaviorTree assets
+				FModuleManager::Get().LoadModule(TEXT("BehaviorTreeEditor"));
+			}
+#endif
+		}
+	};
+
+	static FCheckMeStaticly CheckMeStaticly;
+	return CheckMeStaticly.bBehaviorTreeUsageEnabled;
+}
+
+int32 UBehaviorTreeManager::GetAlignedDataSize(int32 Size)
+{
+	// round to 4 bytes
+	return ((Size + 3) & ~3);
+}
+
+struct FNodeInitializationData
+{
+	UBTNode* Node;
+	UBTCompositeNode* ParentNode;
+	uint16 ExecutionIndex;
+	uint16 DataSize;
+	uint16 SpecialDataSize;
+	uint8 TreeDepth;
+
+	FNodeInitializationData() {}
+	FNodeInitializationData(UBTNode* InNode, UBTCompositeNode* InParentNode,
+		uint16 InExecutionIndex, uint8 InTreeDepth, uint16 NodeMemory, uint16 SpecialNodeMemory = 0)
+		: Node(InNode), ParentNode(InParentNode), ExecutionIndex(InExecutionIndex), TreeDepth(InTreeDepth)
+	{
+		SpecialDataSize = UBehaviorTreeManager::GetAlignedDataSize(SpecialNodeMemory);
+
+		const uint16 NodeMemorySize = NodeMemory + SpecialDataSize;
+		DataSize = (NodeMemorySize <= 2) ? NodeMemorySize : UBehaviorTreeManager::GetAlignedDataSize(NodeMemorySize);
+	}
+
+	struct FMemorySort
+	{
+		FORCEINLINE bool operator()(const FNodeInitializationData& A, const FNodeInitializationData& B) const
+		{
+			return A.DataSize > B.DataSize;
+		}
+	};
+};
+
+static void InitializeNodeHelper(UBTCompositeNode* ParentNode, UBTNode* NodeOb,
+	uint8 TreeDepth, uint16& ExecutionIndex, TArray<FNodeInitializationData>& InitList,
+	class UBehaviorTree* TreeAsset, UObject* NodeOuter)
+{
+	// special case: subtrees
+	UBTTask_RunBehavior* SubtreeTask = Cast<UBTTask_RunBehavior>(NodeOb);
+	if (SubtreeTask)
+	{
+		ExecutionIndex += SubtreeTask->GetInjectedNodesCount();
+	}
+
+	InitList.Add(FNodeInitializationData(NodeOb, ParentNode, ExecutionIndex, TreeDepth, NodeOb->GetInstanceMemorySize(), NodeOb->GetSpecialMemorySize()));
+	NodeOb->InitializeFromAsset(TreeAsset);
+	ExecutionIndex++;
+
+	UBTCompositeNode* CompositeOb = Cast<UBTCompositeNode>(NodeOb);
+	if (CompositeOb)
+	{
+		for (int32 i = 0; i < CompositeOb->Services.Num(); i++)
+		{
+			UBTService* Service = Cast<UBTService>(StaticDuplicateObject(CompositeOb->Services[i], NodeOuter, TEXT("None")));;
+			CompositeOb->Services[i] = Service;
+
+			InitList.Add(FNodeInitializationData(Service, CompositeOb, ExecutionIndex, TreeDepth,
+				Service->GetInstanceMemorySize(), Service->GetSpecialMemorySize()));
+
+			Service->InitializeFromAsset(TreeAsset);
+			ExecutionIndex++;
+		}
+
+		for (int32 i = 0; i < CompositeOb->Children.Num(); i++)
+		{
+			FBTCompositeChild& ChildInfo = CompositeOb->Children[i];
+			for (int32 j = 0; j < ChildInfo.Decorators.Num(); j++)
+			{
+				UBTDecorator* Decorator = Cast<UBTDecorator>(StaticDuplicateObject(ChildInfo.Decorators[j], NodeOuter, TEXT("None")));
+				ChildInfo.Decorators[j] = Decorator;
+
+				InitList.Add(FNodeInitializationData(Decorator, CompositeOb, ExecutionIndex, TreeDepth,
+					Decorator->GetInstanceMemorySize(), Decorator->GetSpecialMemorySize()));
+
+				Decorator->InitializeFromAsset(TreeAsset);
+				Decorator->InitializeDecorator(i);
+				ExecutionIndex++;
+			}
+
+			UBTNode* ChildNode = NULL;
+			
+			if (ChildInfo.ChildComposite)
+			{
+				ChildInfo.ChildComposite = Cast<UBTCompositeNode>(StaticDuplicateObject(ChildInfo.ChildComposite, NodeOuter, TEXT("None")));
+				ChildNode = ChildInfo.ChildComposite;
+			}
+			else if (ChildInfo.ChildTask)
+			{
+				ChildInfo.ChildTask = Cast<UBTTaskNode>(StaticDuplicateObject(ChildInfo.ChildTask, NodeOuter, TEXT("None")));
+				ChildNode = ChildInfo.ChildTask;
+			}
+
+			if (ChildNode)
+			{
+				InitializeNodeHelper(CompositeOb, ChildNode, TreeDepth + 1, ExecutionIndex, InitList, TreeAsset, NodeOuter);
+			}
+		}
+
+		CompositeOb->InitializeComposite(InitList.Num() - 1);
+	}
+}
+
+bool UBehaviorTreeManager::LoadTree(class UBehaviorTree* Asset, UBTCompositeNode*& Root, uint16& InstanceMemorySize)
+{
+	SCOPE_CYCLE_COUNTER(STAT_AI_BehaviorTree_LoadTime);
+
+	for (int32 i = 0; i < LoadedTemplates.Num(); i++)
+	{
+		FBehaviorTreeTemplateInfo& TemplateInfo = LoadedTemplates[i];
+		if (TemplateInfo.Asset == Asset)
+		{
+			Root = TemplateInfo.Template;
+			InstanceMemorySize = TemplateInfo.InstanceMemorySize;
+			return true;
+		}
+	}
+
+	if (Asset->RootNode)
+	{
+		FBehaviorTreeTemplateInfo TemplateInfo;
+		TemplateInfo.Asset = Asset;
+		TemplateInfo.Template = Cast<UBTCompositeNode>(StaticDuplicateObject(Asset->RootNode, this, TEXT("None")));
+
+		TArray<FNodeInitializationData> InitList;
+		uint16 ExecutionIndex = 0;
+		InitializeNodeHelper(NULL, TemplateInfo.Template, 0, ExecutionIndex, InitList, Asset, this);
+
+#if USE_BEHAVIORTREE_DEBUGGER
+		// fill in information about next nodes in execution index, before sorting memory offsets
+		for (int32 i = 0; i < InitList.Num() - 1; i++)
+		{
+			InitList[i].Node->InitializeExecutionOrder(InitList[i + 1].Node);
+		}
+#endif
+
+		// sort nodes by memory size, so they can be packed better
+		// it still won't protect against structures, that are internally misaligned (-> uint8, uint32)
+		// but since all Engine level nodes are good... 
+		InitList.Sort(FNodeInitializationData::FMemorySort());
+		uint16 MemoryOffset = 0;
+		for (int32 i = 0; i < InitList.Num(); i++)
+		{
+			InitList[i].Node->InitializeNode(InitList[i].ParentNode, InitList[i].ExecutionIndex, InitList[i].SpecialDataSize + MemoryOffset, InitList[i].TreeDepth);
+			MemoryOffset += InitList[i].DataSize;
+		}
+		
+		TemplateInfo.InstanceMemorySize = MemoryOffset;
+
+		INC_DWORD_STAT(STAT_AI_BehaviorTree_NumTemplates);
+		LoadedTemplates.Add(TemplateInfo);
+		Root = TemplateInfo.Template;
+		InstanceMemorySize = TemplateInfo.InstanceMemorySize;
+		return true;
+	}
+
+	return false;
+}
+
+void UBehaviorTreeManager::InitializeMemoryHelper(const TArray<UBTDecorator*>& Nodes, TArray<uint16>& MemoryOffsets, int32& MemorySize)
+{
+	TArray<FNodeInitializationData> InitList;
+	for (int32 i = 0; i < Nodes.Num(); i++)
+	{
+		InitList.Add(FNodeInitializationData(Nodes[i], NULL, 0, 0, Nodes[i]->GetInstanceMemorySize(), Nodes[i]->GetSpecialMemorySize()));
+	}
+
+	InitList.Sort(FNodeInitializationData::FMemorySort());
+
+	uint16 MemoryOffset = 0;
+	MemoryOffsets.AddZeroed(Nodes.Num());
+
+	for (int32 i = 0; i < InitList.Num(); i++)
+	{
+		MemoryOffsets[i] = InitList[i].SpecialDataSize + MemoryOffset;
+		MemoryOffset += InitList[i].DataSize;
+	}
+
+	MemorySize = MemoryOffset;
+}
diff --git a/Engine/Source/Runtime/AIModule/Private/BehaviorTree/BehaviorTreeTypes.cpp b/Engine/Source/Runtime/AIModule/Private/BehaviorTree/BehaviorTreeTypes.cpp
new file mode 100644
--- /dev/null
+++ b/Engine/Source/Runtime/AIModule/Private/BehaviorTree/BehaviorTreeTypes.cpp
@@ -0,0 +1,308 @@
+// Copyright 1998-2014 Epic Games, Inc. All Rights Reserved.
+
+#include "AIModulePrivate.h"
+#include "BehaviorTree/Blackboard/BlackboardKeyAllTypes.h"
+#include "BehaviorTree/BTCompositeNode.h"
+#include "BehaviorTree/Tasks/BTTask_RunBehavior.h"
+#include "BehaviorTree/BehaviorTreeTypes.h"
+
+//----------------------------------------------------------------------//
+// FBehaviorTreeInstance
+//----------------------------------------------------------------------//
+void FBehaviorTreeInstance::Initialize(class UBehaviorTreeComponent* OwnerComp, UBTCompositeNode* Node, int32& InstancedIndex)
+{
+	if (Node == NULL)
+	{
+		return;
+	}
+
+	for (int32 i = 0; i < Node->Services.Num(); i++)
+	{
+		Node->Services[i]->InitializeForInstance(OwnerComp, Node->Services[i]->GetNodeMemory<uint8>(*this), InstancedIndex);
+	}
+
+	Node->InitializeForInstance(OwnerComp, Node->GetNodeMemory<uint8>(*this), InstancedIndex);
+
+	for (int32 iChild = 0; iChild < Node->Children.Num(); iChild++)
+	{
+		FBTCompositeChild& ChildInfo = Node->Children[iChild];
+
+		for (int32 i = 0; i < ChildInfo.Decorators.Num(); i++)
+		{
+			ChildInfo.Decorators[i]->InitializeForInstance(OwnerComp, ChildInfo.Decorators[i]->GetNodeMemory<uint8>(*this), InstancedIndex);
+		}
+
+		if (ChildInfo.ChildComposite)
+		{
+			Initialize(OwnerComp, ChildInfo.ChildComposite, InstancedIndex);
+		}
+		else if (ChildInfo.ChildTask)
+		{
+			ChildInfo.ChildTask->InitializeForInstance(OwnerComp, ChildInfo.ChildTask->GetNodeMemory<uint8>(*this), InstancedIndex);
+		}
+	}
+}
+
+void FBehaviorTreeInstance::InjectNodes(class UBehaviorTreeComponent* OwnerComp, UBTCompositeNode* Node, int32& InstancedIndex)
+{
+	if (Node == NULL)
+	{
+		return;
+	}
+
+	for (int32 iChild = 0; iChild < Node->Children.Num(); iChild++)
+	{
+		FBTCompositeChild& ChildInfo = Node->Children[iChild];
+		if (ChildInfo.ChildComposite)
+		{
+			InjectNodes(OwnerComp, ChildInfo.ChildComposite, InstancedIndex);
+		}
+		else
+		{
+			UBTTask_RunBehavior* InjectingTask = Cast<UBTTask_RunBehavior>(ChildInfo.ChildTask);
+			if (InjectingTask)
+			{
+				uint8* NodeMemory = InjectingTask->GetNodeMemory<uint8>(*this);
+				InjectingTask->InjectNodes(OwnerComp, NodeMemory, InstancedIndex);
+			}
+		}
+	}
+}
+
+void FBehaviorTreeInstance::Cleanup(class UBehaviorTreeComponent* OwnerComp)
+{
+	const FBehaviorTreeInstanceId& Info = OwnerComp->KnownInstances[InstanceIdIndex];
+	if (Info.FirstNodeInstance >= 0)
+	{
+		const int32 LastNodeIdx = OwnerComp->KnownInstances.IsValidIndex(InstanceIdIndex + 1) ?
+			OwnerComp->KnownInstances[InstanceIdIndex + 1].FirstNodeInstance :
+			OwnerComp->NodeInstances.Num();
+
+		for (int32 Idx = Info.FirstNodeInstance; Idx < LastNodeIdx; Idx++)
+		{
+			OwnerComp->NodeInstances[Idx]->OnInstanceDestroyed(OwnerComp);
+		}
+	}
+}
+
+//----------------------------------------------------------------------//
+// FBTNodeIndex
+//----------------------------------------------------------------------//
+
+bool FBTNodeIndex::TakesPriorityOver(const FBTNodeIndex& Other) const
+{
+	// instance closer to root is more important
+	if (InstanceIndex != Other.InstanceIndex)
+	{
+		return InstanceIndex < Other.InstanceIndex;
+	}
+
+	// higher priority is more important
+	return ExecutionIndex < Other.ExecutionIndex;
+}
+
+//----------------------------------------------------------------------//
+// FBehaviorTreeSearchData
+//----------------------------------------------------------------------//
+void FBehaviorTreeSearchData::AddUniqueUpdate(const FBehaviorTreeSearchUpdate& UpdateInfo)
+{
+	UE_VLOG(OwnerComp->GetOwner(), LogBehaviorTree, Verbose, TEXT("Search node update[%s]: %s"),
+		*UBehaviorTreeTypes::DescribeNodeUpdateMode(UpdateInfo.Mode),
+		*UBehaviorTreeTypes::DescribeNodeHelper(UpdateInfo.AuxNode ? (UBTNode*)UpdateInfo.AuxNode : (UBTNode*)UpdateInfo.TaskNode));
+
+	bool bSkipAdding = false;
+	for (int32 i = 0; i < PendingUpdates.Num(); i++)
+	{
+		const FBehaviorTreeSearchUpdate& Info = PendingUpdates[i];
+		if (Info.AuxNode == UpdateInfo.AuxNode && Info.TaskNode == UpdateInfo.TaskNode)
+		{
+			// duplicate, skip
+			// special case: AddForLowerPri replacing Add
+			if ((Info.Mode == UpdateInfo.Mode) ||
+				(Info.Mode == EBTNodeUpdateMode::Add && UpdateInfo.Mode == EBTNodeUpdateMode::AddForLowerPri))
+			{
+				bSkipAdding = true;
+				break;
+			}
+
+			// don't add pairs add-remove
+			bSkipAdding = (Info.Mode == EBTNodeUpdateMode::Remove) || (UpdateInfo.Mode == EBTNodeUpdateMode::Remove);
+
+			PendingUpdates.RemoveAt(i, 1, false);
+		}
+	}
+
+	if (!bSkipAdding)
+	{
+		const int32 Idx = PendingUpdates.Add(UpdateInfo);
+		PendingUpdates[Idx].bPostUpdate = (UpdateInfo.Mode == EBTNodeUpdateMode::Add) && (Cast<UBTService>(UpdateInfo.AuxNode) != NULL);
+	}
+}
+
+//----------------------------------------------------------------------//
+// FBlackboardKeySelector
+//----------------------------------------------------------------------//
+void FBlackboardKeySelector::CacheSelectedKey(class UBlackboardData* BlackboardAsset)
+{
+	if (BlackboardAsset && !(bNoneIsAllowedValue && SelectedKeyID == FBlackboard::InvalidKey))
+	{
+		if (SelectedKeyName == NAME_None)
+		{
+			InitSelectedKey(BlackboardAsset);
+		}
+
+		SelectedKeyID = BlackboardAsset->GetKeyID(SelectedKeyName);
+		SelectedKeyType = BlackboardAsset->GetKeyType(SelectedKeyID);
+	}
+}
+
+void FBlackboardKeySelector::InitSelectedKey(class UBlackboardData* BlackboardAsset)
+{
+	check(BlackboardAsset);
+	for (UBlackboardData* It = BlackboardAsset; It; It = It->Parent)
+	{
+		for (int32 i = 0; i < It->Keys.Num(); i++)
+		{
+			const FBlackboardEntry& EntryInfo = It->Keys[i];
+			if (EntryInfo.KeyType)
+			{
+				bool bFilterPassed = false;
+				if (AllowedTypes.Num())
+				{
+					for (int32 iFilter = 0; iFilter < AllowedTypes.Num(); iFilter++)
+					{
+						if (EntryInfo.KeyType->IsAllowedByFilter(AllowedTypes[iFilter]))
+						{
+							bFilterPassed = true;
+							break;
+						}
+					}
+				}
+				else
+				{
+					bFilterPassed = true;
+				}
+
+				SelectedKeyName = EntryInfo.EntryName;
+				break;
+			}
+		}
+	}
+}
+
+void FBlackboardKeySelector::AddObjectFilter(UObject* Owner, TSubclassOf<UObject> AllowedClass)
+{
+	UBlackboardKeyType_Object* FilterOb = NewObject<UBlackboardKeyType_Object>(Owner);
+	FilterOb->BaseClass = AllowedClass;
+	AllowedTypes.Add(FilterOb);
+}
+
+void FBlackboardKeySelector::AddClassFilter(UObject* Owner, TSubclassOf<UClass> AllowedClass)
+{
+	UBlackboardKeyType_Class* FilterOb = NewObject<UBlackboardKeyType_Class>(Owner);
+	FilterOb->BaseClass = AllowedClass;
+	AllowedTypes.Add(FilterOb);
+}
+
+void FBlackboardKeySelector::AddEnumFilter(UObject* Owner, UEnum* AllowedEnum)
+{
+	UBlackboardKeyType_Enum* FilterOb = NewObject<UBlackboardKeyType_Enum>(Owner);
+	FilterOb->EnumType = AllowedEnum;
+	AllowedTypes.Add(FilterOb);
+}
+
+void FBlackboardKeySelector::AddNativeEnumFilter(UObject* Owner, const FString& AllowedEnumName)
+{
+	UBlackboardKeyType_NativeEnum* FilterOb = NewObject<UBlackboardKeyType_NativeEnum>(Owner);
+	FilterOb->EnumName = AllowedEnumName;
+	AllowedTypes.Add(FilterOb);
+}
+
+void FBlackboardKeySelector::AddIntFilter(UObject* Owner)
+{
+	AllowedTypes.Add(NewObject<UBlackboardKeyType_Int>(Owner));
+}
+
+void FBlackboardKeySelector::AddFloatFilter(UObject* Owner)
+{
+	AllowedTypes.Add(NewObject<UBlackboardKeyType_Float>(Owner));
+}
+
+void FBlackboardKeySelector::AddBoolFilter(UObject* Owner)
+{
+	AllowedTypes.Add(NewObject<UBlackboardKeyType_Bool>(Owner));
+}
+
+void FBlackboardKeySelector::AddVectorFilter(UObject* Owner)
+{
+	AllowedTypes.Add(NewObject<UBlackboardKeyType_Vector>(Owner));
+}
+
+void FBlackboardKeySelector::AddStringFilter(UObject* Owner)
+{
+	AllowedTypes.Add(NewObject<UBlackboardKeyType_String>(Owner));
+}
+
+void FBlackboardKeySelector::AddNameFilter(UObject* Owner)
+{
+	AllowedTypes.Add(NewObject<UBlackboardKeyType_Name>(Owner));
+}
+
+//----------------------------------------------------------------------//
+// UBehaviorTreeTypes
+//----------------------------------------------------------------------//
+UBehaviorTreeTypes::UBehaviorTreeTypes(const class FPostConstructInitializeProperties& PCIP) : Super(PCIP)
+{
+}
+
+FString UBehaviorTreeTypes::DescribeNodeResult(EBTNodeResult::Type NodeResult)
+{
+	static FString ResultDesc[] = { TEXT("Succeeded"), TEXT("Failed"), TEXT("Optional"), TEXT("Aborted"), TEXT("InProgress") };
+	return (NodeResult < ARRAY_COUNT(ResultDesc)) ? ResultDesc[NodeResult] : FString();
+}
+
+FString UBehaviorTreeTypes::DescribeFlowAbortMode(EBTFlowAbortMode::Type AbortMode)
+{
+	static FString AbortModeDesc[] = { TEXT("None"), TEXT("Lower Priority"), TEXT("Self"), TEXT("Both") };
+	return (AbortMode < ARRAY_COUNT(AbortModeDesc)) ? AbortModeDesc[AbortMode] : FString();
+}
+
+FString UBehaviorTreeTypes::DescribeActiveNode(EBTActiveNode::Type ActiveNodeType)
+{
+	static FString ActiveDesc[] = { TEXT("Composite"), TEXT("ActiveTask"), TEXT("AbortingTask"), TEXT("InactiveTask") };
+	return (ActiveNodeType < ARRAY_COUNT(ActiveDesc)) ? ActiveDesc[ActiveNodeType] : FString();
+}
+
+FString UBehaviorTreeTypes::DescribeTaskStatus(EBTTaskStatus::Type TaskStatus)
+{
+	static FString TaskStatusDesc[] = { TEXT("Active"), TEXT("Aborting"), TEXT("Inactive") };
+	return (TaskStatus < ARRAY_COUNT(TaskStatusDesc)) ? TaskStatusDesc[TaskStatus] : FString();
+}
+
+FString UBehaviorTreeTypes::DescribeNodeUpdateMode(EBTNodeUpdateMode::Type UpdateMode)
+{
+	static FString UpdateModeDesc[] = { TEXT("Add"), TEXT("AddForLowerPri"), TEXT("Remove") };
+	return (UpdateMode < ARRAY_COUNT(UpdateModeDesc)) ? UpdateModeDesc[UpdateMode] : FString();
+}
+
+FString UBehaviorTreeTypes::DescribeNodeHelper(const UBTNode* Node)
+{
+	return Node ? FString::Printf(TEXT("%s[%d]"), *Node->GetNodeName(), Node->GetExecutionIndex()) : FString();
+}
+
+FString UBehaviorTreeTypes::GetShortTypeName(const UObject* Ob)
+{
+	if (Ob->GetClass()->HasAnyClassFlags(CLASS_CompiledFromBlueprint))
+	{
+		return Ob->GetClass()->GetName().LeftChop(2);
+	}
+
+	FString TypeDesc = Ob->GetClass()->GetName();
+	const int32 ShortNameIdx = TypeDesc.Find(TEXT("_"));
+	if (ShortNameIdx != INDEX_NONE)
+	{
+		TypeDesc = TypeDesc.Mid(ShortNameIdx + 1);
+	}
+
+	return TypeDesc;
+}
diff --git a/Engine/Source/Runtime/AIModule/Private/BehaviorTree/Blackboard/BlackboardKeyType.cpp b/Engine/Source/Runtime/AIModule/Private/BehaviorTree/Blackboard/BlackboardKeyType.cpp
new file mode 100644
--- /dev/null
+++ b/Engine/Source/Runtime/AIModule/Private/BehaviorTree/Blackboard/BlackboardKeyType.cpp
@@ -0,0 +1,71 @@
+// Copyright 1998-2014 Epic Games, Inc. All Rights Reserved.
+
+#include "AIModulePrivate.h"
+#include "BehaviorTree/Blackboard/BlackboardKeyType.h"
+
+UBlackboardKeyType::UBlackboardKeyType(const class FPostConstructInitializeProperties& PCIP) : Super(PCIP)
+{
+	ValueSize = 0;
+	SupportedOp = EBlackboardKeyOperation::Basic;
+}
+
+FString UBlackboardKeyType::DescribeValue(const uint8* RawData) const
+{
+	FString DescBytes;
+	for (int32 i = 0; i < ValueSize; i++)
+	{
+		DescBytes += FString::Printf(TEXT("%X"), RawData[i]);
+	}
+
+	return DescBytes.Len() ? (FString("0x") + DescBytes) : FString("empty");
+}
+
+FString UBlackboardKeyType::DescribeSelf() const
+{
+	return FString();
+};
+
+void UBlackboardKeyType::Initialize(uint8* MemoryBlock) const
+{
+	// empty in base class
+}
+
+bool UBlackboardKeyType::IsAllowedByFilter(UBlackboardKeyType* FilterOb) const
+{
+	return GetClass() == (FilterOb ? FilterOb->GetClass() : NULL);
+}
+
+bool UBlackboardKeyType::GetLocation(const uint8* RawData, FVector& Location) const
+{
+	return false;
+}
+
+bool UBlackboardKeyType::GetRotation(const uint8* MemoryBlock, FRotator& Rotation) const
+{
+	return false;
+}
+
+EBlackboardCompare::Type UBlackboardKeyType::Compare(const uint8* MemoryBlockA, const uint8* MemoryBlockB) const
+{ 
+	return MemoryBlockA == MemoryBlockB ? EBlackboardCompare::Equal : EBlackboardCompare::NotEqual;
+}
+
+bool UBlackboardKeyType::TestBasicOperation(const uint8* MemoryBlock, EBasicKeyOperation::Type Op) const
+{
+	return false;
+}
+
+bool UBlackboardKeyType::TestArithmeticOperation(const uint8* MemoryBlock, EArithmeticKeyOperation::Type Op, int32 OtherIntValue, float OtherFloatValue) const
+{
+	return false;
+}
+
+bool UBlackboardKeyType::TestTextOperation(const uint8* MemoryBlock, ETextKeyOperation::Type Op, const FString& OtherString) const
+{
+	return false;
+}
+
+FString UBlackboardKeyType::DescribeArithmeticParam(int32 IntValue, float FloatValue) const
+{
+	return FString();
+}
diff --git a/Engine/Source/Runtime/AIModule/Private/BehaviorTree/Blackboard/BlackboardKeyType_Bool.cpp b/Engine/Source/Runtime/AIModule/Private/BehaviorTree/Blackboard/BlackboardKeyType_Bool.cpp
new file mode 100644
--- /dev/null
+++ b/Engine/Source/Runtime/AIModule/Private/BehaviorTree/Blackboard/BlackboardKeyType_Bool.cpp
@@ -0,0 +1,36 @@
+// Copyright 1998-2014 Epic Games, Inc. All Rights Reserved.
+
+#include "AIModulePrivate.h"
+#include "BehaviorTree/Blackboard/BlackboardKeyType_Bool.h"
+
+UBlackboardKeyType_Bool::UBlackboardKeyType_Bool(const class FPostConstructInitializeProperties& PCIP) : Super(PCIP)
+{
+	ValueSize = sizeof(uint8);
+	SupportedOp = EBlackboardKeyOperation::Basic;
+}
+
+bool UBlackboardKeyType_Bool::GetValue(const uint8* RawData)
+{
+	return GetValueFromMemory<uint8>(RawData) != 0;
+}
+
+bool UBlackboardKeyType_Bool::SetValue(uint8* RawData, bool bValue)
+{
+	return SetValueInMemory<uint8>(RawData, bValue);
+}
+
+FString UBlackboardKeyType_Bool::DescribeValue(const uint8* RawData) const
+{
+	return GetValue(RawData) ? TEXT("true") : TEXT("false");
+}
+
+EBlackboardCompare::Type UBlackboardKeyType_Bool::Compare(const uint8* MemoryBlockA, const uint8* MemoryBlockB) const
+{
+	return GetValueFromMemory<uint8>(MemoryBlockA) == GetValueFromMemory<uint8>(MemoryBlockB) ? EBlackboardCompare::Equal : EBlackboardCompare::NotEqual;
+}
+
+bool UBlackboardKeyType_Bool::TestBasicOperation(const uint8* MemoryBlock, EBasicKeyOperation::Type Op) const
+{
+	const bool Value = GetValue(MemoryBlock);
+	return (Op == EBasicKeyOperation::Set) ? Value : !Value;
+}
diff --git a/Engine/Source/Runtime/AIModule/Private/BehaviorTree/Blackboard/BlackboardKeyType_Class.cpp b/Engine/Source/Runtime/AIModule/Private/BehaviorTree/Blackboard/BlackboardKeyType_Class.cpp
new file mode 100644
--- /dev/null
+++ b/Engine/Source/Runtime/AIModule/Private/BehaviorTree/Blackboard/BlackboardKeyType_Class.cpp
@@ -0,0 +1,50 @@
+// Copyright 1998-2014 Epic Games, Inc. All Rights Reserved.
+
+#include "AIModulePrivate.h"
+#include "BehaviorTree/Blackboard/BlackboardKeyType_Class.h"
+
+UBlackboardKeyType_Class::UBlackboardKeyType_Class(const class FPostConstructInitializeProperties& PCIP) : Super(PCIP)
+{
+	ValueSize = sizeof(TWeakObjectPtr<UClass>);
+	BaseClass = UObject::StaticClass();
+	SupportedOp = EBlackboardKeyOperation::Basic;
+}
+
+UClass* UBlackboardKeyType_Class::GetValue(const uint8* RawData)
+{
+	TWeakObjectPtr<UClass> WeakObjPtr = GetValueFromMemory< TWeakObjectPtr<UClass> >(RawData);
+	return WeakObjPtr.Get();
+}
+
+bool UBlackboardKeyType_Class::SetValue(uint8* RawData, UClass* Value)
+{
+	TWeakObjectPtr<UClass> WeakObjPtr(Value);
+	return SetWeakObjectInMemory<UClass>(RawData, WeakObjPtr);
+}
+
+FString UBlackboardKeyType_Class::DescribeValue(const uint8* RawData) const
+{
+	return *GetNameSafe(GetValue(RawData));
+}
+
+FString UBlackboardKeyType_Class::DescribeSelf() const
+{
+	return *GetNameSafe(BaseClass);
+}
+
+bool UBlackboardKeyType_Class::IsAllowedByFilter(UBlackboardKeyType* FilterOb) const
+{
+	UBlackboardKeyType_Class* FilterClass = Cast<UBlackboardKeyType_Class>(FilterOb);
+	return (FilterClass && (FilterClass->BaseClass == BaseClass || BaseClass->IsChildOf(FilterClass->BaseClass)));
+}
+
+EBlackboardCompare::Type UBlackboardKeyType_Class::Compare(const uint8* MemoryBlockA, const uint8* MemoryBlockB) const
+{
+	return (FMemory::Memcmp(MemoryBlockA, MemoryBlockB, ValueSize) == 0) ? EBlackboardCompare::Equal : EBlackboardCompare::NotEqual;
+}
+
+bool UBlackboardKeyType_Class::TestBasicOperation(const uint8* MemoryBlock, EBasicKeyOperation::Type Op) const
+{
+	TWeakObjectPtr<UClass> WeakObjPtr = GetValueFromMemory< TWeakObjectPtr<UClass> >(MemoryBlock);
+	return (Op == EBasicKeyOperation::Set) ? WeakObjPtr.IsValid() : !WeakObjPtr.IsValid();
+}
