Extracting AI code from the engine into a separate module, pass 1 #UE4

diff --git a/Engine/Source/Runtime/AIModule/Private/BehaviorTree/Blackboard/BlackboardKeyType_Enum.cpp b/Engine/Source/Runtime/AIModule/Private/BehaviorTree/Blackboard/BlackboardKeyType_Enum.cpp
new file mode 100644
--- /dev/null
+++ b/Engine/Source/Runtime/AIModule/Private/BehaviorTree/Blackboard/BlackboardKeyType_Enum.cpp
@@ -0,0 +1,64 @@
+// Copyright 1998-2014 Epic Games, Inc. All Rights Reserved.
+
+#include "AIModulePrivate.h"
+#include "BehaviorTree/Blackboard/BlackboardKeyType_Enum.h"
+
+UBlackboardKeyType_Enum::UBlackboardKeyType_Enum(const class FPostConstructInitializeProperties& PCIP) : Super(PCIP)
+{
+	ValueSize = sizeof(uint8);
+	SupportedOp = EBlackboardKeyOperation::Arithmetic;
+}
+
+uint8 UBlackboardKeyType_Enum::GetValue(const uint8* RawData)
+{
+	return GetValueFromMemory<uint8>(RawData);
+}
+
+bool UBlackboardKeyType_Enum::SetValue(uint8* RawData, uint8 Value)
+{
+	return SetValueInMemory<uint8>(RawData, Value);
+}
+
+FString UBlackboardKeyType_Enum::DescribeValue(const uint8* RawData) const
+{
+	return EnumType ? EnumType->GetEnumName(GetValue(RawData)) : FString("UNKNOWN!");
+}
+
+FString UBlackboardKeyType_Enum::DescribeSelf() const
+{
+	return *GetNameSafe(EnumType);
+}
+
+bool UBlackboardKeyType_Enum::IsAllowedByFilter(UBlackboardKeyType* FilterOb) const
+{
+	UBlackboardKeyType_Enum* FilterEnum = Cast<UBlackboardKeyType_Enum>(FilterOb);
+	return (FilterEnum && FilterEnum->EnumType == EnumType);
+}
+
+EBlackboardCompare::Type UBlackboardKeyType_Enum::Compare(const uint8* MemoryBlockA, const uint8* MemoryBlockB) const
+{
+	const int8 Diff = GetValueFromMemory<uint8>(MemoryBlockA) - GetValueFromMemory<uint8>(MemoryBlockB);
+	return Diff > 0 ? EBlackboardCompare::Greater : (Diff < 0 ? EBlackboardCompare::Less : EBlackboardCompare::Equal);
+}
+
+bool UBlackboardKeyType_Enum::TestArithmeticOperation(const uint8* MemoryBlock, EArithmeticKeyOperation::Type Op, int32 OtherIntValue, float OtherFloatValue) const
+{
+	const uint8 Value = GetValue(MemoryBlock);
+	switch (Op)
+	{
+	case EArithmeticKeyOperation::Equal:			return (Value == OtherIntValue);
+	case EArithmeticKeyOperation::NotEqual:			return (Value != OtherIntValue);
+	case EArithmeticKeyOperation::Less:				return (Value < OtherIntValue);
+	case EArithmeticKeyOperation::LessOrEqual:		return (Value <= OtherIntValue);
+	case EArithmeticKeyOperation::Greater:			return (Value > OtherIntValue);
+	case EArithmeticKeyOperation::GreaterOrEqual:	return (Value >= OtherIntValue);
+	default: break;
+	}
+
+	return false;
+}
+
+FString UBlackboardKeyType_Enum::DescribeArithmeticParam(int32 IntValue, float FloatValue) const
+{
+	return EnumType ? EnumType->GetEnumName(IntValue) : FString("UNKNOWN!");
+}
diff --git a/Engine/Source/Runtime/AIModule/Private/BehaviorTree/Blackboard/BlackboardKeyType_Float.cpp b/Engine/Source/Runtime/AIModule/Private/BehaviorTree/Blackboard/BlackboardKeyType_Float.cpp
new file mode 100644
--- /dev/null
+++ b/Engine/Source/Runtime/AIModule/Private/BehaviorTree/Blackboard/BlackboardKeyType_Float.cpp
@@ -0,0 +1,54 @@
+// Copyright 1998-2014 Epic Games, Inc. All Rights Reserved.
+
+#include "AIModulePrivate.h"
+#include "BehaviorTree/Blackboard/BlackboardKeyType_Float.h"
+
+UBlackboardKeyType_Float::UBlackboardKeyType_Float(const class FPostConstructInitializeProperties& PCIP) : Super(PCIP)
+{
+	ValueSize = sizeof(float);
+	SupportedOp = EBlackboardKeyOperation::Arithmetic;
+}
+
+float UBlackboardKeyType_Float::GetValue(const uint8* RawData)
+{
+	return GetValueFromMemory<float>(RawData);
+}
+
+bool UBlackboardKeyType_Float::SetValue(uint8* RawData, float Value)
+{
+	return SetValueInMemory<float>(RawData, Value);
+}
+
+FString UBlackboardKeyType_Float::DescribeValue(const uint8* RawData) const
+{
+	return FString::Printf(TEXT("%f"), GetValue(RawData));
+}
+
+EBlackboardCompare::Type UBlackboardKeyType_Float::Compare(const uint8* MemoryBlockA, const uint8* MemoryBlockB) const
+{
+	const float Diff = GetValueFromMemory<float>(MemoryBlockA) - GetValueFromMemory<float>(MemoryBlockB);
+	return FMath::Abs(Diff) < KINDA_SMALL_NUMBER ? EBlackboardCompare::Equal 
+		: (Diff > 0 ? EBlackboardCompare::Greater : EBlackboardCompare::Less);
+}
+
+bool UBlackboardKeyType_Float::TestArithmeticOperation(const uint8* MemoryBlock, EArithmeticKeyOperation::Type Op, int32 OtherIntValue, float OtherFloatValue) const
+{
+	const float Value = GetValue(MemoryBlock);
+	switch (Op)
+	{
+	case EArithmeticKeyOperation::Equal:			return (Value == OtherFloatValue);
+	case EArithmeticKeyOperation::NotEqual:			return (Value != OtherFloatValue);
+	case EArithmeticKeyOperation::Less:				return (Value < OtherFloatValue);
+	case EArithmeticKeyOperation::LessOrEqual:		return (Value <= OtherFloatValue);
+	case EArithmeticKeyOperation::Greater:			return (Value > OtherFloatValue);
+	case EArithmeticKeyOperation::GreaterOrEqual:	return (Value >= OtherFloatValue);
+	default: break;
+	}
+
+	return false;
+}
+
+FString UBlackboardKeyType_Float::DescribeArithmeticParam(int32 IntValue, float FloatValue) const
+{
+	return FString::Printf(TEXT("%f"), FloatValue);
+}
diff --git a/Engine/Source/Runtime/AIModule/Private/BehaviorTree/Blackboard/BlackboardKeyType_Int.cpp b/Engine/Source/Runtime/AIModule/Private/BehaviorTree/Blackboard/BlackboardKeyType_Int.cpp
new file mode 100644
--- /dev/null
+++ b/Engine/Source/Runtime/AIModule/Private/BehaviorTree/Blackboard/BlackboardKeyType_Int.cpp
@@ -0,0 +1,53 @@
+// Copyright 1998-2014 Epic Games, Inc. All Rights Reserved.
+
+#include "AIModulePrivate.h"
+#include "BehaviorTree/Blackboard/BlackboardKeyType_Int.h"
+
+UBlackboardKeyType_Int::UBlackboardKeyType_Int(const class FPostConstructInitializeProperties& PCIP) : Super(PCIP)
+{
+	ValueSize = sizeof(int32);
+	SupportedOp = EBlackboardKeyOperation::Arithmetic;
+}
+
+int32 UBlackboardKeyType_Int::GetValue(const uint8* RawData)
+{
+	return GetValueFromMemory<int32>(RawData);
+}
+
+bool UBlackboardKeyType_Int::SetValue(uint8* RawData, int32 Value)
+{
+	return SetValueInMemory<int32>(RawData, Value);
+}
+
+FString UBlackboardKeyType_Int::DescribeValue(const uint8* RawData) const
+{
+	return FString::Printf(TEXT("%d"), GetValue(RawData));
+}
+
+EBlackboardCompare::Type UBlackboardKeyType_Int::Compare(const uint8* MemoryBlockA, const uint8* MemoryBlockB) const
+{
+	const int32 Diff = GetValueFromMemory<int32>(MemoryBlockA) - GetValueFromMemory<int32>(MemoryBlockB);
+	return Diff > 0 ? EBlackboardCompare::Greater : (Diff < 0 ? EBlackboardCompare::Less : EBlackboardCompare::Equal);
+}
+
+bool UBlackboardKeyType_Int::TestArithmeticOperation(const uint8* MemoryBlock, EArithmeticKeyOperation::Type Op, int32 OtherIntValue, float OtherFloatValue) const
+{
+	const int32 Value = GetValue(MemoryBlock);
+	switch (Op)
+	{
+	case EArithmeticKeyOperation::Equal:			return (Value == OtherIntValue);
+	case EArithmeticKeyOperation::NotEqual:			return (Value != OtherIntValue);
+	case EArithmeticKeyOperation::Less:				return (Value < OtherIntValue);
+	case EArithmeticKeyOperation::LessOrEqual:		return (Value <= OtherIntValue);
+	case EArithmeticKeyOperation::Greater:			return (Value > OtherIntValue);
+	case EArithmeticKeyOperation::GreaterOrEqual:	return (Value >= OtherIntValue);
+	default: break;
+	}
+
+	return false;
+}
+
+FString UBlackboardKeyType_Int::DescribeArithmeticParam(int32 IntValue, float FloatValue) const
+{
+	return FString::Printf(TEXT("%d"), IntValue);
+}
diff --git a/Engine/Source/Runtime/AIModule/Private/BehaviorTree/Blackboard/BlackboardKeyType_Name.cpp b/Engine/Source/Runtime/AIModule/Private/BehaviorTree/Blackboard/BlackboardKeyType_Name.cpp
new file mode 100644
--- /dev/null
+++ b/Engine/Source/Runtime/AIModule/Private/BehaviorTree/Blackboard/BlackboardKeyType_Name.cpp
@@ -0,0 +1,46 @@
+// Copyright 1998-2014 Epic Games, Inc. All Rights Reserved.
+
+#include "AIModulePrivate.h"
+#include "BehaviorTree/Blackboard/BlackboardKeyType_Name.h"
+
+UBlackboardKeyType_Name::UBlackboardKeyType_Name(const class FPostConstructInitializeProperties& PCIP) : Super(PCIP)
+{
+	ValueSize = sizeof(FName);
+	SupportedOp = EBlackboardKeyOperation::Text;
+}
+
+FName UBlackboardKeyType_Name::GetValue(const uint8* RawData)
+{
+	return GetValueFromMemory<FName>(RawData);
+}
+
+bool UBlackboardKeyType_Name::SetValue(uint8* RawData, const FName& Value)
+{
+	return SetValueInMemory<FName>(RawData, Value);
+}
+
+FString UBlackboardKeyType_Name::DescribeValue(const uint8* RawData) const
+{
+	return GetValue(RawData).ToString();
+}
+
+EBlackboardCompare::Type UBlackboardKeyType_Name::Compare(const uint8* MemoryBlockA, const uint8* MemoryBlockB) const
+{
+	return GetValueFromMemory<FName>(MemoryBlockA) == GetValueFromMemory<FName>(MemoryBlockB)
+		? EBlackboardCompare::Equal : EBlackboardCompare::NotEqual;
+}
+
+bool UBlackboardKeyType_Name::TestTextOperation(const uint8* MemoryBlock, ETextKeyOperation::Type Op, const FString& OtherString) const
+{
+	const FString Value = GetValue(MemoryBlock).ToString();
+	switch (Op)
+	{
+	case ETextKeyOperation::Equal:			return (Value == OtherString);
+	case ETextKeyOperation::NotEqual:		return (Value != OtherString);
+	case ETextKeyOperation::Contain:		return (Value.Contains(OtherString) == true);
+	case ETextKeyOperation::NotContain:		return (Value.Contains(OtherString) == false);
+	default: break;
+	}
+
+	return false;
+}
diff --git a/Engine/Source/Runtime/AIModule/Private/BehaviorTree/Blackboard/BlackboardKeyType_NativeEnum.cpp b/Engine/Source/Runtime/AIModule/Private/BehaviorTree/Blackboard/BlackboardKeyType_NativeEnum.cpp
new file mode 100644
--- /dev/null
+++ b/Engine/Source/Runtime/AIModule/Private/BehaviorTree/Blackboard/BlackboardKeyType_NativeEnum.cpp
@@ -0,0 +1,78 @@
+// Copyright 1998-2014 Epic Games, Inc. All Rights Reserved.
+
+#include "AIModulePrivate.h"
+#include "BehaviorTree/Blackboard/BlackboardKeyType_NativeEnum.h"
+
+UBlackboardKeyType_NativeEnum::UBlackboardKeyType_NativeEnum(const class FPostConstructInitializeProperties& PCIP) : Super(PCIP)
+{
+	ValueSize = sizeof(uint8);
+	SupportedOp = EBlackboardKeyOperation::Arithmetic;
+}
+
+uint8 UBlackboardKeyType_NativeEnum::GetValue(const uint8* RawData)
+{
+	return GetValueFromMemory<uint8>(RawData);
+}
+
+bool UBlackboardKeyType_NativeEnum::SetValue(uint8* RawData, uint8 Value)
+{
+	return SetValueInMemory<uint8>(RawData, Value);
+}
+
+FString UBlackboardKeyType_NativeEnum::DescribeValue(const uint8* RawData) const
+{
+	return EnumType ? EnumType->GetEnumName(GetValue(RawData)) : FString("UNKNOWN!");
+}
+
+FString UBlackboardKeyType_NativeEnum::DescribeSelf() const
+{
+	return *GetNameSafe(EnumType);
+}
+
+bool UBlackboardKeyType_NativeEnum::IsAllowedByFilter(UBlackboardKeyType* FilterOb) const
+{
+	UBlackboardKeyType_NativeEnum* FilterEnum = Cast<UBlackboardKeyType_NativeEnum>(FilterOb);
+	return (FilterEnum && FilterEnum->EnumName == EnumName);
+}
+
+#if WITH_EDITOR
+void UBlackboardKeyType_NativeEnum::PostEditChangeProperty(struct FPropertyChangedEvent& PropertyChangedEvent)
+{
+	Super::PostEditChangeProperty(PropertyChangedEvent);
+
+	if (PropertyChangedEvent.Property &&
+		PropertyChangedEvent.Property->GetFName() == GET_MEMBER_NAME_CHECKED(UBlackboardKeyType_NativeEnum, EnumName))
+	{
+		EnumType = FindObject<UEnum>(ANY_PACKAGE, *EnumName, true);
+		bIsEnumNameValid =  (EnumType != NULL);
+	}
+}
+#endif
+
+EBlackboardCompare::Type UBlackboardKeyType_NativeEnum::Compare(const uint8* MemoryBlockA, const uint8* MemoryBlockB) const
+{
+	const int8 Diff = GetValueFromMemory<uint8>(MemoryBlockA) - GetValueFromMemory<uint8>(MemoryBlockB);
+	return Diff > 0 ? EBlackboardCompare::Greater : (Diff < 0 ? EBlackboardCompare::Less : EBlackboardCompare::Equal);
+}
+
+bool UBlackboardKeyType_NativeEnum::TestArithmeticOperation(const uint8* MemoryBlock, EArithmeticKeyOperation::Type Op, int32 OtherIntValue, float OtherFloatValue) const
+{
+	const uint8 Value = GetValue(MemoryBlock);
+	switch (Op)
+	{
+	case EArithmeticKeyOperation::Equal:			return (Value == OtherIntValue);
+	case EArithmeticKeyOperation::NotEqual:			return (Value != OtherIntValue);
+	case EArithmeticKeyOperation::Less:				return (Value < OtherIntValue);
+	case EArithmeticKeyOperation::LessOrEqual:		return (Value <= OtherIntValue);
+	case EArithmeticKeyOperation::Greater:			return (Value > OtherIntValue);
+	case EArithmeticKeyOperation::GreaterOrEqual:	return (Value >= OtherIntValue);
+	default: break;
+	}
+
+	return false;
+}
+
+FString UBlackboardKeyType_NativeEnum::DescribeArithmeticParam(int32 IntValue, float FloatValue) const
+{
+	return EnumType ? EnumType->GetEnumName(IntValue) : FString("UNKNOWN!");
+}
diff --git a/Engine/Source/Runtime/AIModule/Private/BehaviorTree/Blackboard/BlackboardKeyType_Object.cpp b/Engine/Source/Runtime/AIModule/Private/BehaviorTree/Blackboard/BlackboardKeyType_Object.cpp
new file mode 100644
--- /dev/null
+++ b/Engine/Source/Runtime/AIModule/Private/BehaviorTree/Blackboard/BlackboardKeyType_Object.cpp
@@ -0,0 +1,74 @@
+// Copyright 1998-2014 Epic Games, Inc. All Rights Reserved.
+
+#include "AIModulePrivate.h"
+#include "BehaviorTree/Blackboard/BlackboardKeyType_Object.h"
+
+UBlackboardKeyType_Object::UBlackboardKeyType_Object(const class FPostConstructInitializeProperties& PCIP) : Super(PCIP)
+{
+	ValueSize = sizeof(FWeakObjectPtr);
+	BaseClass = UObject::StaticClass();
+	SupportedOp = EBlackboardKeyOperation::Basic;
+}
+
+UObject* UBlackboardKeyType_Object::GetValue(const uint8* RawData)
+{
+	FWeakObjectPtr WeakObjPtr = GetValueFromMemory<FWeakObjectPtr>(RawData);
+	return WeakObjPtr.Get();
+}
+
+bool UBlackboardKeyType_Object::SetValue(uint8* RawData, UObject* Value)
+{
+	TWeakObjectPtr<UObject> WeakObjPtr(Value);
+	return SetWeakObjectInMemory<UObject>(RawData, WeakObjPtr);
+}
+
+FString UBlackboardKeyType_Object::DescribeValue(const uint8* RawData) const
+{
+	return *GetNameSafe(GetValue(RawData));
+}
+
+FString UBlackboardKeyType_Object::DescribeSelf() const
+{
+	return *GetNameSafe(BaseClass);
+}
+
+bool UBlackboardKeyType_Object::IsAllowedByFilter(UBlackboardKeyType* FilterOb) const
+{
+	UBlackboardKeyType_Object* FilterObject = Cast<UBlackboardKeyType_Object>(FilterOb);
+	return (FilterObject && (FilterObject->BaseClass == BaseClass || BaseClass->IsChildOf(FilterObject->BaseClass)));
+}
+
+bool UBlackboardKeyType_Object::GetLocation(const uint8* RawData, FVector& Location) const
+{
+	AActor* MyActor = Cast<AActor>(GetValue(RawData));
+	if (MyActor)
+	{
+		Location = MyActor->GetActorLocation();
+		return true;
+	}
+
+	return false;
+}
+
+bool UBlackboardKeyType_Object::GetRotation(const uint8* RawData, FRotator& Rotation) const
+{
+	AActor* MyActor = Cast<AActor>(GetValue(RawData));
+	if (MyActor)
+	{
+		Rotation = MyActor->GetActorRotation();
+		return true;
+	}
+
+	return false;
+}
+
+EBlackboardCompare::Type UBlackboardKeyType_Object::Compare(const uint8* MemoryBlockA, const uint8* MemoryBlockB) const
+{
+	return (FMemory::Memcmp(MemoryBlockA, MemoryBlockB, ValueSize) == 0) ? EBlackboardCompare::Equal : EBlackboardCompare::NotEqual;
+}
+
+bool UBlackboardKeyType_Object::TestBasicOperation(const uint8* MemoryBlock, EBasicKeyOperation::Type Op) const
+{
+	FWeakObjectPtr WeakObjPtr = GetValueFromMemory<FWeakObjectPtr>(MemoryBlock);
+	return (Op == EBasicKeyOperation::Set) ? WeakObjPtr.IsValid() : !WeakObjPtr.IsValid();
+}
diff --git a/Engine/Source/Runtime/AIModule/Private/BehaviorTree/Blackboard/BlackboardKeyType_String.cpp b/Engine/Source/Runtime/AIModule/Private/BehaviorTree/Blackboard/BlackboardKeyType_String.cpp
new file mode 100644
--- /dev/null
+++ b/Engine/Source/Runtime/AIModule/Private/BehaviorTree/Blackboard/BlackboardKeyType_String.cpp
@@ -0,0 +1,46 @@
+// Copyright 1998-2014 Epic Games, Inc. All Rights Reserved.
+
+#include "AIModulePrivate.h"
+#include "BehaviorTree/Blackboard/BlackboardKeyType_String.h"
+
+UBlackboardKeyType_String::UBlackboardKeyType_String(const class FPostConstructInitializeProperties& PCIP) : Super(PCIP)
+{
+	ValueSize = sizeof(FString);
+	SupportedOp = EBlackboardKeyOperation::Text;
+}
+
+FString UBlackboardKeyType_String::GetValue(const uint8* RawData)
+{
+	return GetValueFromMemory<FString>(RawData);
+}
+
+bool UBlackboardKeyType_String::SetValue(uint8* RawData, const FString& Value)
+{
+	return SetValueInMemory<FString>(RawData, Value);
+}
+
+FString UBlackboardKeyType_String::DescribeValue(const uint8* RawData) const
+{
+	return GetValue(RawData);
+}
+
+EBlackboardCompare::Type UBlackboardKeyType_String::Compare(const uint8* MemoryBlockA, const uint8* MemoryBlockB) const
+{
+	return GetValueFromMemory<FString>(MemoryBlockA) == GetValueFromMemory<FString>(MemoryBlockB)
+		? EBlackboardCompare::Equal : EBlackboardCompare::NotEqual;
+}
+
+bool UBlackboardKeyType_String::TestTextOperation(const uint8* MemoryBlock, ETextKeyOperation::Type Op, const FString& OtherString) const
+{
+	const FString Value = GetValue(MemoryBlock);
+	switch (Op)
+	{
+	case ETextKeyOperation::Equal:			return (Value == OtherString);
+	case ETextKeyOperation::NotEqual:		return (Value != OtherString);
+	case ETextKeyOperation::Contain:		return (Value.Contains(OtherString) == true);
+	case ETextKeyOperation::NotContain:		return (Value.Contains(OtherString) == false);
+	default: break;
+	}
+
+	return false;
+}
diff --git a/Engine/Source/Runtime/AIModule/Private/BehaviorTree/Blackboard/BlackboardKeyType_Vector.cpp b/Engine/Source/Runtime/AIModule/Private/BehaviorTree/Blackboard/BlackboardKeyType_Vector.cpp
new file mode 100644
--- /dev/null
+++ b/Engine/Source/Runtime/AIModule/Private/BehaviorTree/Blackboard/BlackboardKeyType_Vector.cpp
@@ -0,0 +1,49 @@
+// Copyright 1998-2014 Epic Games, Inc. All Rights Reserved.
+
+#include "AIModulePrivate.h"
+#include "BehaviorTree/Blackboard/BlackboardKeyType_Vector.h"
+
+UBlackboardKeyType_Vector::UBlackboardKeyType_Vector(const class FPostConstructInitializeProperties& PCIP) : Super(PCIP)
+{
+	ValueSize = sizeof(FVector);
+	SupportedOp = EBlackboardKeyOperation::Basic;
+}
+
+FVector UBlackboardKeyType_Vector::GetValue(const uint8* RawData)
+{
+	return GetValueFromMemory<FVector>(RawData);
+}
+
+bool UBlackboardKeyType_Vector::SetValue(uint8* RawData, const FVector& Value)
+{
+	return SetValueInMemory<FVector>(RawData, Value);
+}
+
+FString UBlackboardKeyType_Vector::DescribeValue(const uint8* RawData) const
+{
+	const FVector Location = GetValue(RawData);
+	return FAISystem::IsValidLocation(Location) ? Location.ToString() : TEXT("(invalid)");
+}
+
+bool UBlackboardKeyType_Vector::GetLocation(const uint8* RawData, FVector& Location) const
+{
+	Location = GetValue(RawData);
+	return FAISystem::IsValidLocation(Location);
+}
+
+void UBlackboardKeyType_Vector::Initialize(uint8* RawData) const
+{
+	SetValue(RawData, FAISystem::InvalidLocation);
+}
+
+EBlackboardCompare::Type UBlackboardKeyType_Vector::Compare(const uint8* MemoryBlockA, const uint8* MemoryBlockB) const
+{
+	return GetValueFromMemory<FVector>(MemoryBlockA).Equals(GetValueFromMemory<FVector>(MemoryBlockB))
+		? EBlackboardCompare::Equal : EBlackboardCompare::NotEqual;
+}
+
+bool UBlackboardKeyType_Vector::TestBasicOperation(const uint8* MemoryBlock, EBasicKeyOperation::Type Op) const
+{
+	const FVector Location = GetValue(MemoryBlock);
+	return (Op == EBasicKeyOperation::Set) ? FAISystem::IsValidLocation(Location) : !FAISystem::IsValidLocation(Location);
+}
diff --git a/Engine/Source/Runtime/AIModule/Private/BehaviorTree/BlackboardComponent.cpp b/Engine/Source/Runtime/AIModule/Private/BehaviorTree/BlackboardComponent.cpp
new file mode 100644
--- /dev/null
+++ b/Engine/Source/Runtime/AIModule/Private/BehaviorTree/BlackboardComponent.cpp
@@ -0,0 +1,730 @@
+// Copyright 1998-2014 Epic Games, Inc. All Rights Reserved.
+
+#include "AIModulePrivate.h"
+#include "EngineUtils.h"
+#include "BehaviorTree/Blackboard/BlackboardKeyAllTypes.h"
+#include "BehaviorTree/BlackboardComponent.h"
+
+UBlackboardComponent::UBlackboardComponent(const class FPostConstructInitializeProperties& PCIP) : Super(PCIP)
+{
+	PrimaryComponentTick.bCanEverTick = false;
+	bWantsInitializeComponent = true;
+	bPausedNotifies = false;
+}
+
+void UBlackboardComponent::InitializeComponent()
+{
+	Super::InitializeComponent();
+
+	// cache blackboard component if owner has one
+	BrainComp = GetOwner()->FindComponentByClass<UBrainComponent>();
+	if (BrainComp)
+	{
+		BrainComp->CacheBlackboardComponent(this);
+	}
+	InitializeBlackboard(BlackboardAsset);
+}
+
+void UBlackboardComponent::CacheBrainComponent(UBrainComponent* BrainComponent)
+{
+	if (BrainComponent != BrainComp)
+	{
+		BrainComp = BrainComponent;
+	}
+}
+
+struct FBlackboardInitializationData
+{
+	uint8 KeyID;
+	uint16 DataSize;
+
+	FBlackboardInitializationData() {}
+	FBlackboardInitializationData(uint8 InKeyID, uint16 InDataSize) : KeyID(InKeyID)
+	{
+		DataSize = (InDataSize <= 2) ? InDataSize : ((InDataSize + 3) & ~3);
+	}
+
+	struct FMemorySort
+	{
+		FORCEINLINE bool operator()(const FBlackboardInitializationData& A, const FBlackboardInitializationData& B) const
+		{
+			return A.DataSize > B.DataSize;
+		}
+	};
+};
+
+void UBlackboardComponent::InitializeBlackboard(class UBlackboardData* NewAsset)
+{
+	// if we re-initialize with the same asset then there's no point
+	// in reseting, since we'd loose all the accumulated knowledge
+	if (NewAsset == BlackboardAsset)
+	{
+		return;
+	}
+
+	BlackboardAsset = NewAsset;
+	ValueMemory.Reset();
+	ValueOffsets.Reset();
+
+	if (BlackboardAsset)
+	{
+		if (BlackboardAsset->IsValid())
+		{
+			TArray<FBlackboardInitializationData> InitList;
+			const int32 NumKeys = BlackboardAsset->GetNumKeys();
+			InitList.Reserve(NumKeys);
+			ValueOffsets.AddZeroed(NumKeys);
+
+			for (UBlackboardData* It = BlackboardAsset; It; It = It->Parent)
+			{
+				for (int32 i = 0; i < It->Keys.Num(); i++)
+				{
+					if (It->Keys[i].KeyType)
+					{
+						InitList.Add(FBlackboardInitializationData(i + It->GetFirstKeyID(), It->Keys[i].KeyType->GetValueSize()));
+					}
+				}
+			}
+
+			// sort key values by memory size, so they can be packed better
+			// it still won't protect against structures, that are internally misaligned (-> uint8, uint32)
+			// but since all Engine level keys are good... 
+			InitList.Sort(FBlackboardInitializationData::FMemorySort());
+			uint16 MemoryOffset = 0;
+			for (int32 i = 0; i < InitList.Num(); i++)
+			{
+				ValueOffsets[InitList[i].KeyID] = MemoryOffset;
+				MemoryOffset += InitList[i].DataSize;
+			}
+
+			ValueMemory.AddZeroed(MemoryOffset);
+
+			// initialize memory
+			for (int32 i = 0; i < InitList.Num(); i++)
+			{
+				const FBlackboardEntry* KeyData = BlackboardAsset->GetKey(InitList[i].KeyID);
+				uint8* RawData = GetKeyRawData(InitList[i].KeyID);
+
+				KeyData->KeyType->Initialize(RawData);
+			}
+		}
+		else
+		{
+			UE_LOG(LogBehaviorTree, Error, TEXT("Blackboard asset (%s) has errors and can't be used!"), *GetNameSafe(BlackboardAsset));
+		}
+	}
+}
+
+class UBrainComponent* UBlackboardComponent::GetBrainComponent() const
+{
+	return BrainComp;
+}
+
+class UBlackboardData* UBlackboardComponent::GetBlackboardAsset() const
+{
+	return BlackboardAsset;
+}
+
+FName UBlackboardComponent::GetKeyName(uint8 KeyID) const
+{
+	return BlackboardAsset ? BlackboardAsset->GetKeyName(KeyID) : NAME_None;
+}
+
+uint8 UBlackboardComponent::GetKeyID(const FName& KeyName) const
+{
+	return BlackboardAsset ? BlackboardAsset->GetKeyID(KeyName) : 0xFF;
+}
+
+TSubclassOf<UBlackboardKeyType> UBlackboardComponent::GetKeyType(uint8 KeyID) const
+{
+	return BlackboardAsset ? BlackboardAsset->GetKeyType(KeyID) : NULL;
+}
+
+int32 UBlackboardComponent::GetNumKeys() const
+{
+	return BlackboardAsset ? BlackboardAsset->GetNumKeys() : 0;
+}
+
+void UBlackboardComponent::RegisterObserver(uint8 KeyID, FOnBlackboardChange ObserverDelegate)
+{
+	Observers.AddUnique(KeyID, ObserverDelegate);
+}
+
+void UBlackboardComponent::UnregisterObserver(uint8 KeyID, FOnBlackboardChange ObserverDelegate)
+{
+	Observers.RemoveSingle(KeyID, ObserverDelegate);
+}
+
+void UBlackboardComponent::PauseUpdates()
+{
+	bPausedNotifies = true;
+}
+
+void UBlackboardComponent::ResumeUpdates()
+{
+	bPausedNotifies = false;
+
+	for (int32 i = 0; i < QueuedUpdates.Num(); i++)
+	{
+		NotifyObservers(QueuedUpdates[i]);
+	}
+
+	QueuedUpdates.Empty();
+}
+
+void UBlackboardComponent::NotifyObservers(uint8 KeyID) const
+{
+	if (bPausedNotifies)
+	{
+		QueuedUpdates.AddUnique(KeyID);
+	}
+	else
+	{
+		for (TMultiMap<uint8, FOnBlackboardChange>::TConstKeyIterator KeyIt(Observers, KeyID); KeyIt; ++KeyIt)
+		{
+			const FOnBlackboardChange& ObserverDelegate = KeyIt.Value();
+			ObserverDelegate.ExecuteIfBound(this, KeyID);
+		}
+	}
+}
+
+bool UBlackboardComponent::IsCompatibileWith(class UBlackboardData* TestAsset) const
+{
+	for (UBlackboardData* It = BlackboardAsset; It; It = It->Parent)
+	{
+		if (It == TestAsset)
+		{
+			return true;
+		}
+	}
+
+	return false;
+}
+
+EBlackboardCompare::Type UBlackboardComponent::CompareKeyValues(TSubclassOf<class UBlackboardKeyType> KeyType, uint8 KeyA, uint8 KeyB) const
+{
+	UBlackboardKeyType* KeyCDO = KeyType->GetDefaultObject<UBlackboardKeyType>();
+	return KeyCDO->Compare(GetKeyRawData(KeyA), GetKeyRawData(KeyB));
+}
+
+FString UBlackboardComponent::GetDebugInfoString(EBlackboardDescription::Type Mode) const 
+{
+	FString DebugString = FString::Printf(TEXT("Blackboard (asset: %s)\n"), *GetNameSafe(BlackboardAsset));
+
+	TArray<FString> KeyDesc;
+	uint8 Offset = 0;
+	for (UBlackboardData* It = BlackboardAsset; It; It = It->Parent)
+	{
+		for (int32 i = 0; i < It->Keys.Num(); i++)
+		{
+			KeyDesc.Add(DescribeKeyValue(i + Offset, Mode));
+		}
+		Offset += It->Keys.Num();
+	}
+	
+	KeyDesc.Sort();
+	for (int32 i = 0; i < KeyDesc.Num(); i++)
+	{
+		DebugString += TEXT("  ");
+		DebugString += KeyDesc[i];
+		DebugString += TEXT('\n');
+	}
+
+	if (Mode == EBlackboardDescription::Full && BlackboardAsset)
+	{
+		DebugString += TEXT("Observed Keys:\n");
+
+		TArray<uint8> ObserversKeys;
+		if (Observers.GetKeys(ObserversKeys) > 0)
+		{
+			DebugString += TEXT("Observed Keys:\n");
+		
+			for (int32 KeyIndex = 0; KeyIndex < ObserversKeys.Num(); ++KeyIndex)
+			{
+				const uint8 KeyID = ObserversKeys[KeyIndex];
+				//@todo shouldn't be using a localized value?; GetKeyName() [10/11/2013 justin.sargent]
+				DebugString += FString::Printf(TEXT("  %s:\n"), *BlackboardAsset->GetKeyName(KeyID).ToString());
+			}
+		}
+		else
+		{
+			DebugString += TEXT("  NONE\n");
+		}
+	}
+
+	return DebugString;
+}
+
+FString UBlackboardComponent::DescribeKeyValue(uint8 KeyID, EBlackboardDescription::Type Mode) const
+{
+	FString Description;
+
+	const FBlackboardEntry* Key = BlackboardAsset ? BlackboardAsset->GetKey(KeyID) : NULL;
+	if (Key)
+	{
+		const uint8* ValueData = GetKeyRawData(KeyID);
+		FString ValueDesc = Key->KeyType ? *(Key->KeyType->DescribeValue(ValueData)) : TEXT("empty");
+
+		if (Mode == EBlackboardDescription::OnlyValue)
+		{
+			Description = ValueDesc;
+		}
+		else if (Mode == EBlackboardDescription::KeyWithValue)
+		{
+			Description = FString::Printf(TEXT("%s: %s"), *(Key->EntryName.ToString()), *ValueDesc);
+		}
+		else
+		{
+			const FString CommonTypePrefix = UBlackboardKeyType::StaticClass()->GetName().AppendChar(TEXT('_'));
+			const FString FullKeyType = Key->KeyType ? GetNameSafe(Key->KeyType->GetClass()) : FString();
+			const FString DescKeyType = FullKeyType.StartsWith(CommonTypePrefix) ? FullKeyType.RightChop(CommonTypePrefix.Len()) : FullKeyType;
+
+			Description = FString::Printf(TEXT("%s [%s]: %s"), *(Key->EntryName.ToString()), *DescKeyType, *ValueDesc);
+		}
+	}
+
+	return Description;
+}
+
+#if ENABLE_VISUAL_LOG
+
+void UBlackboardComponent::DescribeSelfToVisLog(struct FVisLogEntry* Snapshot) const
+{
+	FVisLogEntry::FStatusCategory Category;
+	Category.Category = FString::Printf(TEXT("Blackboard (asset: %s)"), *GetNameSafe(BlackboardAsset));
+
+	for (UBlackboardData* It = BlackboardAsset; It; It = It->Parent)
+	{
+		for (int32 i = 0; i < It->Keys.Num(); i++)
+		{
+			const FBlackboardEntry& Key = It->Keys[i];
+
+			const uint8* ValueData = GetKeyRawData(It->GetFirstKeyID() + i);
+			FString ValueDesc = Key.KeyType ? *(Key.KeyType->DescribeValue(ValueData)) : TEXT("empty");
+
+			Category.Add(Key.EntryName.ToString(), ValueDesc);
+		}
+	}
+
+	Category.Data.Sort();
+	Snapshot->Status.Add(Category);
+}
+
+#endif
+
+UObject* UBlackboardComponent::GetValueAsObject(const FName& KeyName) const
+{
+	uint8 KeyID = GetKeyID(KeyName);
+	if (GetKeyType(KeyID) != UBlackboardKeyType_Object::StaticClass())
+	{
+		KeyID = INDEX_NONE;
+	}
+
+	return GetValueAsObject(KeyID);
+}
+
+UObject* UBlackboardComponent::GetValueAsObject(uint8 KeyID) const
+{
+	const uint8* RawData = GetKeyRawData(KeyID);
+	return RawData ? UBlackboardKeyType_Object::GetValue(RawData) : NULL;
+}
+
+UClass* UBlackboardComponent::GetValueAsClass(const FName& KeyName) const
+{
+	uint8 KeyID = GetKeyID(KeyName);
+	if (GetKeyType(KeyID) != UBlackboardKeyType_Class::StaticClass())
+	{
+		KeyID = INDEX_NONE;
+	}
+
+	return GetValueAsClass(KeyID);
+}
+
+UClass* UBlackboardComponent::GetValueAsClass(uint8 KeyID) const
+{
+	const uint8* RawData = GetKeyRawData(KeyID);
+	return RawData ? UBlackboardKeyType_Class::GetValue(RawData) : NULL;
+}
+
+uint8 UBlackboardComponent::GetValueAsEnum(const FName& KeyName) const
+{
+	uint8 KeyID = GetKeyID(KeyName);
+	if (GetKeyType(KeyID) != UBlackboardKeyType_Enum::StaticClass() &&
+		GetKeyType(KeyID) != UBlackboardKeyType_NativeEnum::StaticClass())
+	{
+		KeyID = INDEX_NONE;
+	}
+
+	return GetValueAsEnum(KeyID);
+}
+
+uint8 UBlackboardComponent::GetValueAsEnum(uint8 KeyID) const
+{
+	const uint8* RawData = GetKeyRawData(KeyID);
+	return RawData ? UBlackboardKeyType_Enum::GetValue(RawData) : 0;
+}
+
+int32 UBlackboardComponent::GetValueAsInt(const FName& KeyName) const
+{
+	uint8 KeyID = GetKeyID(KeyName);
+	if (GetKeyType(KeyID) != UBlackboardKeyType_Int::StaticClass())
+	{
+		KeyID = INDEX_NONE;
+	}
+
+	return GetValueAsInt(KeyID);
+}
+
+int32 UBlackboardComponent::GetValueAsInt(uint8 KeyID) const
+{
+	const uint8* RawData = GetKeyRawData(KeyID);
+	return RawData ? UBlackboardKeyType_Int::GetValue(RawData) : 0;
+}
+
+float UBlackboardComponent::GetValueAsFloat(const FName& KeyName) const
+{
+	uint8 KeyID = GetKeyID(KeyName);
+	if (GetKeyType(KeyID) != UBlackboardKeyType_Float::StaticClass())
+	{
+		KeyID = INDEX_NONE;
+	}
+
+	return GetValueAsFloat(KeyID);
+}
+
+float UBlackboardComponent::GetValueAsFloat(uint8 KeyID) const
+{
+	const uint8* RawData = GetKeyRawData(KeyID);
+	return RawData ? UBlackboardKeyType_Float::GetValue(RawData) : 0.0f;
+}
+
+bool UBlackboardComponent::GetValueAsBool(const FName& KeyName) const
+{
+	uint8 KeyID = GetKeyID(KeyName);
+	if (GetKeyType(KeyID) != UBlackboardKeyType_Bool::StaticClass())
+	{
+		KeyID = INDEX_NONE;
+	}
+
+	return GetValueAsBool(KeyID);
+}
+
+bool UBlackboardComponent::GetValueAsBool(uint8 KeyID) const
+{
+	const uint8* RawData = GetKeyRawData(KeyID);
+	return RawData ? UBlackboardKeyType_Bool::GetValue(RawData) : false;
+}
+
+FString UBlackboardComponent::GetValueAsString(const FName& KeyName) const
+{
+	uint8 KeyID = GetKeyID(KeyName);
+	if (GetKeyType(KeyID) != UBlackboardKeyType_String::StaticClass())
+	{
+		KeyID = INDEX_NONE;
+	}
+
+	return GetValueAsString(KeyID);
+}
+
+FString UBlackboardComponent::GetValueAsString(uint8 KeyID) const
+{
+	const uint8* RawData = GetKeyRawData(KeyID);
+	return RawData ? UBlackboardKeyType_String::GetValue(RawData) : FString();
+}
+
+FName UBlackboardComponent::GetValueAsName(const FName& KeyName) const
+{
+	uint8 KeyID = GetKeyID(KeyName);
+	if (GetKeyType(KeyID) != UBlackboardKeyType_Name::StaticClass())
+	{
+		KeyID = INDEX_NONE;
+	}
+
+	return GetValueAsName(KeyID);
+}
+
+FName UBlackboardComponent::GetValueAsName(uint8 KeyID) const
+{
+	const uint8* RawData = GetKeyRawData(KeyID);
+	return RawData ? UBlackboardKeyType_Name::GetValue(RawData) : NAME_None;
+}
+
+FVector UBlackboardComponent::GetValueAsVector(const FName& KeyName) const
+{
+	uint8 KeyID = GetKeyID(KeyName);
+	if (GetKeyType(KeyID) != UBlackboardKeyType_Vector::StaticClass())
+	{
+		KeyID = INDEX_NONE;
+	}
+
+	return GetValueAsVector(KeyID);
+}
+
+FVector UBlackboardComponent::GetValueAsVector(uint8 KeyID) const
+{
+	const uint8* RawData = GetKeyRawData(KeyID);
+	return RawData ? UBlackboardKeyType_Vector::GetValue(RawData) : FVector::ZeroVector;
+}
+
+void UBlackboardComponent::SetValueAsObject(const FName& KeyName, UObject* ObjectValue)
+{
+	const uint8 KeyID = GetKeyID(KeyName);
+	if (GetKeyType(KeyID) == UBlackboardKeyType_Object::StaticClass())
+	{
+		SetValueAsObject(KeyID, ObjectValue);
+	}
+}
+
+void UBlackboardComponent::SetValueAsObject(uint8 KeyID, UObject* ObjectValue)
+{
+	uint8* RawData = GetKeyRawData(KeyID);
+	if (RawData)
+	{
+		const bool bChanged = UBlackboardKeyType_Object::SetValue(RawData, ObjectValue);
+		if (bChanged)
+		{
+			NotifyObservers(KeyID);
+		}
+	}
+}
+
+void UBlackboardComponent::SetValueAsClass(const FName& KeyName, UClass* ClassValue)
+{
+	const uint8 KeyID = GetKeyID(KeyName);
+	if (GetKeyType(KeyID) == UBlackboardKeyType_Class::StaticClass())
+	{
+		SetValueAsClass(KeyID, ClassValue);
+	}
+}
+
+void UBlackboardComponent::SetValueAsClass(uint8 KeyID, UClass* ClassValue)
+{
+	uint8* RawData = GetKeyRawData(KeyID);
+	if (RawData)
+	{
+		const bool bChanged = UBlackboardKeyType_Class::SetValue(RawData, ClassValue);
+		if (bChanged)
+		{
+			NotifyObservers(KeyID);
+		}
+	}
+}
+
+void UBlackboardComponent::SetValueAsEnum(const FName& KeyName, uint8 EnumValue)
+{
+	const uint8 KeyID = GetKeyID(KeyName);
+	if (GetKeyType(KeyID) == UBlackboardKeyType_Enum::StaticClass() ||
+		GetKeyType(KeyID) == UBlackboardKeyType_NativeEnum::StaticClass())
+	{
+		SetValueAsEnum(KeyID, EnumValue);
+	}
+}
+
+void UBlackboardComponent::SetValueAsEnum(uint8 KeyID, uint8 EnumValue)
+{
+	uint8* RawData = GetKeyRawData(KeyID);
+	if (RawData)
+	{
+		const bool bChanged = UBlackboardKeyType_Enum::SetValue(RawData, EnumValue);
+		if (bChanged)
+		{
+			NotifyObservers(KeyID);
+		}
+	}
+}
+
+void UBlackboardComponent::SetValueAsInt(const FName& KeyName, int32 IntValue)
+{
+	const uint8 KeyID = GetKeyID(KeyName);
+	if (GetKeyType(KeyID) == UBlackboardKeyType_Int::StaticClass())
+	{
+		SetValueAsInt(KeyID, IntValue);
+	}
+}
+
+void UBlackboardComponent::SetValueAsInt(uint8 KeyID, int32 IntValue)
+{
+	uint8* RawData = GetKeyRawData(KeyID);
+	if (RawData)
+	{
+		const bool bChanged = UBlackboardKeyType_Int::SetValue(RawData, IntValue);
+		if (bChanged)
+		{
+			NotifyObservers(KeyID);
+		}
+	}
+}
+
+void UBlackboardComponent::SetValueAsFloat(const FName& KeyName, float FloatValue)
+{
+	const uint8 KeyID = GetKeyID(KeyName);
+	if (GetKeyType(KeyID) == UBlackboardKeyType_Float::StaticClass())
+	{
+		SetValueAsFloat(KeyID, FloatValue);
+	}
+}
+
+void UBlackboardComponent::SetValueAsFloat(uint8 KeyID, float FloatValue)
+{
+	uint8* RawData = GetKeyRawData(KeyID);
+	if (RawData)
+	{
+		const bool bChanged = UBlackboardKeyType_Float::SetValue(RawData, FloatValue);
+		if (bChanged)
+		{
+			NotifyObservers(KeyID);
+		}
+	}
+}
+
+void UBlackboardComponent::SetValueAsBool(const FName& KeyName, bool BoolValue)
+{
+	const uint8 KeyID = GetKeyID(KeyName);
+	if (GetKeyType(KeyID) == UBlackboardKeyType_Bool::StaticClass())
+	{
+		SetValueAsBool(KeyID, BoolValue);
+	}
+}
+
+void UBlackboardComponent::SetValueAsBool(uint8 KeyID, bool BoolValue)
+{
+	uint8* RawData = GetKeyRawData(KeyID);
+	if (RawData)
+	{
+		const bool bChanged = UBlackboardKeyType_Bool::SetValue(RawData, BoolValue);
+		if (bChanged)
+		{
+			NotifyObservers(KeyID);
+		}
+	}
+}
+
+void UBlackboardComponent::SetValueAsString(const FName& KeyName, FString StringValue)
+{
+	const uint8 KeyID = GetKeyID(KeyName);
+	if (GetKeyType(KeyID) == UBlackboardKeyType_String::StaticClass())
+	{
+		SetValueAsString(KeyID, StringValue);
+	}
+}
+
+void UBlackboardComponent::SetValueAsString(uint8 KeyID, const FString& StringValue)
+{
+	uint8* RawData = GetKeyRawData(KeyID);
+	if (RawData)
+	{
+		const bool bChanged = UBlackboardKeyType_String::SetValue(RawData, StringValue);
+		if (bChanged)
+		{
+			NotifyObservers(KeyID);
+		}
+	}
+}
+
+void UBlackboardComponent::SetValueAsName(const FName& KeyName, FName NameValue)
+{
+	const uint8 KeyID = GetKeyID(KeyName);
+	if (GetKeyType(KeyID) == UBlackboardKeyType_Name::StaticClass())
+	{
+		SetValueAsName(KeyID, NameValue);
+	}
+}
+
+void UBlackboardComponent::SetValueAsName(uint8 KeyID, const FName& NameValue)
+{
+	uint8* RawData = GetKeyRawData(KeyID);
+	if (RawData)
+	{
+		const bool bChanged = UBlackboardKeyType_Name::SetValue(RawData, NameValue);
+		if (bChanged)
+		{
+			NotifyObservers(KeyID);
+		}
+	}
+}
+
+void UBlackboardComponent::SetValueAsVector(const FName& KeyName, FVector VectorValue)
+{
+	const uint8 KeyID = GetKeyID(KeyName);
+	if (GetKeyType(KeyID) == UBlackboardKeyType_Vector::StaticClass())
+	{
+		SetValueAsVector(KeyID, VectorValue);
+	}
+}
+
+void UBlackboardComponent::SetValueAsVector(uint8 KeyID, const FVector& VectorValue)
+{
+	uint8* RawData = GetKeyRawData(KeyID);
+	if (RawData)
+	{
+		const bool bChanged = UBlackboardKeyType_Vector::SetValue(RawData, VectorValue);
+		if (bChanged)
+		{
+			NotifyObservers(KeyID);
+		}
+	}
+}
+
+void UBlackboardComponent::ClearValueAsVector(const FName& KeyName)
+{
+	const uint8 KeyID = GetKeyID(KeyName);
+	if (GetKeyType(KeyID) == UBlackboardKeyType_Vector::StaticClass())
+	{
+		ClearValueAsVector(KeyID);
+	}
+}
+
+void UBlackboardComponent::ClearValueAsVector(uint8 KeyID)
+{
+	uint8* RawData = GetKeyRawData(KeyID);
+	if (RawData)
+	{
+		const bool bChanged = UBlackboardKeyType_Vector::SetValue(RawData, FAISystem::InvalidLocation);
+		if (bChanged)
+		{
+			NotifyObservers(KeyID);
+		}
+	}
+}
+
+bool UBlackboardComponent::GetLocationFromEntry(const FName& KeyName, FVector& ResultLocation) const
+{
+	const uint8 KeyID = GetKeyID(KeyName);
+	return GetLocationFromEntry(KeyID, ResultLocation);
+}
+
+bool UBlackboardComponent::GetLocationFromEntry(uint8 KeyID, FVector& ResultLocation) const
+{
+	if (BlackboardAsset && ValueOffsets.IsValidIndex(KeyID))
+	{
+		const FBlackboardEntry* EntryInfo = BlackboardAsset->GetKey(KeyID);
+		if (EntryInfo && EntryInfo->KeyType)
+		{
+			const uint8* ValueData = ValueMemory.GetTypedData() + ValueOffsets[KeyID];
+			return EntryInfo->KeyType->GetLocation(ValueData, ResultLocation);
+		}
+	}
+
+	return false;
+}
+
+bool UBlackboardComponent::GetRotationFromEntry(const FName& KeyName, FRotator& ResultRotation) const
+{
+	const uint8 KeyID = GetKeyID(KeyName);
+	return GetRotationFromEntry(KeyID, ResultRotation);
+}
+
+bool UBlackboardComponent::GetRotationFromEntry(uint8 KeyID, FRotator& ResultRotation) const
+{
+	if (BlackboardAsset && ValueOffsets.IsValidIndex(KeyID))
+	{
+		const FBlackboardEntry* EntryInfo = BlackboardAsset->GetKey(KeyID);
+		if (EntryInfo && EntryInfo->KeyType)
+		{
+			const uint8* ValueData = ValueMemory.GetTypedData() + ValueOffsets[KeyID];
+			return EntryInfo->KeyType->GetRotation(ValueData, ResultRotation);
+		}
+	}
+
+	return false;
+}
diff --git a/Engine/Source/Runtime/AIModule/Private/BehaviorTree/BlackboardData.cpp b/Engine/Source/Runtime/AIModule/Private/BehaviorTree/BlackboardData.cpp
new file mode 100644
--- /dev/null
+++ b/Engine/Source/Runtime/AIModule/Private/BehaviorTree/BlackboardData.cpp
@@ -0,0 +1,165 @@
+// Copyright 1998-2014 Epic Games, Inc. All Rights Reserved.
+
+#include "AIModulePrivate.h"
+#include "BehaviorTree/Blackboard/BlackboardKeyType_Object.h"
+#include "BehaviorTree/BlackboardData.h"
+
+UBlackboardData::FKeyUpdate UBlackboardData::OnUpdateKeys;
+
+static void UpdatePersistentKeys(class UBlackboardData* Asset)
+{
+	UBlackboardKeyType_Object* SelfKeyType = Asset->UpdatePersistentKey<UBlackboardKeyType_Object>(FBlackboard::KeySelf);
+	if (SelfKeyType)
+	{
+		SelfKeyType->BaseClass = AActor::StaticClass();
+	}
+}
+
+UBlackboardData::UBlackboardData(const class FPostConstructInitializeProperties& PCIP) : Super(PCIP)
+{
+	if (HasAnyFlags(RF_ClassDefaultObject))
+	{
+		OnUpdateKeys.AddStatic(&UpdatePersistentKeys);
+	}
+}
+
+uint8 UBlackboardData::GetKeyID(const FName& KeyName) const
+{
+	return InternalGetKeyID(KeyName, CheckParentKeys);
+}
+
+FName UBlackboardData::GetKeyName(uint8 KeyID) const
+{
+	const FBlackboardEntry* KeyEntry = GetKey(KeyID);
+	return KeyEntry ? KeyEntry->EntryName : NAME_None;
+}
+
+TSubclassOf<UBlackboardKeyType> UBlackboardData::GetKeyType(uint8 KeyID) const
+{
+	const FBlackboardEntry* KeyEntry = GetKey(KeyID);
+	return KeyEntry ? KeyEntry->KeyType->GetClass() : NULL;
+}
+
+const FBlackboardEntry* UBlackboardData::GetKey(uint8 KeyID) const
+{
+	if (KeyID != InvalidKeyID)
+	{
+		if (KeyID >= FirstKeyID)
+		{
+			return &Keys[KeyID - FirstKeyID];
+		}
+		else if (Parent)
+		{
+			return Parent->GetKey(KeyID);
+		}
+	}
+
+	return NULL;
+}
+
+int32 UBlackboardData::GetNumKeys() const
+{
+	return FirstKeyID + Keys.Num();
+}
+
+uint8 UBlackboardData::InternalGetKeyID(const FName& KeyName, EKeyLookupMode LookupMode) const
+{
+	for (int32 i = 0; i < Keys.Num(); i++)
+	{
+		if (Keys[i].EntryName == KeyName)
+		{
+			return i + FirstKeyID;
+		}
+	}
+
+	return Parent && (LookupMode == CheckParentKeys) ? Parent->InternalGetKeyID(KeyName, LookupMode) : InvalidKeyID;
+}
+
+bool UBlackboardData::IsValid() const
+{
+	if (Parent)
+	{
+		for (int32 i = 0; i < Keys.Num(); i++)
+		{
+			const uint8 KeyID = Parent->InternalGetKeyID(Keys[i].EntryName, CheckParentKeys);
+			if (KeyID != InvalidKeyID)
+			{
+				UE_LOG(LogBehaviorTree, Warning, TEXT("Blackboard asset (%s) has duplicated key (%s) in parent chain!"),
+					*GetName(), *Keys[i].EntryName.ToString());
+
+				return false;
+			}
+		}
+	}
+
+	return true;
+}
+
+void UBlackboardData::PostLoad()
+{
+	Super::PostLoad();
+
+	UpdateParentKeys();
+}
+
+#if WITH_EDITOR
+
+void UBlackboardData::PostEditChangeProperty(struct FPropertyChangedEvent& PropertyChangedEvent)
+{
+	Super::PostEditChangeProperty(PropertyChangedEvent);
+
+	if (PropertyChangedEvent.Property && PropertyChangedEvent.Property->GetFName() == GET_MEMBER_NAME_CHECKED(UBlackboardData, Parent))
+	{
+		UpdateParentKeys();
+	}
+}
+
+#endif // WITH_EDITOR
+
+static bool ContainsKeyName(FName KeyName, const TArray<FBlackboardEntry>& Keys, const TArray<FBlackboardEntry>& ParentKeys)
+{
+	for (int32 i = 0; i < Keys.Num(); i++)
+	{
+		if (Keys[i].EntryName == KeyName)
+		{
+			return true;
+		}
+	}
+
+	for (int32 i = 0; i < ParentKeys.Num(); i++)
+	{
+		if (ParentKeys[i].EntryName == KeyName)
+		{
+			return true;
+		}
+	}
+
+	return false;
+}
+
+void UBlackboardData::UpdateParentKeys()
+{
+	if (Parent == this)
+	{
+		Parent = NULL;
+	}
+
+#if WITH_EDITORONLY_DATA
+	ParentKeys.Reset();
+
+	for (UBlackboardData* It = Parent; It; It = It->Parent)
+	{
+		for (int32 i = 0; i < It->Keys.Num(); i++)
+		{
+			const bool bAlreadyExist = ContainsKeyName(It->Keys[i].EntryName, Keys, ParentKeys);
+			if (!bAlreadyExist)
+			{
+				ParentKeys.Add(It->Keys[i]);
+			}
+		}
+	}
+#endif // WITH_EDITORONLY_DATA
+
+	FirstKeyID = Parent ? Parent->GetNumKeys() : 0;
+	OnUpdateKeys.Broadcast(this);
+}
diff --git a/Engine/Source/Runtime/AIModule/Private/BehaviorTree/BlueprintNodeHelpers.cpp b/Engine/Source/Runtime/AIModule/Private/BehaviorTree/BlueprintNodeHelpers.cpp
new file mode 100644
--- /dev/null
+++ b/Engine/Source/Runtime/AIModule/Private/BehaviorTree/BlueprintNodeHelpers.cpp
@@ -0,0 +1,221 @@
+// Copyright 1998-2014 Epic Games, Inc. All Rights Reserved.
+
+#include "AIModulePrivate.h"
+#include "BlueprintNodeHelpers.h"
+
+namespace BlueprintNodeHelpers
+{
+	uint16 GetPropertiesMemorySize(const TArray<UProperty*>& PropertyData)
+	{
+		int32 TotalMem = 0;
+		for (int32 i = 0; i < PropertyData.Num(); i++)
+		{
+			TotalMem += PropertyData[i]->GetSize();
+		}
+
+		if (TotalMem > MAX_uint16)
+		{
+			TotalMem = 0;
+		}
+
+		return TotalMem;
+	}
+
+#define GET_STRUCT_NAME_CHECKED(StructName) \
+	((void)sizeof(StructName), TEXT(#StructName))
+
+	bool CanUsePropertyType(const UProperty* TestProperty)
+	{
+		if (TestProperty->IsA(UNumericProperty::StaticClass()) ||
+			TestProperty->IsA(UBoolProperty::StaticClass()) ||
+			TestProperty->IsA(UNameProperty::StaticClass()))
+		{
+			return true;
+		}
+
+		const UStructProperty* StructProp = Cast<const UStructProperty>(TestProperty);
+		if (StructProp)
+		{
+			FString CPPType = StructProp->GetCPPType(NULL, CPPF_None);
+			if (CPPType.Contains(GET_STRUCT_NAME_CHECKED(FVector)) ||
+				CPPType.Contains(GET_STRUCT_NAME_CHECKED(FRotator)))
+			{
+				return true;
+			}
+		}
+
+		return false;
+	}
+
+	void CollectPropertyData(const UObject* Ob, const UClass* StopAtClass, TArray<UProperty*>& PropertyData)
+	{
+		UE_LOG(LogBehaviorTree, Verbose, TEXT("Looking for runtime properties of class: %s"), *GetNameSafe(Ob->GetClass()));
+
+		PropertyData.Reset();
+		for (UProperty* TestProperty = Ob->GetClass()->PropertyLink; TestProperty; TestProperty = TestProperty->PropertyLinkNext)
+		{
+			// stop when reaching base class
+			if (TestProperty->GetOuter() == StopAtClass)
+			{
+				break;
+			}
+
+			// skip properties without any setup data
+			if (TestProperty->HasAnyPropertyFlags(CPF_Transient) ||
+				TestProperty->HasAnyPropertyFlags(CPF_DisableEditOnInstance) == false)
+			{
+				continue;
+			}
+
+			// serialize only simple types
+			if (CanUsePropertyType(TestProperty))
+			{
+				UE_LOG(LogBehaviorTree, Verbose, TEXT("> name: '%s'"), *GetNameSafe(TestProperty));
+				PropertyData.Add(TestProperty);
+			}
+		}
+	}
+
+	FString DescribeProperty(const UProperty* Prop, const uint8* PropertyAddr)
+	{
+		FString ExportedStringValue;
+		const UStructProperty* StructProp = Cast<const UStructProperty>(Prop);
+		const UFloatProperty* FloatProp = Cast<const UFloatProperty>(Prop);
+
+		if (StructProp && StructProp->GetCPPType(NULL, CPPF_None).Contains(GET_STRUCT_NAME_CHECKED(FBlackboardKeySelector)))
+		{
+			// special case for blackboard key selectors
+			const FBlackboardKeySelector* PropertyValue = (const FBlackboardKeySelector*)PropertyAddr;
+			ExportedStringValue = PropertyValue->SelectedKeyName.ToString();
+		}
+		else if (FloatProp)
+		{
+			// special case for floats to remove unnecessary zeros
+			const float FloatValue = FloatProp->GetPropertyValue(PropertyAddr);
+			ExportedStringValue = FString::SanitizeFloat(FloatValue);
+		}
+		else
+		{
+			Prop->ExportTextItem(ExportedStringValue, PropertyAddr, NULL, NULL, 0, NULL);
+		}
+
+		const bool bIsBool = Prop->IsA(UBoolProperty::StaticClass());
+		return FString::Printf(TEXT("%s: %s"), *FName::NameToDisplayString(Prop->GetName(), bIsBool), *ExportedStringValue);
+	}
+
+	void CollectBlackboardSelectors(const UObject* Ob, const UClass* StopAtClass, TArray<FName>& KeyNames)
+	{
+		for (UProperty* TestProperty = Ob->GetClass()->PropertyLink; TestProperty; TestProperty = TestProperty->PropertyLinkNext)
+		{
+			// stop when reaching base class
+			if (TestProperty->GetOuter() == StopAtClass)
+			{
+				break;
+			}
+
+			// skip properties without any setup data	
+			if (TestProperty->HasAnyPropertyFlags(CPF_Transient) ||
+				TestProperty->HasAnyPropertyFlags(CPF_DisableEditOnInstance))
+			{
+				continue;
+			}
+
+			const UStructProperty* StructProp = Cast<const UStructProperty>(TestProperty);
+			if (StructProp && StructProp->GetCPPType(NULL, CPPF_None).Contains(GET_STRUCT_NAME_CHECKED(FBlackboardKeySelector)))
+			{
+				const FBlackboardKeySelector* PropData = TestProperty->ContainerPtrToValuePtr<FBlackboardKeySelector>(Ob);
+				KeyNames.AddUnique(PropData->SelectedKeyName);
+			}
+		}
+	}
+
+#undef GET_STRUCT_NAME_CHECKED
+
+	FString CollectPropertyDescription(const UObject* Ob, const UClass* StopAtClass, const TArray<UProperty*>& PropertyData)
+	{
+		FString RetString;
+		for (UProperty* TestProperty = Ob->GetClass()->PropertyLink; TestProperty; TestProperty = TestProperty->PropertyLinkNext)
+		{
+			// stop when reaching base class
+			if (TestProperty->GetOuter() == StopAtClass)
+			{
+				break;
+			}
+
+			// skip properties without any setup data	
+			if (TestProperty->HasAnyPropertyFlags(CPF_Transient) ||
+				TestProperty->HasAnyPropertyFlags(CPF_DisableEditOnInstance) ||
+				PropertyData.Contains(TestProperty))
+			{
+				continue;
+			}
+
+			if (TestProperty->IsA(UClassProperty::StaticClass()) ||
+				TestProperty->IsA(UStructProperty::StaticClass()) ||
+				CanUsePropertyType(TestProperty))
+			{
+				if (RetString.Len())
+				{
+					RetString.AppendChar(TEXT('\n'));
+				}
+
+				const uint8* PropData = TestProperty->ContainerPtrToValuePtr<uint8>(Ob);
+				RetString += DescribeProperty(TestProperty, PropData);
+			}
+		}
+
+		return RetString;
+	}
+
+	void DescribeRuntimeValues(const UObject* Ob, const TArray<UProperty*>& PropertyData, TArray<FString>& RuntimeValues)
+	{
+		for (int32 i = 0; i < PropertyData.Num(); i++)
+		{
+			UProperty* TestProperty = PropertyData[i];
+			const uint8* PropAddr = TestProperty->ContainerPtrToValuePtr<uint8>(Ob);
+
+			RuntimeValues.Add(DescribeProperty(TestProperty, PropAddr));
+		}
+	}
+
+	bool FindNodeOwner(AActor* OwningActor, class UBTNode* Node, UBehaviorTreeComponent*& OwningComp, int32& OwningInstanceIdx)
+	{
+		bool bFound = false;
+
+		APawn* OwningPawn = Cast<APawn>(OwningActor);
+		if (OwningPawn && OwningPawn->Controller)
+		{
+			bFound = FindNodeOwner(OwningPawn->Controller, Node, OwningComp, OwningInstanceIdx);
+		}
+
+		if (OwningActor && !bFound)
+		{
+			TArray<UBehaviorTreeComponent*> Components;
+			OwningActor->GetComponents(Components);
+			for (int32 i = 0; i < Components.Num(); i++)
+			{
+				UBehaviorTreeComponent* BTComp = Components[i];
+				if (BTComp)
+				{
+					const int32 InstanceIdx = BTComp->FindInstanceContainingNode(Node);
+					if (InstanceIdx != INDEX_NONE)
+					{
+						OwningComp = BTComp;
+						OwningInstanceIdx = InstanceIdx;
+						bFound = true;
+						break;
+					}
+				}
+			}
+		}
+
+		return bFound;
+	}
+
+	void AbortLatentActions(AActor* OwningActor, const UObject* Ob)
+	{
+		UWorld* MyWorld = OwningActor->GetWorld();
+		MyWorld->GetLatentActionManager().RemoveActionsForObject(Ob);
+		MyWorld->GetTimerManager().ClearAllTimersForObject(Ob);
+	}
+}
diff --git a/Engine/Source/Runtime/AIModule/Private/BehaviorTree/Composites/BTComposite_Selector.cpp b/Engine/Source/Runtime/AIModule/Private/BehaviorTree/Composites/BTComposite_Selector.cpp
new file mode 100644
--- /dev/null
+++ b/Engine/Source/Runtime/AIModule/Private/BehaviorTree/Composites/BTComposite_Selector.cpp
@@ -0,0 +1,31 @@
+// Copyright 1998-2014 Epic Games, Inc. All Rights Reserved.
+
+#include "AIModulePrivate.h"
+#include "BehaviorTree/Composites/BTComposite_Selector.h"
+
+UBTComposite_Selector::UBTComposite_Selector(const class FPostConstructInitializeProperties& PCIP) : Super(PCIP)
+{
+	NodeName = "Selector";
+
+	OptionalDeactivationResult = EBTNodeResult::Failed;
+	OnNextChild.BindUObject(this, &UBTComposite_Selector::GetNextChildHandler);
+}
+
+int32 UBTComposite_Selector::GetNextChildHandler(struct FBehaviorTreeSearchData& SearchData, int32 PrevChild, EBTNodeResult::Type LastResult) const
+{
+	// success = quit
+	int32 NextChildIdx = BTSpecialChild::ReturnToParent;
+
+	if (PrevChild == BTSpecialChild::NotInitialized)
+	{
+		// newly activated: start from first
+		NextChildIdx = 0;
+	}
+	else if ((LastResult == EBTNodeResult::Failed || LastResult == EBTNodeResult::Optional) && (PrevChild + 1) < GetChildrenNum())
+	{
+		// failed = choose next child
+		NextChildIdx = PrevChild + 1;
+	}
+
+	return NextChildIdx;
+}
diff --git a/Engine/Source/Runtime/AIModule/Private/BehaviorTree/Composites/BTComposite_Sequence.cpp b/Engine/Source/Runtime/AIModule/Private/BehaviorTree/Composites/BTComposite_Sequence.cpp
new file mode 100644
--- /dev/null
+++ b/Engine/Source/Runtime/AIModule/Private/BehaviorTree/Composites/BTComposite_Sequence.cpp
@@ -0,0 +1,41 @@
+// Copyright 1998-2014 Epic Games, Inc. All Rights Reserved.
+
+#include "AIModulePrivate.h"
+#include "BehaviorTree/Composites/BTComposite_Sequence.h"
+
+UBTComposite_Sequence::UBTComposite_Sequence(const class FPostConstructInitializeProperties& PCIP) : Super(PCIP)
+{
+	NodeName = "Sequence";
+	OptionalDeactivationResult = EBTNodeResult::Succeeded;
+
+	OnNextChild.BindUObject(this, &UBTComposite_Sequence::GetNextChildHandler);
+}
+
+int32 UBTComposite_Sequence::GetNextChildHandler(struct FBehaviorTreeSearchData& SearchData, int32 PrevChild, EBTNodeResult::Type LastResult) const
+{
+	// failure = quit
+	int32 NextChildIdx = BTSpecialChild::ReturnToParent;
+
+	if (PrevChild == BTSpecialChild::NotInitialized)
+	{
+		// newly activated: start from first
+		NextChildIdx = 0;
+	}
+	else if ((LastResult == EBTNodeResult::Succeeded || LastResult == EBTNodeResult::Optional) && (PrevChild + 1) < GetChildrenNum())
+	{
+		// success = choose next child
+		NextChildIdx = PrevChild + 1;
+	}
+
+	return NextChildIdx;
+}
+
+#if WITH_EDITOR
+
+bool UBTComposite_Sequence::CanAbortLowerPriority() const
+{
+	// don't allow aborting lower priorities, as it breaks sequence order and doesn't makes sense
+	return false;
+}
+
+#endif
diff --git a/Engine/Source/Runtime/AIModule/Private/BehaviorTree/Composites/BTComposite_SimpleParallel.cpp b/Engine/Source/Runtime/AIModule/Private/BehaviorTree/Composites/BTComposite_SimpleParallel.cpp
new file mode 100644
--- /dev/null
+++ b/Engine/Source/Runtime/AIModule/Private/BehaviorTree/Composites/BTComposite_SimpleParallel.cpp
@@ -0,0 +1,155 @@
+// Copyright 1998-2014 Epic Games, Inc. All Rights Reserved.
+
+#include "AIModulePrivate.h"
+#include "BehaviorTree/Composites/BTComposite_SimpleParallel.h"
+
+UBTComposite_SimpleParallel::UBTComposite_SimpleParallel(const class FPostConstructInitializeProperties& PCIP) : Super(PCIP)
+{
+	NodeName = "Simple Parallel";
+	bUseChildExecutionNotify = true;
+	bUseNodeDeactivationNotify = true;
+	OnNextChild.BindUObject(this, &UBTComposite_SimpleParallel::GetNextChildHandler);
+}
+
+int32 UBTComposite_SimpleParallel::GetNextChildHandler(struct FBehaviorTreeSearchData& SearchData, int32 PrevChild, EBTNodeResult::Type LastResult) const
+{
+	FBTParallelMemory* MyMemory = GetNodeMemory<FBTParallelMemory>(SearchData);
+	int32 NextChildIdx = BTSpecialChild::ReturnToParent;
+
+	if (PrevChild == BTSpecialChild::NotInitialized)
+	{
+		NextChildIdx = EBTParallelChild::MainTask;
+	}
+	else if ((MyMemory->bMainTaskIsActive || MyMemory->bForceBackgroundTree) && !SearchData.OwnerComp->IsRestartPending())
+	{
+		// if main task is running - always pick background tree
+		// unless search is from abort from background tree - return to parent (and break from search when node gets deactivated)
+		NextChildIdx = EBTParallelChild::BackgroundTree;
+		MyMemory->bForceBackgroundTree = false;
+	}
+
+	return NextChildIdx;
+}
+
+void UBTComposite_SimpleParallel::NotifyChildExecution(class UBehaviorTreeComponent* OwnerComp, uint8* NodeMemory, int32 ChildIdx, EBTNodeResult::Type& NodeResult) const
+{
+	FBTParallelMemory* MyMemory = (FBTParallelMemory*)NodeMemory;
+	if (ChildIdx == EBTParallelChild::MainTask)
+	{
+		if (NodeResult == EBTNodeResult::InProgress)
+		{
+			EBTTaskStatus::Type Status = OwnerComp->GetTaskStatus(Children[EBTParallelChild::MainTask].ChildTask);
+			if (Status == EBTTaskStatus::Active)
+			{
+				// can't register if task is currently being aborted (latent abort returns in progress)
+
+				MyMemory->bMainTaskIsActive = true;
+				MyMemory->bForceBackgroundTree = false;
+				
+				OwnerComp->RegisterParallelTask(Children[EBTParallelChild::MainTask].ChildTask);
+				RequestDelayedExecution(OwnerComp, EBTNodeResult::Succeeded);
+			}
+		}
+		else if (MyMemory->bMainTaskIsActive)
+		{
+			MyMemory->MainTaskResult = NodeResult;
+			MyMemory->bMainTaskIsActive = false;
+
+			const int32 MyInstanceIdx = OwnerComp->FindInstanceContainingNode(this);
+
+			OwnerComp->UnregisterParallelTask(Children[EBTParallelChild::MainTask].ChildTask, MyInstanceIdx);
+			if (NodeResult != EBTNodeResult::Aborted)
+			{
+				// check if subtree should be aborted when task finished with success/failed result
+				if (FinishMode == EBTParallelMode::AbortBackground)
+				{
+					OwnerComp->RequestExecution((UBTCompositeNode*)this, MyInstanceIdx,
+						Children[EBTParallelChild::MainTask].ChildTask, EBTParallelChild::MainTask,
+						NodeResult);
+				}
+			}
+		}
+		else if (NodeResult == EBTNodeResult::Succeeded && FinishMode == EBTParallelMode::WaitForBackground)
+		{
+			// special case: if main task finished instantly, but composite is supposed to wait for background tree,
+			// request single run of background tree anyway
+
+			MyMemory->bForceBackgroundTree = true;
+
+			RequestDelayedExecution(OwnerComp, EBTNodeResult::Succeeded);
+		}
+	}
+}
+
+void UBTComposite_SimpleParallel::NotifyNodeDeactivation(struct FBehaviorTreeSearchData& SearchData, EBTNodeResult::Type& NodeResult) const
+{
+	FBTParallelMemory* MyMemory = GetNodeMemory<FBTParallelMemory>(SearchData);
+	const uint16 ActiveInstanceIdx = SearchData.OwnerComp->GetActiveInstanceIdx();
+
+	// modify node result if main task has finished
+	if (!MyMemory->bMainTaskIsActive)
+	{
+		NodeResult = MyMemory->MainTaskResult;
+	}
+
+	// remove main task
+	SearchData.AddUniqueUpdate(FBehaviorTreeSearchUpdate(Children[EBTParallelChild::MainTask].ChildTask, ActiveInstanceIdx, EBTNodeUpdateMode::Remove));
+
+	// remove all active nodes from background tree
+	const FBTNodeIndex FirstBackgroundIndex(ActiveInstanceIdx, GetChildExecutionIndex(EBTParallelChild::BackgroundTree, EBTChildIndex::FirstNode));
+	SearchData.OwnerComp->UnregisterAuxNodesUpTo(FirstBackgroundIndex);
+}
+
+uint16 UBTComposite_SimpleParallel::GetInstanceMemorySize() const
+{
+	return sizeof(FBTParallelMemory);
+}
+
+bool UBTComposite_SimpleParallel::CanPushSubtree(class UBehaviorTreeComponent* OwnerComp, uint8* NodeMemory, int32 ChildIdx) const
+{
+	return (ChildIdx != EBTParallelChild::MainTask);
+}
+
+FString UBTComposite_SimpleParallel::DescribeFinishMode(EBTParallelMode::Type Mode)
+{
+	static FString FinishDesc[] = { TEXT("AbortBackground"), TEXT("WaitForBackground") };
+	return FinishDesc[Mode];
+}
+
+FString UBTComposite_SimpleParallel::GetStaticDescription() const
+{
+	static FString FinishDesc[] = { TEXT("finish with main task"), TEXT("wait for subtree") };
+
+	return FString::Printf(TEXT("%s: %s"), *Super::GetStaticDescription(),
+		FinishMode < 2 ? *FinishDesc[FinishMode] : TEXT(""));
+}
+
+void UBTComposite_SimpleParallel::DescribeRuntimeValues(const class UBehaviorTreeComponent* OwnerComp, uint8* NodeMemory, EBTDescriptionVerbosity::Type Verbosity, TArray<FString>& Values) const
+{
+	Super::DescribeRuntimeValues(OwnerComp, NodeMemory, Verbosity, Values);
+
+	if (Verbosity == EBTDescriptionVerbosity::Detailed)
+	{
+		FBTParallelMemory* MyMemory = (FBTParallelMemory*)NodeMemory;
+		Values.Add(FString::Printf(TEXT("main task: %s"), MyMemory->bMainTaskIsActive ? TEXT("active") : TEXT("not active")));
+
+		if (MyMemory->bForceBackgroundTree)
+		{
+			Values.Add(TEXT("force background tree"));
+		}
+	}
+}
+
+#if WITH_EDITOR
+
+bool UBTComposite_SimpleParallel::CanAbortLowerPriority() const
+{
+	return false;
+}
+
+bool UBTComposite_SimpleParallel::CanAbortSelf() const
+{
+	return false;
+}
+
+#endif // WITH_EDITOR
diff --git a/Engine/Source/Runtime/AIModule/Private/BehaviorTree/Decorators/BTDecorator_Blackboard.cpp b/Engine/Source/Runtime/AIModule/Private/BehaviorTree/Decorators/BTDecorator_Blackboard.cpp
new file mode 100644
--- /dev/null
+++ b/Engine/Source/Runtime/AIModule/Private/BehaviorTree/Decorators/BTDecorator_Blackboard.cpp
@@ -0,0 +1,201 @@
+// Copyright 1998-2014 Epic Games, Inc. All Rights Reserved.
+
+#include "AIModulePrivate.h"
+#include "BehaviorTree/Decorators/BTDecorator_Blackboard.h"
+#include "BehaviorTree/BlackboardComponent.h"
+#include "BehaviorTree/Blackboard/BlackboardKeyAllTypes.h"
+
+UBTDecorator_Blackboard::UBTDecorator_Blackboard(const class FPostConstructInitializeProperties& PCIP) : Super(PCIP)
+{
+	NodeName = "Blackboard";
+}
+
+bool UBTDecorator_Blackboard::CalculateRawConditionValue(class UBehaviorTreeComponent* OwnerComp, uint8* NodeMemory) const
+{
+	const UBlackboardComponent* BlackboardComp = OwnerComp->GetBlackboardComponent();
+	return EvaluateOnBlackboard(BlackboardComp);
+}
+
+bool UBTDecorator_Blackboard::EvaluateOnBlackboard(const UBlackboardComponent* BlackboardComp) const
+{
+	bool bResult = false;
+	if (BlackboardComp && BlackboardKey.SelectedKeyType)
+	{
+		UBlackboardKeyType* KeyCDO = BlackboardKey.SelectedKeyType->GetDefaultObject<UBlackboardKeyType>();
+		const uint8* KeyMemory = BlackboardComp->GetKeyRawData(BlackboardKey.GetSelectedKeyID());
+
+		const EBlackboardKeyOperation::Type Op = KeyCDO->GetTestOperation();
+		switch (Op)
+		{
+		case EBlackboardKeyOperation::Basic:		bResult = KeyCDO->TestBasicOperation(KeyMemory, (EBasicKeyOperation::Type)OperationType); break;
+		case EBlackboardKeyOperation::Arithmetic:	bResult = KeyCDO->TestArithmeticOperation(KeyMemory, (EArithmeticKeyOperation::Type)OperationType, IntValue, FloatValue); break;
+		case EBlackboardKeyOperation::Text:			bResult = KeyCDO->TestTextOperation(KeyMemory, (ETextKeyOperation::Type)OperationType, StringValue); break;
+		default: break;
+		}
+	}
+
+	return bResult;
+}
+
+void UBTDecorator_Blackboard::OnBlackboardChange(const class UBlackboardComponent* Blackboard, uint8 ChangedKeyID)
+{
+	UBehaviorTreeComponent* BehaviorComp = Blackboard ? (UBehaviorTreeComponent*)Blackboard->GetBrainComponent() : NULL;
+	if (BlackboardKey.GetSelectedKeyID() == ChangedKeyID && BehaviorComp)
+	{
+		const bool bIsExecutingBranch = BehaviorComp->IsExecutingBranch(this, GetChildIndex());
+		const bool bPass = EvaluateOnBlackboard(Blackboard);
+
+		UE_VLOG(BehaviorComp->GetOwner(), LogBehaviorTree, Verbose, TEXT("%s, OnBlackboardChange[%s] pass:%d executing:%d => %s"),
+			*UBehaviorTreeTypes::DescribeNodeHelper(this),
+			*Blackboard->GetKeyName(ChangedKeyID).ToString(), bPass, bIsExecutingBranch,
+			(bIsExecutingBranch && !bPass) || (!bIsExecutingBranch && bPass) ? TEXT("restart") : TEXT("skip"));
+
+		if ((bIsExecutingBranch && !bPass) || 
+			(!bIsExecutingBranch && bPass))
+		{
+			BehaviorComp->RequestExecution(this);
+		}
+	}
+}
+
+void UBTDecorator_Blackboard::DescribeRuntimeValues(const class UBehaviorTreeComponent* OwnerComp, uint8* NodeMemory, EBTDescriptionVerbosity::Type Verbosity, TArray<FString>& Values) const
+{
+	Super::DescribeRuntimeValues(OwnerComp, NodeMemory, Verbosity, Values);
+
+	const UBlackboardComponent* BlackboardComp = OwnerComp->GetBlackboardComponent();
+	FString DescKeyValue;
+
+	if (BlackboardComp)
+	{
+		DescKeyValue = BlackboardComp->DescribeKeyValue(BlackboardKey.GetSelectedKeyID(), EBlackboardDescription::OnlyValue);
+	}
+
+	const bool bResult = EvaluateOnBlackboard(BlackboardComp);
+	Values.Add(FString::Printf(TEXT("value: %s (%s)"), *DescKeyValue, bResult ? TEXT("pass") : TEXT("fail")));
+}
+
+FString UBTDecorator_Blackboard::GetStaticDescription() const
+{
+	return FString::Printf(TEXT("%s: %s"), *Super::GetStaticDescription(), *CachedDescription);
+}
+
+#if WITH_EDITOR
+static UEnum* BasicOpEnum = NULL;
+static UEnum* ArithmeticOpEnum = NULL;
+static UEnum* TextOpEnum = NULL;
+
+static void CacheOperationEnums()
+{
+	if (BasicOpEnum == NULL)
+	{
+		BasicOpEnum = FindObject<UEnum>(ANY_PACKAGE, TEXT("EBasicKeyOperation"));
+		check(BasicOpEnum);
+	}
+
+	if (ArithmeticOpEnum == NULL)
+	{
+		ArithmeticOpEnum = FindObject<UEnum>(ANY_PACKAGE, TEXT("EArithmeticKeyOperation"));
+		check(ArithmeticOpEnum);
+	}
+
+	if (TextOpEnum == NULL)
+	{
+		TextOpEnum = FindObject<UEnum>(ANY_PACKAGE, TEXT("ETextKeyOperation"));
+		check(TextOpEnum);
+	}
+}
+
+void UBTDecorator_Blackboard::BuildDescription()
+{
+	UBlackboardData* BlackboardAsset = GetBlackboardAsset();
+	const FBlackboardEntry* EntryInfo = BlackboardAsset ? BlackboardAsset->GetKey(BlackboardKey.GetSelectedKeyID()) : NULL;
+
+	FString BlackboardDesc = "invalid";
+	if (EntryInfo)
+	{
+		// safety feature to not crash when changing couple of properties on a bunch
+		// while "post edit property" triggers for every each of them
+		if (EntryInfo->KeyType->GetClass() == BlackboardKey.SelectedKeyType)
+		{
+			const FString KeyName = EntryInfo->EntryName.ToString();
+			CacheOperationEnums();		
+
+			const EBlackboardKeyOperation::Type Op = EntryInfo->KeyType->GetTestOperation();
+			switch (Op)
+			{
+			case EBlackboardKeyOperation::Basic:
+				BlackboardDesc = FString::Printf(TEXT("%s is %s"), *KeyName, *BasicOpEnum->GetEnumName(OperationType));
+				break;
+
+			case EBlackboardKeyOperation::Arithmetic:
+				BlackboardDesc = FString::Printf(TEXT("%s %s %s"), *KeyName, *ArithmeticOpEnum->GetDisplayNameText(OperationType).ToString(),
+					*EntryInfo->KeyType->DescribeArithmeticParam(IntValue, FloatValue));
+				break;
+
+			case EBlackboardKeyOperation::Text:
+				BlackboardDesc = FString::Printf(TEXT("%s %s [%s]"), *KeyName, *TextOpEnum->GetEnumName(OperationType), *StringValue);
+				break;
+
+			default: break;
+			}
+		}
+	}
+
+	CachedDescription = BlackboardDesc;
+}
+
+void UBTDecorator_Blackboard::PostEditChangeProperty(FPropertyChangedEvent& PropertyChangedEvent)
+{
+	Super::PostEditChangeProperty(PropertyChangedEvent);
+	if (PropertyChangedEvent.Property == NULL)
+	{
+		return;
+	}
+
+	const FName ChangedPropName = PropertyChangedEvent.Property->GetFName();
+	if (ChangedPropName == GET_MEMBER_NAME_CHECKED(UBTDecorator_Blackboard, BlackboardKey.SelectedKeyName))
+	{
+		if (BlackboardKey.SelectedKeyType == UBlackboardKeyType_Enum::StaticClass() ||
+			BlackboardKey.SelectedKeyType == UBlackboardKeyType_NativeEnum::StaticClass())
+		{
+			IntValue = 0;
+		}
+	}
+
+#if WITH_EDITORONLY_DATA
+
+	UBlackboardKeyType* KeyCDO = BlackboardKey.SelectedKeyType ? BlackboardKey.SelectedKeyType->GetDefaultObject<UBlackboardKeyType>() : NULL;
+	if (ChangedPropName == GET_MEMBER_NAME_CHECKED(UBTDecorator_Blackboard, BasicOperation))
+	{
+		if (KeyCDO && KeyCDO->GetTestOperation() == EBlackboardKeyOperation::Basic)
+		{
+			OperationType = BasicOperation;
+		}
+	}
+	else if (ChangedPropName == GET_MEMBER_NAME_CHECKED(UBTDecorator_Blackboard, ArithmeticOperation))
+	{
+		if (KeyCDO && KeyCDO->GetTestOperation() == EBlackboardKeyOperation::Arithmetic)
+		{
+			OperationType = ArithmeticOperation;
+		}
+	}
+	else if (ChangedPropName == GET_MEMBER_NAME_CHECKED(UBTDecorator_Blackboard, TextOperation))
+	{
+		if (KeyCDO && KeyCDO->GetTestOperation() == EBlackboardKeyOperation::Text)
+		{
+			OperationType = TextOperation;
+		}
+	}
+
+#endif // WITH_EDITORONLY_DATA
+
+	BuildDescription();
+}
+
+void UBTDecorator_Blackboard::InitializeFromAsset(class UBehaviorTree* Asset)
+{
+	Super::InitializeFromAsset(Asset);
+	BuildDescription();
+}
+
+#endif	// WITH_EDITOR
diff --git a/Engine/Source/Runtime/AIModule/Private/BehaviorTree/Decorators/BTDecorator_BlackboardBase.cpp b/Engine/Source/Runtime/AIModule/Private/BehaviorTree/Decorators/BTDecorator_BlackboardBase.cpp
new file mode 100644
--- /dev/null
+++ b/Engine/Source/Runtime/AIModule/Private/BehaviorTree/Decorators/BTDecorator_BlackboardBase.cpp
@@ -0,0 +1,54 @@
+// Copyright 1998-2014 Epic Games, Inc. All Rights Reserved.
+
+#include "AIModulePrivate.h"
+#include "BehaviorTree/BlackboardComponent.h"
+#include "BehaviorTree/Decorators/BTDecorator_BlackboardBase.h"
+
+UBTDecorator_BlackboardBase::UBTDecorator_BlackboardBase(const class FPostConstructInitializeProperties& PCIP) : Super(PCIP)
+{
+	NodeName = "BlackboardBase";
+
+	bNotifyBecomeRelevant = true;
+	bNotifyCeaseRelevant = true;
+
+	// empty KeySelector = allow everything
+}
+
+void UBTDecorator_BlackboardBase::PostInitProperties()
+{
+	Super::PostInitProperties();
+	BBKeyObserver = FOnBlackboardChange::CreateUObject(this, &UBTDecorator_BlackboardBase::OnBlackboardChange);
+}
+
+void UBTDecorator_BlackboardBase::InitializeFromAsset(class UBehaviorTree* Asset)
+{
+	Super::InitializeFromAsset(Asset);
+	BlackboardKey.CacheSelectedKey(GetBlackboardAsset());
+}
+
+void UBTDecorator_BlackboardBase::OnBecomeRelevant(UBehaviorTreeComponent* OwnerComp, uint8* NodeMemory)
+{
+	UBlackboardComponent* BlackboardComp = OwnerComp->GetBlackboardComponent();
+	if (BlackboardComp)
+	{
+		BlackboardComp->RegisterObserver(BlackboardKey.GetSelectedKeyID(), BBKeyObserver);
+	}
+}
+
+void UBTDecorator_BlackboardBase::OnCeaseRelevant(UBehaviorTreeComponent* OwnerComp, uint8* NodeMemory)
+{
+	UBlackboardComponent* BlackboardComp = OwnerComp->GetBlackboardComponent();
+	if (BlackboardComp)
+	{
+		BlackboardComp->UnregisterObserver(BlackboardKey.GetSelectedKeyID(), BBKeyObserver);
+	}
+}
+
+void UBTDecorator_BlackboardBase::OnBlackboardChange(const UBlackboardComponent* Blackboard, uint8 ChangedKeyID)
+{
+	UBehaviorTreeComponent* BehaviorComp = Blackboard ? (UBehaviorTreeComponent*)Blackboard->GetBrainComponent() : NULL;
+	if (BlackboardKey.GetSelectedKeyID() == ChangedKeyID && BehaviorComp)
+	{
+		BehaviorComp->RequestExecution(this);		
+	}
+}
diff --git a/Engine/Source/Runtime/AIModule/Private/BehaviorTree/Decorators/BTDecorator_BlueprintBase.cpp b/Engine/Source/Runtime/AIModule/Private/BehaviorTree/Decorators/BTDecorator_BlueprintBase.cpp
new file mode 100644
--- /dev/null
+++ b/Engine/Source/Runtime/AIModule/Private/BehaviorTree/Decorators/BTDecorator_BlueprintBase.cpp
@@ -0,0 +1,195 @@
+// Copyright 1998-2014 Epic Games, Inc. All Rights Reserved.
+
+#include "AIModulePrivate.h"
+#include "BlueprintNodeHelpers.h"
+#include "BehaviorTree/Decorators/BTDecorator_BlueprintBase.h"
+#include "BehaviorTree/BlackboardComponent.h"
+#include "BehaviorTree/BlackboardData.h"
+
+UBTDecorator_BlueprintBase::UBTDecorator_BlueprintBase(const class FPostConstructInitializeProperties& PCIP) : Super(PCIP)
+{
+	UClass* StopAtClass = UBTDecorator_BlueprintBase::StaticClass();
+	bImplementsReceiveTick = BlueprintNodeHelpers::HasBlueprintFunction(TEXT("ReceiveTick"), this, StopAtClass);
+	bImplementsReceiveExecutionStart = BlueprintNodeHelpers::HasBlueprintFunction(TEXT("ReceiveExecutionStart"), this, StopAtClass);
+	bImplementsReceiveExecutionFinish = BlueprintNodeHelpers::HasBlueprintFunction(TEXT("ReceiveExecutionFinish"), this, StopAtClass);
+	bImplementsReceiveObserverActivated = BlueprintNodeHelpers::HasBlueprintFunction(TEXT("ReceiveObserverActivated"), this, StopAtClass);
+	bImplementsReceiveObserverDeactivated = BlueprintNodeHelpers::HasBlueprintFunction(TEXT("ReceiveObserverDeactivated"), this, StopAtClass);
+	bImplementsReceiveConditionCheck = BlueprintNodeHelpers::HasBlueprintFunction(TEXT("ReceiveConditionCheck"), this, StopAtClass);
+
+	bNotifyBecomeRelevant = bImplementsReceiveObserverActivated;
+	bNotifyCeaseRelevant = bImplementsReceiveObserverDeactivated;
+	bNotifyTick = bImplementsReceiveTick;
+	bNotifyActivation = bImplementsReceiveExecutionStart;
+	bNotifyDeactivation = bImplementsReceiveExecutionFinish;
+	bShowPropertyDetails = true;
+
+	// all blueprint based nodes must create instances
+	bCreateNodeInstance = true;
+
+	InitializeProperties();
+}
+
+void UBTDecorator_BlueprintBase::InitializeProperties()
+{
+	if (HasAnyFlags(RF_ClassDefaultObject))
+	{
+		UClass* StopAtClass = UBTDecorator_BlueprintBase::StaticClass();
+		BlueprintNodeHelpers::CollectPropertyData(this, StopAtClass, PropertyData);
+
+		if (GetFlowAbortMode() != EBTFlowAbortMode::None)
+		{
+			// find all blackboard key selectors and store their names
+			BlueprintNodeHelpers::CollectBlackboardSelectors(this, StopAtClass, ObservedKeyNames);
+		}
+	}
+	else
+	{
+		UBTDecorator_BlueprintBase* CDO = (UBTDecorator_BlueprintBase*)(GetClass()->GetDefaultObject());
+		if (CDO && CDO->ObservedKeyNames.Num())
+		{
+			bNotifyBecomeRelevant = true;
+			bNotifyCeaseRelevant = true;
+		}
+	}
+}
+
+void UBTDecorator_BlueprintBase::PostInitProperties()
+{
+	Super::PostInitProperties();
+	NodeName = BlueprintNodeHelpers::GetNodeName(this);
+	BBKeyObserver = FOnBlackboardChange::CreateUObject(this, &UBTDecorator_BlueprintBase::OnBlackboardChange);
+}
+
+void UBTDecorator_BlueprintBase::OnBecomeRelevant(UBehaviorTreeComponent* OwnerComp, uint8* NodeMemory)
+{
+	if (bImplementsReceiveObserverActivated)
+	{
+		ReceiveObserverActivated(OwnerComp->GetOwner());
+	}
+
+	UBlackboardComponent* BlackboardComp = OwnerComp->GetBlackboardComponent();
+	UBTDecorator_BlueprintBase* CDO = (UBTDecorator_BlueprintBase*)(GetClass()->GetDefaultObject());
+	if (BlackboardComp && CDO)
+	{
+		for (int32 i = 0; i < CDO->ObservedKeyNames.Num(); i++)
+		{
+			const uint8 KeyID = BlackboardComp->GetKeyID(CDO->ObservedKeyNames[i]);
+			if (KeyID != UBlackboardData::InvalidKeyID)
+			{
+				BlackboardComp->RegisterObserver(KeyID, BBKeyObserver);
+			}
+		}
+	}
+}
+
+void UBTDecorator_BlueprintBase::OnCeaseRelevant(UBehaviorTreeComponent* OwnerComp, uint8* NodeMemory)
+{
+	UBlackboardComponent* BlackboardComp = OwnerComp->GetBlackboardComponent();
+	UBTDecorator_BlueprintBase* CDO = (UBTDecorator_BlueprintBase*)(GetClass()->GetDefaultObject());
+	if (BlackboardComp && CDO)
+	{
+		for (int32 i = 0; i < CDO->ObservedKeyNames.Num(); i++)
+		{
+			const uint8 KeyID = BlackboardComp->GetKeyID(CDO->ObservedKeyNames[i]);
+			if (KeyID != UBlackboardData::InvalidKeyID)
+			{
+				BlackboardComp->UnregisterObserver(KeyID, BBKeyObserver);
+			}
+		}
+	}
+		
+	if (bImplementsReceiveObserverDeactivated)
+	{
+		ReceiveObserverDeactivated(OwnerComp->GetOwner());
+	}
+}
+
+void UBTDecorator_BlueprintBase::OnNodeActivation(struct FBehaviorTreeSearchData& SearchData)
+{
+	// skip flag, will be handled by bNotifyActivation
+
+	ReceiveExecutionStart(SearchData.OwnerComp->GetOwner());
+}
+
+void UBTDecorator_BlueprintBase::OnNodeDeactivation(struct FBehaviorTreeSearchData& SearchData, EBTNodeResult::Type NodeResult)
+{
+	// skip flag, will be handled by bNotifyDeactivation
+
+	ReceiveExecutionFinish(SearchData.OwnerComp->GetOwner(), NodeResult);
+}
+
+void UBTDecorator_BlueprintBase::TickNode(UBehaviorTreeComponent* OwnerComp, uint8* NodeMemory, float DeltaSeconds)
+{
+	// skip flag, will be handled by bNotifyTick
+
+	ReceiveTick(OwnerComp->GetOwner(), DeltaSeconds);
+}
+
+bool UBTDecorator_BlueprintBase::CalculateRawConditionValue(class UBehaviorTreeComponent* OwnerComp, uint8* NodeMemory) const
+{
+	CurrentCallResult = false;
+	if (bImplementsReceiveConditionCheck)
+	{
+		// can't use const functions with blueprints
+		UBTDecorator_BlueprintBase* MyNode = (UBTDecorator_BlueprintBase*)this;
+		MyNode->ReceiveConditionCheck(OwnerComp->GetOwner());
+	}
+
+	return CurrentCallResult;
+}
+
+void UBTDecorator_BlueprintBase::FinishConditionCheck(bool bAllowExecution)
+{
+	CurrentCallResult = bAllowExecution;
+}
+
+bool UBTDecorator_BlueprintBase::IsDecoratorExecutionActive() const
+{
+	UBehaviorTreeComponent* OwnerComp = Cast<UBehaviorTreeComponent>(GetOuter());
+	const bool bIsActive = OwnerComp->IsExecutingBranch(GetMyNode(), GetChildIndex());
+	return bIsActive;
+}
+
+bool UBTDecorator_BlueprintBase::IsDecoratorObserverActive() const
+{
+	UBehaviorTreeComponent* OwnerComp = Cast<UBehaviorTreeComponent>(GetOuter());
+	const bool bIsActive = OwnerComp->IsAuxNodeActive(this);
+	return bIsActive;
+}
+
+FString UBTDecorator_BlueprintBase::GetStaticDescription() const
+{
+	FString ReturnDesc = Super::GetStaticDescription();
+
+	UBTDecorator_BlueprintBase* CDO = (UBTDecorator_BlueprintBase*)(GetClass()->GetDefaultObject());
+	if (bShowPropertyDetails && CDO)
+	{
+		UClass* StopAtClass = UBTDecorator_BlueprintBase::StaticClass();
+		FString PropertyDesc = BlueprintNodeHelpers::CollectPropertyDescription(this, StopAtClass, CDO->PropertyData);
+		if (PropertyDesc.Len())
+		{
+			ReturnDesc += TEXT(":\n\n");
+			ReturnDesc += PropertyDesc;
+		}
+	}
+
+	return ReturnDesc;
+}
+
+void UBTDecorator_BlueprintBase::DescribeRuntimeValues(const class UBehaviorTreeComponent* OwnerComp, uint8* NodeMemory, EBTDescriptionVerbosity::Type Verbosity, TArray<FString>& Values) const
+{
+	UBTDecorator_BlueprintBase* CDO = (UBTDecorator_BlueprintBase*)(GetClass()->GetDefaultObject());
+	if (CDO && CDO->PropertyData.Num())
+	{
+		BlueprintNodeHelpers::DescribeRuntimeValues(this, CDO->PropertyData, Values);
+	}
+}
+
+void UBTDecorator_BlueprintBase::OnBlackboardChange(const UBlackboardComponent* Blackboard, uint8 ChangedKeyID)
+{
+	UBehaviorTreeComponent* BehaviorComp = Blackboard ? (UBehaviorTreeComponent*)Blackboard->GetBrainComponent() : NULL;
+	if (BehaviorComp)
+	{
+		BehaviorComp->RequestExecution(this);		
+	}
+}
diff --git a/Engine/Source/Runtime/AIModule/Private/BehaviorTree/Decorators/BTDecorator_CompareBBEntries.cpp b/Engine/Source/Runtime/AIModule/Private/BehaviorTree/Decorators/BTDecorator_CompareBBEntries.cpp
new file mode 100644
--- /dev/null
+++ b/Engine/Source/Runtime/AIModule/Private/BehaviorTree/Decorators/BTDecorator_CompareBBEntries.cpp
@@ -0,0 +1,96 @@
+// Copyright 1998-2014 Epic Games, Inc. All Rights Reserved.
+
+#include "AIModulePrivate.h"
+#include "BehaviorTree/BlackboardComponent.h"
+#include "BehaviorTree/Decorators/BTDecorator_CompareBBEntries.h"
+
+UBTDecorator_CompareBBEntries::UBTDecorator_CompareBBEntries(const class FPostConstructInitializeProperties& PCIP) 
+	: Super(PCIP)
+	, Operator(EArithmeticKeyOperation::LessOrEqual)
+{
+	NodeName = "Compare Blackboard entries";
+
+	bNotifyBecomeRelevant = true;
+	bNotifyCeaseRelevant = true;
+}
+
+void UBTDecorator_CompareBBEntries::PostInitProperties()
+{
+	Super::PostInitProperties();
+	BBKeyObserver = FOnBlackboardChange::CreateUObject(this, &UBTDecorator_CompareBBEntries::OnBlackboardChange);
+}
+
+void UBTDecorator_CompareBBEntries::InitializeFromAsset(class UBehaviorTree* Asset)
+{
+	Super::InitializeFromAsset(Asset);
+
+	UBlackboardData* BBAsset = GetBlackboardAsset();
+	BlackboardKeyA.CacheSelectedKey(BBAsset);
+	BlackboardKeyB.CacheSelectedKey(BBAsset);
+}
+
+// @note I know it's ugly to have "return" statements in many places inside a function, but the way 
+// around was very awkward here 
+bool UBTDecorator_CompareBBEntries::CalculateRawConditionValue(class UBehaviorTreeComponent* OwnerComp, uint8* NodeMemory) const 
+{
+	// first of all require same type
+	// @todo this could be checked statically (i.e. in editor, asset creation time)!
+	if (BlackboardKeyA.SelectedKeyType != BlackboardKeyB.SelectedKeyType)
+	{
+		return false;
+	}
+	
+	const UBlackboardComponent* BlackboardComp = OwnerComp->GetBlackboardComponent();
+	if (BlackboardComp)
+	{
+		//BlackboardComp->GetKeyType
+		const EBlackboardCompare::Type Result = BlackboardComp->CompareKeyValues(BlackboardKeyA.SelectedKeyType, 
+			BlackboardKeyA.GetSelectedKeyID(), BlackboardKeyB.GetSelectedKeyID());
+
+		checkAtCompileTime(int32(EBlackboardCompare::Equal) == int32(EBlackBoardEntryComparison::Equal)
+			&& int32(EBlackboardCompare::NotEqual) == int32(EBlackBoardEntryComparison::NotEqual)
+			, "These values need to be equal");
+
+		return int32(Operator) == int32(Result);
+	}
+
+	return false;
+}
+
+FString UBTDecorator_CompareBBEntries::GetStaticDescription() const 
+{
+	return FString::Printf(TEXT("%s:\n%s and %s\ncontain %s values")
+		, *Super::GetStaticDescription()
+		, *BlackboardKeyA.SelectedKeyName.ToString()
+		, *BlackboardKeyB.SelectedKeyName.ToString()
+		, Operator == EBlackBoardEntryComparison::Equal ? TEXT("EQUAL") : TEXT("NOT EQUAL"));
+}
+
+void UBTDecorator_CompareBBEntries::OnBecomeRelevant(UBehaviorTreeComponent* OwnerComp, uint8* NodeMemory)
+{
+	UBlackboardComponent* BlackboardComp = OwnerComp->GetBlackboardComponent();
+	if (BlackboardComp)
+	{
+		BlackboardComp->RegisterObserver(BlackboardKeyA.GetSelectedKeyID(), BBKeyObserver);
+		BlackboardComp->RegisterObserver(BlackboardKeyB.GetSelectedKeyID(), BBKeyObserver);
+	}
+}
+
+void UBTDecorator_CompareBBEntries::OnCeaseRelevant(UBehaviorTreeComponent* OwnerComp, uint8* NodeMemory)
+{
+	UBlackboardComponent* BlackboardComp = OwnerComp->GetBlackboardComponent();
+	if (BlackboardComp)
+	{
+		BlackboardComp->UnregisterObserver(BlackboardKeyA.GetSelectedKeyID(), BBKeyObserver);
+		BlackboardComp->UnregisterObserver(BlackboardKeyB.GetSelectedKeyID(), BBKeyObserver);
+	}
+}
+
+void UBTDecorator_CompareBBEntries::OnBlackboardChange(const UBlackboardComponent* Blackboard, uint8 ChangedKeyID)
+{
+	UBehaviorTreeComponent* BehaviorComp = Blackboard ? (UBehaviorTreeComponent*)Blackboard->GetBrainComponent() : NULL;
+	if (BehaviorComp && (BlackboardKeyA.GetSelectedKeyID() == ChangedKeyID || BlackboardKeyB.GetSelectedKeyID() == ChangedKeyID))
+	{
+		BehaviorComp->RequestExecution(this);		
+	}
+}
diff --git a/Engine/Source/Runtime/AIModule/Private/BehaviorTree/Decorators/BTDecorator_ConeCheck.cpp b/Engine/Source/Runtime/AIModule/Private/BehaviorTree/Decorators/BTDecorator_ConeCheck.cpp
new file mode 100644
--- /dev/null
+++ b/Engine/Source/Runtime/AIModule/Private/BehaviorTree/Decorators/BTDecorator_ConeCheck.cpp
@@ -0,0 +1,136 @@
+// Copyright 1998-2014 Epic Games, Inc. All Rights Reserved.
+
+#include "AIModulePrivate.h"
+#include "BehaviorTree/BlackboardComponent.h"
+#include "BehaviorTree/Decorators/BTDecorator_ConeCheck.h"
+
+UBTDecorator_ConeCheck::UBTDecorator_ConeCheck(const class FPostConstructInitializeProperties& PCIP) : Super(PCIP)
+{
+	NodeName = "Keep in Cone";
+
+	// accept only actors and vectors
+	ConeOrigin.AddObjectFilter(this, AActor::StaticClass());
+	ConeOrigin.AddVectorFilter(this);
+	ConeDirection.AddObjectFilter(this, AActor::StaticClass());
+	ConeDirection.AddVectorFilter(this);
+	ConeDirection.AllowNoneAsValue(true);
+	Observed.AddObjectFilter(this, AActor::StaticClass());
+	Observed.AddVectorFilter(this);
+
+	bNotifyBecomeRelevant = true;
+	bNotifyTick = true;
+
+	// KeepInCone always abort current branch
+	FlowAbortMode = EBTFlowAbortMode::None;
+	
+	ConeHalfAngle = 45.0f;
+}
+
+void UBTDecorator_ConeCheck::InitializeFromAsset(class UBehaviorTree* Asset)
+{
+	Super::InitializeFromAsset(Asset);
+
+	ConeHalfAngleDot = FMath::Cos(FMath::DegreesToRadians(ConeHalfAngle));
+
+	UBlackboardData* BBAsset = GetBlackboardAsset();
+	ConeOrigin.CacheSelectedKey(BBAsset);
+	ConeDirection.CacheSelectedKey(BBAsset);
+	Observed.CacheSelectedKey(BBAsset);
+}
+
+bool UBTDecorator_ConeCheck::CalculateDirection(const UBlackboardComponent* BlackboardComp, const FBlackboardKeySelector& Origin, const FBlackboardKeySelector& End, FVector& Direction) const
+{
+	FVector PointA = FVector::ZeroVector;
+	FVector PointB = FVector::ZeroVector;
+	FRotator Rotation = FRotator::ZeroRotator;
+
+	if (BlackboardComp)
+	{
+		if (End.IsNone())
+		{
+			if (BlackboardComp->GetRotationFromEntry(Origin.GetSelectedKeyID(), Rotation))
+			{
+				Direction = Rotation.Vector();
+				return true;
+			}
+		}
+		else if (BlackboardComp->GetLocationFromEntry(Origin.GetSelectedKeyID(), PointA) && BlackboardComp->GetLocationFromEntry(End.GetSelectedKeyID(), PointB))
+		{
+			Direction = (PointB - PointA).SafeNormal();
+			return true;
+		}
+	}
+
+	return false;
+}
+
+FORCEINLINE bool UBTDecorator_ConeCheck::CalcConditionImpl(UBehaviorTreeComponent* OwnerComp, uint8* NodeMemory) const
+{
+	const UBlackboardComponent* BBComponent = OwnerComp->GetBlackboardComponent();
+
+	FVector ConeDir;
+	FVector DirectionToObserve;
+
+	return CalculateDirection(BBComponent, ConeOrigin, Observed, DirectionToObserve)
+		&& CalculateDirection(BBComponent, ConeOrigin, ConeDirection, ConeDir)
+		&& ConeDir.CosineAngle2D(DirectionToObserve) > ConeHalfAngleDot;
+}
+
+bool UBTDecorator_ConeCheck::CalculateRawConditionValue(class UBehaviorTreeComponent* OwnerComp, uint8* NodeMemory) const
+{
+	return CalcConditionImpl(OwnerComp, NodeMemory);
+}
+
+void UBTDecorator_ConeCheck::OnBecomeRelevant(UBehaviorTreeComponent* OwnerComp, uint8* NodeMemory)
+{
+	TNodeInstanceMemory* DecoratorMemory = (TNodeInstanceMemory*)NodeMemory;
+	DecoratorMemory->bLastRawResult = CalcConditionImpl(OwnerComp, NodeMemory);
+}
+
+void UBTDecorator_ConeCheck::OnBlackboardChange(const class UBlackboardComponent* Blackboard, uint8 ChangedKeyID)
+{
+	check(false);
+}
+
+void UBTDecorator_ConeCheck::TickNode(UBehaviorTreeComponent* OwnerComp, uint8* NodeMemory, float DeltaSeconds)
+{
+	const TNodeInstanceMemory* DecoratorMemory = (TNodeInstanceMemory*)NodeMemory;
+
+	if (CalcConditionImpl(OwnerComp, NodeMemory) != DecoratorMemory->bLastRawResult)
+	{
+		OwnerComp->RequestExecution(this);
+	}
+}
+
+FString UBTDecorator_ConeCheck::GetStaticDescription() const
+{
+	return FString::Printf(TEXT("%s: is %s in %.2f degree %s-%s cone")
+		, *Super::GetStaticDescription(), *Observed.SelectedKeyName.ToString()
+		, ConeHalfAngle * 2, *ConeOrigin.SelectedKeyName.ToString(), *ConeDirection.SelectedKeyName.ToString());
+}
+
+void UBTDecorator_ConeCheck::DescribeRuntimeValues(const class UBehaviorTreeComponent* OwnerComp, uint8* NodeMemory, EBTDescriptionVerbosity::Type Verbosity, TArray<FString>& Values) const
+{
+	const UBlackboardComponent* BBComponent = OwnerComp->GetBlackboardComponent();
+			
+	FVector ConeDir;
+	FVector DirectionToObserved;
+
+	if (CalculateDirection(BBComponent, ConeOrigin, Observed, DirectionToObserved)
+		&& CalculateDirection(BBComponent, ConeOrigin, ConeDirection, ConeDir))
+	{
+		const float CurrentAngleDot = ConeDir.CosineAngle2D(DirectionToObserved);
+		const float CurrentAngleRad = FMath::Acos(CurrentAngleDot);
+
+		Values.Add(FString::Printf(TEXT("Angle: %.0f (%s cone)"),
+			FMath::RadiansToDegrees(CurrentAngleRad),
+			CurrentAngleDot < ConeHalfAngleDot ? TEXT("outside") : TEXT("inside")
+			));
+
+	}
+}
+
+uint16 UBTDecorator_ConeCheck::GetInstanceMemorySize() const
+{
+	return sizeof(TNodeInstanceMemory);
+}
diff --git a/Engine/Source/Runtime/AIModule/Private/BehaviorTree/Decorators/BTDecorator_Cooldown.cpp b/Engine/Source/Runtime/AIModule/Private/BehaviorTree/Decorators/BTDecorator_Cooldown.cpp
new file mode 100644
--- /dev/null
+++ b/Engine/Source/Runtime/AIModule/Private/BehaviorTree/Decorators/BTDecorator_Cooldown.cpp
@@ -0,0 +1,85 @@
+// Copyright 1998-2014 Epic Games, Inc. All Rights Reserved.
+
+#include "AIModulePrivate.h"
+#include "BehaviorTree/Decorators/BTDecorator_Cooldown.h"
+
+UBTDecorator_Cooldown::UBTDecorator_Cooldown(const class FPostConstructInitializeProperties& PCIP) : Super(PCIP)
+{
+	NodeName = "Cooldown";
+	CoolDownTime = 5.0f;
+	
+	// aborting child nodes doesn't makes sense, cooldown starts after leaving this branch
+	bAllowAbortNone = false;
+	bAllowAbortLowerPri = false;
+	bAllowAbortChildNodes = false;
+
+	bNotifyTick = false;
+	bNotifyDeactivation = true;
+}
+
+void UBTDecorator_Cooldown::PostLoad()
+{
+	Super::PostLoad();
+	bNotifyTick = (FlowAbortMode != EBTFlowAbortMode::None);
+}
+
+bool UBTDecorator_Cooldown::CalculateRawConditionValue(class UBehaviorTreeComponent* OwnerComp, uint8* NodeMemory) const 
+{
+	FBTCooldownDecoratorMemory* DecoratorMemory = (FBTCooldownDecoratorMemory*)NodeMemory;
+	const float TimePassed = (OwnerComp->GetWorld()->GetTimeSeconds() - DecoratorMemory->LastUseTimestamp);
+	return TimePassed >= CoolDownTime;
+}
+
+void UBTDecorator_Cooldown::InitializeMemory(class UBehaviorTreeComponent* OwnerComp, uint8* NodeMemory) const
+{
+	FBTCooldownDecoratorMemory* DecoratorMemory = (FBTCooldownDecoratorMemory*)NodeMemory;
+	DecoratorMemory->LastUseTimestamp = -FLT_MAX;
+}
+
+void UBTDecorator_Cooldown::OnNodeDeactivation(struct FBehaviorTreeSearchData& SearchData, EBTNodeResult::Type NodeResult)
+{
+	FBTCooldownDecoratorMemory* DecoratorMemory = GetNodeMemory<FBTCooldownDecoratorMemory>(SearchData);
+	DecoratorMemory->LastUseTimestamp = SearchData.OwnerComp->GetWorld()->GetTimeSeconds();
+	DecoratorMemory->RequestedRestart = false;
+}
+
+void UBTDecorator_Cooldown::TickNode(class UBehaviorTreeComponent* OwnerComp, uint8* NodeMemory, float DeltaSeconds)
+{
+	FBTCooldownDecoratorMemory* DecoratorMemory = (FBTCooldownDecoratorMemory*)NodeMemory;
+	if (!DecoratorMemory->RequestedRestart)
+	{
+		const float TimePassed = (OwnerComp->GetWorld()->GetTimeSeconds() - DecoratorMemory->LastUseTimestamp);
+		if (TimePassed >= CoolDownTime)
+		{
+			DecoratorMemory->RequestedRestart = true;
+			OwnerComp->RequestExecution(this);
+		}
+	}
+}
+
+FString UBTDecorator_Cooldown::GetStaticDescription() const
+{
+	// basic info: result after time
+	return FString::Printf(TEXT("%s: lock for %.1fs after execution and return %s"), *Super::GetStaticDescription(),
+		CoolDownTime, *UBehaviorTreeTypes::DescribeNodeResult(EBTNodeResult::Failed));
+}
+
+void UBTDecorator_Cooldown::DescribeRuntimeValues(const class UBehaviorTreeComponent* OwnerComp, uint8* NodeMemory, EBTDescriptionVerbosity::Type Verbosity, TArray<FString>& Values) const
+{
+	Super::DescribeRuntimeValues(OwnerComp, NodeMemory, Verbosity, Values);
+
+	FBTCooldownDecoratorMemory* DecoratorMemory = (FBTCooldownDecoratorMemory*)NodeMemory;
+	const float TimePassed = OwnerComp->GetWorld()->GetTimeSeconds() - DecoratorMemory->LastUseTimestamp;
+	
+	if (TimePassed < CoolDownTime)
+	{
+		Values.Add(FString::Printf(TEXT("%s in %ss"),
+			(FlowAbortMode == EBTFlowAbortMode::None) ? TEXT("unlock") : TEXT("restart"),
+			*FString::SanitizeFloat(CoolDownTime - TimePassed)));
+	}
+}
+
+uint16 UBTDecorator_Cooldown::GetInstanceMemorySize() const
+{
+	return sizeof(FBTCooldownDecoratorMemory);
+}
diff --git a/Engine/Source/Runtime/AIModule/Private/BehaviorTree/Decorators/BTDecorator_DoesPathExist.cpp b/Engine/Source/Runtime/AIModule/Private/BehaviorTree/Decorators/BTDecorator_DoesPathExist.cpp
new file mode 100644
--- /dev/null
+++ b/Engine/Source/Runtime/AIModule/Private/BehaviorTree/Decorators/BTDecorator_DoesPathExist.cpp
@@ -0,0 +1,91 @@
+// Copyright 1998-2014 Epic Games, Inc. All Rights Reserved.
+
+#include "AIModulePrivate.h"
+#include "Navigation/NavigationComponent.h"
+#include "BehaviorTree/BlackboardComponent.h"
+#include "BehaviorTree/Decorators/BTDecorator_DoesPathExist.h"
+
+UBTDecorator_DoesPathExist::UBTDecorator_DoesPathExist(const class FPostConstructInitializeProperties& PCIP) : Super(PCIP)
+{
+	NodeName = "Does path exist";
+
+	// accept only actors and vectors
+	BlackboardKeyA.AddObjectFilter(this, AActor::StaticClass());
+	BlackboardKeyA.AddVectorFilter(this);
+	BlackboardKeyB.AddObjectFilter(this, AActor::StaticClass());
+	BlackboardKeyB.AddVectorFilter(this);
+
+	bAllowAbortLowerPri = false;
+	bAllowAbortNone = true;
+	bAllowAbortChildNodes = false;
+	FlowAbortMode = EBTFlowAbortMode::None;
+
+	BlackboardKeyA.SelectedKeyName = FBlackboard::KeySelf;
+	PathQueryType = EPathExistanceQueryType::HierarchicalQuery;
+}
+
+void UBTDecorator_DoesPathExist::InitializeFromAsset(class UBehaviorTree* Asset)
+{
+	Super::InitializeFromAsset(Asset);
+
+	if (bUseSelf)
+	{
+		BlackboardKeyA.SelectedKeyName = FBlackboard::KeySelf;
+		bUseSelf = false;
+	}
+
+	UBlackboardData* BBAsset = GetBlackboardAsset();
+	BlackboardKeyA.CacheSelectedKey(BBAsset);
+	BlackboardKeyB.CacheSelectedKey(BBAsset);
+}
+
+bool UBTDecorator_DoesPathExist::CalculateRawConditionValue(class UBehaviorTreeComponent* OwnerComp, uint8* NodeMemory) const 
+{
+	const UBlackboardComponent* BlackboardComp = OwnerComp->GetBlackboardComponent();
+	if (BlackboardComp == NULL)
+	{
+		return false;
+	}
+
+	FVector PointA = FVector::ZeroVector;
+	FVector PointB = FVector::ZeroVector;	
+	const bool bHasPointA = BlackboardComp->GetLocationFromEntry(BlackboardKeyA.GetSelectedKeyID(), PointA);
+	const bool bHasPointB = BlackboardComp->GetLocationFromEntry(BlackboardKeyB.GetSelectedKeyID(), PointB);
+	
+	bool bHasPath = false;
+
+	const UNavigationSystem* NavSys = UNavigationSystem::GetCurrent(OwnerComp->GetWorld());
+	if (NavSys && bHasPointA && bHasPointB)
+	{
+		const AAIController* AIOwner = Cast<AAIController>(OwnerComp->GetOwner());
+		const ANavigationData* NavData = AIOwner && AIOwner->NavComponent ? AIOwner->NavComponent->GetNavData() : NULL;
+		TSharedPtr<const FNavigationQueryFilter> QueryFilter = UNavigationQueryFilter::GetQueryFilter(NavData, FilterClass);
+
+		if (PathQueryType == EPathExistanceQueryType::NavmeshRaycast2D)
+		{
+#if WITH_RECAST
+			const ARecastNavMesh* RecastNavMesh = Cast<const ARecastNavMesh>(NavData);
+			bHasPath = RecastNavMesh && RecastNavMesh->IsSegmentOnNavmesh(PointA, PointB, QueryFilter);
+#endif
+		}
+		else
+		{		
+			EPathFindingMode::Type TestMode = (PathQueryType == EPathExistanceQueryType::HierarchicalQuery) ? EPathFindingMode::Hierarchical : EPathFindingMode::Regular;
+			bHasPath = NavSys->TestPathSync(FPathFindingQuery(AIOwner, NavData, PointA, PointB, QueryFilter), TestMode);
+		}
+	}
+
+	return bHasPath;
+}
+
+FString UBTDecorator_DoesPathExist::GetStaticDescription() const 
+{
+	const UEnum* PathTypeEnum = FindObject<UEnum>(ANY_PACKAGE, TEXT("EPathExistanceQueryType"));
+	check(PathTypeEnum);
+
+	return FString::Printf(TEXT("%s: Find path from %s to %s (mode:%s)"),
+		*Super::GetStaticDescription(),
+		*BlackboardKeyA.SelectedKeyName.ToString(),
+		*BlackboardKeyB.SelectedKeyName.ToString(),
+		*PathTypeEnum->GetEnumName(PathQueryType));
+}
diff --git a/Engine/Source/Runtime/AIModule/Private/BehaviorTree/Decorators/BTDecorator_ForceSuccess.cpp b/Engine/Source/Runtime/AIModule/Private/BehaviorTree/Decorators/BTDecorator_ForceSuccess.cpp
new file mode 100644
--- /dev/null
+++ b/Engine/Source/Runtime/AIModule/Private/BehaviorTree/Decorators/BTDecorator_ForceSuccess.cpp
@@ -0,0 +1,19 @@
+// Copyright 1998-2013 Epic Games, Inc. All Rights Reserved.
+
+#include "AIModulePrivate.h"
+#include "BehaviorTree/Decorators/BTDecorator_ForceSuccess.h"
+
+UBTDecorator_ForceSuccess::UBTDecorator_ForceSuccess(const class FPostConstructInitializeProperties& PCIP) : Super(PCIP)
+{
+	NodeName = TEXT("Force Success");
+	bNotifyProcessed = true;
+
+	bAllowAbortNone = false;
+	bAllowAbortLowerPri = false;
+	bAllowAbortChildNodes = false;
+}
+
+void UBTDecorator_ForceSuccess::OnNodeProcessed(struct FBehaviorTreeSearchData& SearchData, EBTNodeResult::Type& NodeResult)
+{
+	NodeResult = EBTNodeResult::Succeeded;
+}
diff --git a/Engine/Source/Runtime/AIModule/Private/BehaviorTree/Decorators/BTDecorator_KeepInCone.cpp b/Engine/Source/Runtime/AIModule/Private/BehaviorTree/Decorators/BTDecorator_KeepInCone.cpp
new file mode 100644
--- /dev/null
+++ b/Engine/Source/Runtime/AIModule/Private/BehaviorTree/Decorators/BTDecorator_KeepInCone.cpp
@@ -0,0 +1,129 @@
+// Copyright 1998-2014 Epic Games, Inc. All Rights Reserved.
+
+#include "AIModulePrivate.h"
+#include "BehaviorTree/BlackboardComponent.h"
+#include "BehaviorTree/Decorators/BTDecorator_KeepInCone.h"
+
+UBTDecorator_KeepInCone::UBTDecorator_KeepInCone(const class FPostConstructInitializeProperties& PCIP) : Super(PCIP)
+{
+	NodeName = "Keep in Cone";
+
+	// accept only actors and vectors
+	ConeOrigin.AddObjectFilter(this, AActor::StaticClass());
+	ConeOrigin.AddVectorFilter(this);
+	Observed.AddObjectFilter(this, AActor::StaticClass());
+	Observed.AddVectorFilter(this);
+
+	bNotifyBecomeRelevant = true;
+	bNotifyTick = true;
+
+	// KeepInCone always abort current branch
+	bAllowAbortLowerPri = false;
+	bAllowAbortNone = false;
+	FlowAbortMode = EBTFlowAbortMode::Self;
+	
+	ConeOrigin.SelectedKeyName = FBlackboard::KeySelf;
+	ConeHalfAngle = 45.0f;
+}
+
+void UBTDecorator_KeepInCone::InitializeFromAsset(class UBehaviorTree* Asset)
+{
+	Super::InitializeFromAsset(Asset);
+
+	ConeHalfAngleDot = FMath::Cos(FMath::DegreesToRadians(ConeHalfAngle));
+	
+	if (bUseSelfAsOrigin)
+	{
+		ConeOrigin.SelectedKeyName = FBlackboard::KeySelf;
+		bUseSelfAsOrigin = false;
+	}
+
+	if (bUseSelfAsObserved)
+	{
+		Observed.SelectedKeyName = FBlackboard::KeySelf;
+		bUseSelfAsObserved = false;
+	}
+
+	UBlackboardData* BBAsset = GetBlackboardAsset();
+	ConeOrigin.CacheSelectedKey(BBAsset);
+	Observed.CacheSelectedKey(BBAsset);
+}
+
+bool UBTDecorator_KeepInCone::CalculateCurrentDirection(const UBehaviorTreeComponent* OwnerComp, FVector& Direction) const
+{
+	const UBlackboardComponent* BlackboardComp = OwnerComp->GetBlackboardComponent();
+	if (BlackboardComp == NULL)
+	{
+		return false;
+	}
+
+	FVector PointA = FVector::ZeroVector;
+	FVector PointB = FVector::ZeroVector;
+	const bool bHasPointA = BlackboardComp->GetLocationFromEntry(ConeOrigin.GetSelectedKeyID(), PointA);
+	const bool bHasPointB = BlackboardComp->GetLocationFromEntry(Observed.GetSelectedKeyID(), PointB);
+
+	if (bHasPointA && bHasPointB)
+	{
+		Direction = (PointB - PointA).SafeNormal();
+		return true;
+	}
+
+	return false;
+}
+
+void UBTDecorator_KeepInCone::OnBecomeRelevant(UBehaviorTreeComponent* OwnerComp, uint8* NodeMemory)
+{
+	TNodeInstanceMemory* DecoratorMemory = (TNodeInstanceMemory*)NodeMemory;
+	FVector InitialDir(1.0f, 0, 0);
+
+	CalculateCurrentDirection(OwnerComp, InitialDir);
+	DecoratorMemory->InitialDirection = InitialDir;
+}
+
+void UBTDecorator_KeepInCone::TickNode(UBehaviorTreeComponent* OwnerComp, uint8* NodeMemory, float DeltaSeconds)
+{
+	TNodeInstanceMemory* DecoratorMemory = (TNodeInstanceMemory*)NodeMemory;
+	FVector CurrentDir(1.0f, 0, 0);
+	
+	if (CalculateCurrentDirection(OwnerComp, CurrentDir))
+	{
+		const float Angle = DecoratorMemory->InitialDirection.CosineAngle2D(CurrentDir);
+		if (Angle < ConeHalfAngleDot || (IsInversed() && Angle > ConeHalfAngleDot))
+		{
+			OwnerComp->RequestExecution(this);
+		}
+	}
+}
+
+FString UBTDecorator_KeepInCone::GetStaticDescription() const
+{
+	return FString::Printf(TEXT("%s: %s in %.2f degree cone of initial direction [%s-%s]"),
+		*Super::GetStaticDescription(),
+		*Observed.SelectedKeyName.ToString(),
+		ConeHalfAngle * 2,
+		*ConeOrigin.SelectedKeyName.ToString(),
+		*Observed.SelectedKeyName.ToString());
+}
+
+void UBTDecorator_KeepInCone::DescribeRuntimeValues(const class UBehaviorTreeComponent* OwnerComp, uint8* NodeMemory, EBTDescriptionVerbosity::Type Verbosity, TArray<FString>& Values) const
+{
+	TNodeInstanceMemory* DecoratorMemory = (TNodeInstanceMemory*)NodeMemory;
+	FVector CurrentDir(1.0f, 0, 0);
+	
+	if (CalculateCurrentDirection(OwnerComp, CurrentDir))
+	{
+		const float CurrentAngleDot = DecoratorMemory->InitialDirection.CosineAngle2D(CurrentDir);
+		const float CurrentAngleRad = FMath::Acos(CurrentAngleDot);
+
+		Values.Add(FString::Printf(TEXT("Angle: %.0f (%s cone)"),
+			FMath::RadiansToDegrees(CurrentAngleRad),
+			CurrentAngleDot < ConeHalfAngleDot ? TEXT("outside") : TEXT("inside")
+			));
+
+	}
+}
+
+uint16 UBTDecorator_KeepInCone::GetInstanceMemorySize() const
+{
+	return sizeof(TNodeInstanceMemory);
+}
