Extracting AI code from the engine into a separate module, pass 1 #UE4

diff --git a/Engine/Source/Runtime/AIModule/Private/BehaviorTree/Decorators/BTDecorator_Loop.cpp b/Engine/Source/Runtime/AIModule/Private/BehaviorTree/Decorators/BTDecorator_Loop.cpp
new file mode 100644
--- /dev/null
+++ b/Engine/Source/Runtime/AIModule/Private/BehaviorTree/Decorators/BTDecorator_Loop.cpp
@@ -0,0 +1,63 @@
+// Copyright 1998-2014 Epic Games, Inc. All Rights Reserved.
+
+#include "AIModulePrivate.h"
+#include "BehaviorTree/BTCompositeNode.h"
+#include "BehaviorTree/Decorators/BTDecorator_Loop.h"
+
+UBTDecorator_Loop::UBTDecorator_Loop(const class FPostConstructInitializeProperties& PCIP) : Super(PCIP)
+{
+	NodeName = "Loop";
+	NumLoops = 3;
+	bNotifyActivation = true;
+	
+	bAllowAbortNone = false;
+	bAllowAbortLowerPri = false;
+	bAllowAbortChildNodes = false;
+}
+
+void UBTDecorator_Loop::OnNodeActivation(FBehaviorTreeSearchData& SearchData)
+{
+	FBTLoopDecoratorMemory* DecoratorMemory = GetNodeMemory<FBTLoopDecoratorMemory>(SearchData);
+	FBTCompositeMemory* ParentMemory = GetParentNode()->GetNodeMemory<FBTCompositeMemory>(SearchData);
+	if (ParentMemory->CurrentChild != ChildIndex)
+	{
+		// initialize counter if it's first activation
+		DecoratorMemory->RemainingExecutions = NumLoops;
+	}
+
+	if (!bInfiniteLoop)
+	{
+		DecoratorMemory->RemainingExecutions--;
+	}
+
+	// set child selection overrides
+	const bool bShouldLoop = bInfiniteLoop || (DecoratorMemory->RemainingExecutions > 0);
+	if (bShouldLoop)
+	{
+		GetParentNode()->SetChildOverride(SearchData, ChildIndex);
+	}
+}
+
+FString UBTDecorator_Loop::GetStaticDescription() const
+{
+	// basic info: infinite / num loops
+	return bInfiniteLoop ? 
+		FString::Printf(TEXT("%s: infinite"), *Super::GetStaticDescription()) :
+		FString::Printf(TEXT("%s: %d loops"), *Super::GetStaticDescription(), NumLoops);
+}
+
+void UBTDecorator_Loop::DescribeRuntimeValues(const class UBehaviorTreeComponent* OwnerComp, uint8* NodeMemory, EBTDescriptionVerbosity::Type Verbosity, TArray<FString>& Values) const
+{
+	Super::DescribeRuntimeValues(OwnerComp, NodeMemory, Verbosity, Values);
+
+	if (!bInfiniteLoop)
+	{
+		FBTLoopDecoratorMemory* DecoratorMemory = (FBTLoopDecoratorMemory*)NodeMemory;
+		Values.Add(FString::Printf(TEXT("remaining: %d"), DecoratorMemory->RemainingExecutions));
+	}
+}
+
+uint16 UBTDecorator_Loop::GetInstanceMemorySize() const
+{
+	return sizeof(FBTLoopDecoratorMemory);
+}
diff --git a/Engine/Source/Runtime/AIModule/Private/BehaviorTree/Decorators/BTDecorator_Optional.cpp b/Engine/Source/Runtime/AIModule/Private/BehaviorTree/Decorators/BTDecorator_Optional.cpp
new file mode 100644
--- /dev/null
+++ b/Engine/Source/Runtime/AIModule/Private/BehaviorTree/Decorators/BTDecorator_Optional.cpp
@@ -0,0 +1,19 @@
+// Copyright 1998-2014 Epic Games, Inc. All Rights Reserved.
+
+#include "AIModulePrivate.h"
+#include "BehaviorTree/Decorators/BTDecorator_Optional.h"
+
+UBTDecorator_Optional::UBTDecorator_Optional(const class FPostConstructInitializeProperties& PCIP) : Super(PCIP)
+{
+	NodeName = TEXT("Optional");
+	bNotifyProcessed = true;
+
+	bAllowAbortNone = false;
+	bAllowAbortLowerPri = false;
+	bAllowAbortChildNodes = false;
+}
+
+void UBTDecorator_Optional::OnNodeProcessed(struct FBehaviorTreeSearchData& SearchData, EBTNodeResult::Type& NodeResult)
+{
+	NodeResult = EBTNodeResult::Optional;
+}
diff --git a/Engine/Source/Runtime/AIModule/Private/BehaviorTree/Decorators/BTDecorator_ReachedMoveGoal.cpp b/Engine/Source/Runtime/AIModule/Private/BehaviorTree/Decorators/BTDecorator_ReachedMoveGoal.cpp
new file mode 100644
--- /dev/null
+++ b/Engine/Source/Runtime/AIModule/Private/BehaviorTree/Decorators/BTDecorator_ReachedMoveGoal.cpp
@@ -0,0 +1,22 @@
+// Copyright 1998-2014 Epic Games, Inc. All Rights Reserved.
+
+#include "AIModulePrivate.h"
+#include "BehaviorTree/Decorators/BTDecorator_ReachedMoveGoal.h"
+
+UBTDecorator_ReachedMoveGoal::UBTDecorator_ReachedMoveGoal(const class FPostConstructInitializeProperties& PCIP) : Super(PCIP)
+{
+	NodeName = "Reached move goal";
+
+	// can't abort, it's not observing anything
+	bAllowAbortLowerPri = false;
+	bAllowAbortNone = false;
+	bAllowAbortChildNodes = false;
+	FlowAbortMode = EBTFlowAbortMode::None;
+}
+
+bool UBTDecorator_ReachedMoveGoal::CalculateRawConditionValue(class UBehaviorTreeComponent* OwnerComp, uint8* NodeMemory) const 
+{
+	AAIController* AIOwner = Cast<AAIController>(OwnerComp->GetOwner());
+	const bool bReachedGoal = AIOwner && AIOwner->PathFollowingComponent && AIOwner->PathFollowingComponent->DidMoveReachGoal();
+	return bReachedGoal;
+}
diff --git a/Engine/Source/Runtime/AIModule/Private/BehaviorTree/Decorators/BTDecorator_TimeLimit.cpp b/Engine/Source/Runtime/AIModule/Private/BehaviorTree/Decorators/BTDecorator_TimeLimit.cpp
new file mode 100644
--- /dev/null
+++ b/Engine/Source/Runtime/AIModule/Private/BehaviorTree/Decorators/BTDecorator_TimeLimit.cpp
@@ -0,0 +1,50 @@
+// Copyright 1998-2014 Epic Games, Inc. All Rights Reserved.
+
+#include "AIModulePrivate.h"
+#include "BehaviorTree/Decorators/BTDecorator_TimeLimit.h"
+
+UBTDecorator_TimeLimit::UBTDecorator_TimeLimit(const class FPostConstructInitializeProperties& PCIP) : Super(PCIP)
+{
+	NodeName = "TimeLimit";
+	TimeLimit = 5.0f;
+	bNotifyBecomeRelevant = true;
+	bNotifyTick = true;
+	bTickIntervals = true;
+
+	// time limit always abort current branch
+	bAllowAbortLowerPri = false;
+	bAllowAbortNone = false;
+	FlowAbortMode = EBTFlowAbortMode::Self;
+}
+
+void UBTDecorator_TimeLimit::OnBecomeRelevant(UBehaviorTreeComponent* OwnerComp, uint8* NodeMemory)
+{
+	FBTAuxiliaryMemory* DecoratorMemory = GetSpecialNodeMemory<FBTAuxiliaryMemory>(NodeMemory);
+	DecoratorMemory->NextTickRemainingTime = TimeLimit;
+}
+
+void UBTDecorator_TimeLimit::TickNode(UBehaviorTreeComponent* OwnerComp, uint8* NodeMemory, float DeltaSeconds)
+{
+	OwnerComp->RequestExecution(this);
+}
+
+FString UBTDecorator_TimeLimit::GetStaticDescription() const
+{
+	// basic info: result after time
+	return FString::Printf(TEXT("%s: %s after %.1fs"), *Super::GetStaticDescription(),
+		*UBehaviorTreeTypes::DescribeNodeResult(EBTNodeResult::Failed), TimeLimit);
+}
+
+void UBTDecorator_TimeLimit::DescribeRuntimeValues(const class UBehaviorTreeComponent* OwnerComp, uint8* NodeMemory, EBTDescriptionVerbosity::Type Verbosity, TArray<FString>& Values) const
+{
+	Super::DescribeRuntimeValues(OwnerComp, NodeMemory, Verbosity, Values);
+
+	FBTAuxiliaryMemory* DecoratorMemory = GetSpecialNodeMemory<FBTAuxiliaryMemory>(NodeMemory);
+	if (OwnerComp && OwnerComp->GetWorld())
+	{
+		const float TimeLeft = DecoratorMemory->NextTickRemainingTime > 0 ? DecoratorMemory->NextTickRemainingTime : 0;
+		Values.Add(FString::Printf(TEXT("%s in %ss"),
+			*UBehaviorTreeTypes::DescribeNodeResult(EBTNodeResult::Failed),
+			*FString::SanitizeFloat(TimeLeft)));
+	}
+}
diff --git a/Engine/Source/Runtime/AIModule/Private/BehaviorTree/Services/BTService_BlackboardBase.cpp b/Engine/Source/Runtime/AIModule/Private/BehaviorTree/Services/BTService_BlackboardBase.cpp
new file mode 100644
--- /dev/null
+++ b/Engine/Source/Runtime/AIModule/Private/BehaviorTree/Services/BTService_BlackboardBase.cpp
@@ -0,0 +1,17 @@
+// Copyright 1998-2014 Epic Games, Inc. All Rights Reserved.
+
+#include "AIModulePrivate.h"
+#include "BehaviorTree/Services/BTService_BlackboardBase.h"
+
+UBTService_BlackboardBase::UBTService_BlackboardBase(const class FPostConstructInitializeProperties& PCIP) : Super(PCIP)
+{
+	NodeName = "BlackboardBase";
+
+	// empty KeySelector = allow everything
+}
+
+void UBTService_BlackboardBase::InitializeFromAsset(class UBehaviorTree* Asset)
+{
+	Super::InitializeFromAsset(Asset);
+	BlackboardKey.CacheSelectedKey(GetBlackboardAsset());
+}
diff --git a/Engine/Source/Runtime/AIModule/Private/BehaviorTree/Services/BTService_BlueprintBase.cpp b/Engine/Source/Runtime/AIModule/Private/BehaviorTree/Services/BTService_BlueprintBase.cpp
new file mode 100644
--- /dev/null
+++ b/Engine/Source/Runtime/AIModule/Private/BehaviorTree/Services/BTService_BlueprintBase.cpp
@@ -0,0 +1,102 @@
+// Copyright 1998-2014 Epic Games, Inc. All Rights Reserved.
+
+#include "AIModulePrivate.h"
+#include "BlueprintNodeHelpers.h"
+#include "BehaviorTree/Services/BTService_BlueprintBase.h"
+
+UBTService_BlueprintBase::UBTService_BlueprintBase(const class FPostConstructInitializeProperties& PCIP) : Super(PCIP)
+{
+	UClass* StopAtClass = UBTService_BlueprintBase::StaticClass();
+	bImplementsReceiveTick = BlueprintNodeHelpers::HasBlueprintFunction(TEXT("ReceiveTick"), this, StopAtClass);
+	bImplementsReceiveActivation = BlueprintNodeHelpers::HasBlueprintFunction(TEXT("ReceiveActivation"), this, StopAtClass);
+	bImplementsReceiveDeactivation = BlueprintNodeHelpers::HasBlueprintFunction(TEXT("ReceiveDeactivation"), this, StopAtClass);
+
+	bNotifyBecomeRelevant = bImplementsReceiveActivation || bImplementsReceiveTick;
+	bNotifyCeaseRelevant = bNotifyBecomeRelevant;
+	bNotifyTick = bImplementsReceiveTick;
+	bShowPropertyDetails = true;
+
+	// all blueprint based nodes must create instances
+	bCreateNodeInstance = true;
+
+	if (HasAnyFlags(RF_ClassDefaultObject))
+	{
+		BlueprintNodeHelpers::CollectPropertyData(this, StopAtClass, PropertyData);
+	}
+}
+
+void UBTService_BlueprintBase::PostInitProperties()
+{
+	Super::PostInitProperties();
+	NodeName = BlueprintNodeHelpers::GetNodeName(this);
+}
+
+void UBTService_BlueprintBase::OnBecomeRelevant(class UBehaviorTreeComponent* OwnerComp, uint8* NodeMemory)
+{
+	Super::OnBecomeRelevant(OwnerComp, NodeMemory);
+
+	// check flag, it could be used because user wants tick
+	if (bImplementsReceiveActivation)
+	{
+		ReceiveActivation(OwnerComp->GetOwner());
+	}
+}
+
+void UBTService_BlueprintBase::OnCeaseRelevant(class UBehaviorTreeComponent* OwnerComp, uint8* NodeMemory)
+{
+	// force dropping all pending latent actions associated with this blueprint
+	// we can't have those resuming activity when node is/was aborted
+	BlueprintNodeHelpers::AbortLatentActions(OwnerComp->GetOwner(), this);
+
+	Super::OnCeaseRelevant(OwnerComp, NodeMemory);
+
+	if (bImplementsReceiveDeactivation)
+	{
+		ReceiveDeactivation(OwnerComp->GetOwner());
+	}
+}
+
+void UBTService_BlueprintBase::TickNode(class UBehaviorTreeComponent* OwnerComp, uint8* NodeMemory, float DeltaSeconds)
+{
+	Super::TickNode(OwnerComp, NodeMemory, DeltaSeconds);
+
+	// skip flag, will be handled by bNotifyTick
+
+	ReceiveTick(OwnerComp->GetOwner(), DeltaSeconds);
+}
+
+bool UBTService_BlueprintBase::IsServiceActive() const
+{
+	UBehaviorTreeComponent* OwnerComp = Cast<UBehaviorTreeComponent>(GetOuter());
+	const bool bIsActive = OwnerComp->IsAuxNodeActive(this);
+	return bIsActive;
+}
+
+FString UBTService_BlueprintBase::GetStaticDescription() const
+{
+	FString ReturnDesc = Super::GetStaticDescription();
+
+	UBTService_BlueprintBase* CDO = (UBTService_BlueprintBase*)(GetClass()->GetDefaultObject());
+	if (bShowPropertyDetails && CDO)
+	{
+		UClass* StopAtClass = UBTService_BlueprintBase::StaticClass();
+		FString PropertyDesc = BlueprintNodeHelpers::CollectPropertyDescription(this, StopAtClass, CDO->PropertyData);
+		if (PropertyDesc.Len())
+		{
+			ReturnDesc += TEXT(":\n\n");
+			ReturnDesc += PropertyDesc;
+		}
+	}
+
+	return ReturnDesc;
+}
+
+void UBTService_BlueprintBase::DescribeRuntimeValues(const class UBehaviorTreeComponent* OwnerComp, uint8* NodeMemory, EBTDescriptionVerbosity::Type Verbosity, TArray<FString>& Values) const
+{
+	UBTService_BlueprintBase* CDO = (UBTService_BlueprintBase*)(GetClass()->GetDefaultObject());
+	if (CDO && CDO->PropertyData.Num())
+	{
+		UClass* StopAtClass = UBTService_BlueprintBase::StaticClass();
+		BlueprintNodeHelpers::DescribeRuntimeValues(this, CDO->PropertyData, Values);
+	}
+}
diff --git a/Engine/Source/Runtime/AIModule/Private/BehaviorTree/Services/BTService_DefaultFocus.cpp b/Engine/Source/Runtime/AIModule/Private/BehaviorTree/Services/BTService_DefaultFocus.cpp
new file mode 100644
--- /dev/null
+++ b/Engine/Source/Runtime/AIModule/Private/BehaviorTree/Services/BTService_DefaultFocus.cpp
@@ -0,0 +1,101 @@
+// Copyright 1998-2014 Epic Games, Inc. All Rights Reserved.
+
+#include "AIModulePrivate.h"
+#include "BehaviorTree/Services/BTService_DefaultFocus.h"
+#include "BehaviorTree/Blackboard/BlackboardKeyAllTypes.h"
+
+UBTService_DefaultFocus::UBTService_DefaultFocus(const class FPostConstructInitializeProperties& PCIP) 
+	: Super(PCIP)
+{
+	NodeName = "Set default focus";
+
+	bNotifyTick = false;
+	bTickIntervals = false;
+	bNotifyBecomeRelevant = true;
+	bNotifyCeaseRelevant = true;
+
+	// accept only actors and vectors
+	BlackboardKey.AddObjectFilter(this, AActor::StaticClass());
+	BlackboardKey.AddVectorFilter(this);
+}
+
+void UBTService_DefaultFocus::OnBecomeRelevant(class UBehaviorTreeComponent* OwnerComp, uint8* NodeMemory) 
+{
+	Super::OnBecomeRelevant(OwnerComp, NodeMemory);
+
+	FBTDefaultFocusMemory* MyMemory = (FBTDefaultFocusMemory*)NodeMemory;
+	check(MyMemory);
+	MyMemory->Reset();
+
+	AAIController* OwnerController = OwnerComp ? Cast<AAIController>(OwnerComp->GetOwner()) : NULL;	
+	const UBlackboardComponent* MyBlackboard = OwnerComp->GetBlackboardComponent();
+	
+	if (OwnerController != NULL && MyBlackboard != NULL)
+	{
+		if (BlackboardKey.SelectedKeyType == UBlackboardKeyType_Object::StaticClass())
+		{
+			UObject* KeyValue = MyBlackboard->GetValueAsObject(BlackboardKey.GetSelectedKeyID());
+			AActor* TargetActor = Cast<AActor>(KeyValue);
+			if (TargetActor)
+			{
+				OwnerController->SetFocus(TargetActor, EAIFocusPriority::Default);
+				MyMemory->FocusActorSet = TargetActor;
+				MyMemory->bActorSet = true;
+			}
+		}
+		else
+		{
+			const FVector FocusLocation = MyBlackboard->GetValueAsVector(BlackboardKey.GetSelectedKeyID());
+			OwnerController->SetFocalPoint(FocusLocation, /*bOffsetFromBase=*/false, EAIFocusPriority::Default);
+			MyMemory->FocusLocationSet = FocusLocation;
+		}
+	}
+}
+
+void UBTService_DefaultFocus::OnCeaseRelevant(class UBehaviorTreeComponent* OwnerComp, uint8* NodeMemory)
+{
+	Super::OnCeaseRelevant(OwnerComp, NodeMemory);
+
+	FBTDefaultFocusMemory* MyMemory = (FBTDefaultFocusMemory*)NodeMemory;
+	check(MyMemory);
+	AAIController* OwnerController = OwnerComp ? Cast<AAIController>(OwnerComp->GetOwner()) : NULL;	
+	if (OwnerController != NULL)
+	{
+		bool bClearFocus = false;
+		if (MyMemory->bActorSet)
+		{
+			bClearFocus = (MyMemory->FocusActorSet == OwnerController->GetFocusActor(EAIFocusPriority::Default));
+		}
+		else
+		{
+			bClearFocus = (MyMemory->FocusLocationSet == OwnerController->GetFocalPoint(EAIFocusPriority::Default));
+		}
+
+		if (bClearFocus)
+		{
+			OwnerController->ClearFocus(EAIFocusPriority::Default);
+		}
+	}
+}
+
+FString UBTService_DefaultFocus::GetStaticDescription() const
+{	
+	FString KeyDesc("invalid");
+	if (BlackboardKey.SelectedKeyType == UBlackboardKeyType_Object::StaticClass() ||
+		BlackboardKey.SelectedKeyType == UBlackboardKeyType_Vector::StaticClass())
+	{
+		KeyDesc = BlackboardKey.SelectedKeyName.ToString();
+	}
+
+	return FString::Printf(TEXT("Set default focus to %s"), *KeyDesc);
+}
+
+void UBTService_DefaultFocus::DescribeRuntimeValues(const class UBehaviorTreeComponent* OwnerComp, uint8* NodeMemory, EBTDescriptionVerbosity::Type Verbosity, TArray<FString>& Values) const
+{
+	Super::DescribeRuntimeValues(OwnerComp, NodeMemory, Verbosity, Values);
+}
+
+uint16 UBTService_DefaultFocus::GetInstanceMemorySize() const
+{
+	return sizeof(FBTDefaultFocusMemory);
+}
diff --git a/Engine/Source/Runtime/AIModule/Private/BehaviorTree/Tasks/BTTask_BlackboardBase.cpp b/Engine/Source/Runtime/AIModule/Private/BehaviorTree/Tasks/BTTask_BlackboardBase.cpp
new file mode 100644
--- /dev/null
+++ b/Engine/Source/Runtime/AIModule/Private/BehaviorTree/Tasks/BTTask_BlackboardBase.cpp
@@ -0,0 +1,17 @@
+// Copyright 1998-2014 Epic Games, Inc. All Rights Reserved.
+
+#include "AIModulePrivate.h"
+#include "BehaviorTree/Tasks/BTTask_BlackboardBase.h"
+
+UBTTask_BlackboardBase::UBTTask_BlackboardBase(const class FPostConstructInitializeProperties& PCIP) : Super(PCIP)
+{
+	NodeName = "BlackboardBase";
+
+	// empty KeySelector = allow everything
+}
+
+void UBTTask_BlackboardBase::InitializeFromAsset(class UBehaviorTree* Asset)
+{
+	Super::InitializeFromAsset(Asset);
+	BlackboardKey.CacheSelectedKey(GetBlackboardAsset());
+}
diff --git a/Engine/Source/Runtime/AIModule/Private/BehaviorTree/Tasks/BTTask_BlueprintBase.cpp b/Engine/Source/Runtime/AIModule/Private/BehaviorTree/Tasks/BTTask_BlueprintBase.cpp
new file mode 100644
--- /dev/null
+++ b/Engine/Source/Runtime/AIModule/Private/BehaviorTree/Tasks/BTTask_BlueprintBase.cpp
@@ -0,0 +1,155 @@
+// Copyright 1998-2014 Epic Games, Inc. All Rights Reserved.
+
+#include "AIModulePrivate.h"
+#include "BlueprintNodeHelpers.h"
+#include "BehaviorTree/Tasks/BTTask_BlueprintBase.h"
+
+UBTTask_BlueprintBase::UBTTask_BlueprintBase(const class FPostConstructInitializeProperties& PCIP) : Super(PCIP)
+{
+	UClass* StopAtClass = UBTTask_BlueprintBase::StaticClass();
+	bImplementsReceiveTick = BlueprintNodeHelpers::HasBlueprintFunction(TEXT("ReceiveTick"), this, StopAtClass);
+	bImplementsReceiveExecute = BlueprintNodeHelpers::HasBlueprintFunction(TEXT("ReceiveExecute"), this, StopAtClass);
+	bImplementsReceiveAbort = BlueprintNodeHelpers::HasBlueprintFunction(TEXT("ReceiveAbort"), this, StopAtClass);
+
+	bNotifyTick = bImplementsReceiveTick;
+	bShowPropertyDetails = true;
+
+	// all blueprint based nodes must create instances
+	bCreateNodeInstance = true;
+
+	if (HasAnyFlags(RF_ClassDefaultObject))
+	{
+		BlueprintNodeHelpers::CollectPropertyData(this, StopAtClass, PropertyData);
+	}
+}
+
+void UBTTask_BlueprintBase::PostInitProperties()
+{
+	Super::PostInitProperties();
+	NodeName = BlueprintNodeHelpers::GetNodeName(this);
+}
+
+EBTNodeResult::Type UBTTask_BlueprintBase::ExecuteTask(class UBehaviorTreeComponent* OwnerComp, uint8* NodeMemory)
+{
+	CurrentCallResult = bImplementsReceiveExecute ? EBTNodeResult::InProgress : EBTNodeResult::Failed;
+	if (bImplementsReceiveExecute)
+	{
+		bStoreFinishResult = true;
+
+		ReceiveExecute(OwnerComp->GetOwner());
+
+		bStoreFinishResult = false;
+	}
+
+	return CurrentCallResult;
+}
+
+EBTNodeResult::Type UBTTask_BlueprintBase::AbortTask(class UBehaviorTreeComponent* OwnerComp, uint8* NodeMemory)
+{
+	// force dropping all pending latent actions associated with this blueprint
+	// we can't have those resuming activity when node is/was aborted
+	BlueprintNodeHelpers::AbortLatentActions(OwnerComp->GetOwner(), this);
+
+	CurrentCallResult = bImplementsReceiveAbort ? EBTNodeResult::InProgress : EBTNodeResult::Aborted;
+	if (bImplementsReceiveAbort)
+	{
+		bStoreFinishResult = true;
+
+		ReceiveAbort(OwnerComp->GetOwner());
+
+		bStoreFinishResult = false;
+	}
+
+	return CurrentCallResult;
+}
+
+void UBTTask_BlueprintBase::TickTask(class UBehaviorTreeComponent* OwnerComp, uint8* NodeMemory, float DeltaSeconds) 
+{
+	// skip flag, will be handled by bNotifyTick
+
+	ReceiveTick(OwnerComp->GetOwner(), DeltaSeconds);
+}
+
+void UBTTask_BlueprintBase::FinishExecute(bool bSuccess)
+{
+	UBehaviorTreeComponent* OwnerComp = Cast<UBehaviorTreeComponent>(GetOuter());
+	EBTNodeResult::Type NodeResult(bSuccess ? EBTNodeResult::Succeeded : EBTNodeResult::Failed);
+
+	if (bStoreFinishResult)
+	{
+		CurrentCallResult = NodeResult;
+	}
+	else if (OwnerComp)
+	{
+		FinishLatentTask(OwnerComp, NodeResult);
+	}
+}
+
+void UBTTask_BlueprintBase::FinishAbort()
+{
+	UBehaviorTreeComponent* OwnerComp = Cast<UBehaviorTreeComponent>(GetOuter());
+	EBTNodeResult::Type NodeResult(EBTNodeResult::Aborted);
+
+	if (bStoreFinishResult)
+	{
+		CurrentCallResult = NodeResult;
+	}
+	else if (OwnerComp)
+	{
+		FinishLatentAbort(OwnerComp);
+	}
+}
+
+bool UBTTask_BlueprintBase::IsTaskExecuting() const
+{
+	UBehaviorTreeComponent* OwnerComp = Cast<UBehaviorTreeComponent>(GetOuter());
+	EBTTaskStatus::Type TaskStatus = OwnerComp->GetTaskStatus(this);
+
+	return (TaskStatus == EBTTaskStatus::Active);
+}
+
+void UBTTask_BlueprintBase::SetFinishOnMessage(FName MessageName)
+{
+	UBehaviorTreeComponent* OwnerComp = Cast<UBehaviorTreeComponent>(GetOuter());
+	if (OwnerComp)
+	{
+		OwnerComp->RegisterMessageObserver(this, MessageName);
+	}
+}
+
+void UBTTask_BlueprintBase::SetFinishOnMessageWithId(FName MessageName, int32 RequestID)
+{
+	UBehaviorTreeComponent* OwnerComp = Cast<UBehaviorTreeComponent>(GetOuter());
+	if (OwnerComp)
+	{
+		OwnerComp->RegisterMessageObserver(this, MessageName, RequestID);
+	}
+}
+
+FString UBTTask_BlueprintBase::GetStaticDescription() const
+{
+	FString ReturnDesc = Super::GetStaticDescription();
+
+	UBTTask_BlueprintBase* CDO = (UBTTask_BlueprintBase*)(GetClass()->GetDefaultObject());
+	if (bShowPropertyDetails && CDO)
+	{
+		UClass* StopAtClass = UBTTask_BlueprintBase::StaticClass();
+		FString PropertyDesc = BlueprintNodeHelpers::CollectPropertyDescription(this, StopAtClass, CDO->PropertyData);
+		if (PropertyDesc.Len())
+		{
+			ReturnDesc += TEXT(":\n\n");
+			ReturnDesc += PropertyDesc;
+		}
+	}
+
+	return ReturnDesc;
+}
+
+void UBTTask_BlueprintBase::DescribeRuntimeValues(const class UBehaviorTreeComponent* OwnerComp, uint8* NodeMemory, EBTDescriptionVerbosity::Type Verbosity, TArray<FString>& Values) const
+{
+	UBTTask_BlueprintBase* CDO = (UBTTask_BlueprintBase*)(GetClass()->GetDefaultObject());
+	if (CDO && CDO->PropertyData.Num())
+	{
+		BlueprintNodeHelpers::DescribeRuntimeValues(this, CDO->PropertyData, Values);
+	}
+}
diff --git a/Engine/Source/Runtime/AIModule/Private/BehaviorTree/Tasks/BTTask_MakeNoise.cpp b/Engine/Source/Runtime/AIModule/Private/BehaviorTree/Tasks/BTTask_MakeNoise.cpp
new file mode 100644
--- /dev/null
+++ b/Engine/Source/Runtime/AIModule/Private/BehaviorTree/Tasks/BTTask_MakeNoise.cpp
@@ -0,0 +1,27 @@
+// Copyright 1998-2014 Epic Games, Inc. All Rights Reserved.
+
+#include "AIModulePrivate.h"
+#include "SoundDefinitions.h"
+#include "BehaviorTree/Tasks/BTTask_MakeNoise.h"
+
+UBTTask_MakeNoise::UBTTask_MakeNoise(const class FPostConstructInitializeProperties& PCIP) 
+	: Super(PCIP)
+	, Loudnes(1.0f)
+{
+	NodeName = "MakeNoise";
+}
+
+EBTNodeResult::Type UBTTask_MakeNoise::ExecuteTask(UBehaviorTreeComponent* OwnerComp, uint8* NodeMemory)
+{
+	const AAIController* MyController = OwnerComp ? Cast<AAIController>(OwnerComp->GetOwner()) : NULL;
+	APawn* MyPawn = MyController ? MyController->GetPawn() : NULL;
+
+	if (MyPawn)
+	{
+		MyPawn->MakeNoise(Loudnes, MyPawn);
+		return EBTNodeResult::Succeeded;
+	}
+
+	return EBTNodeResult::Failed;
+}
+
diff --git a/Engine/Source/Runtime/AIModule/Private/BehaviorTree/Tasks/BTTask_MoveDirectlyToward.cpp b/Engine/Source/Runtime/AIModule/Private/BehaviorTree/Tasks/BTTask_MoveDirectlyToward.cpp
new file mode 100644
--- /dev/null
+++ b/Engine/Source/Runtime/AIModule/Private/BehaviorTree/Tasks/BTTask_MoveDirectlyToward.cpp
@@ -0,0 +1,110 @@
+// Copyright 1998-2014 Epic Games, Inc. All Rights Reserved.
+
+#include "AIModulePrivate.h"
+#include "BehaviorTree/BlackboardComponent.h"
+#include "BehaviorTree/Blackboard/BlackboardKeyType_Object.h"
+#include "BehaviorTree/Blackboard/BlackboardKeyType_Vector.h"
+#include "BehaviorTree/Tasks/BTTask_MoveDirectlyToward.h"
+
+UBTTask_MoveDirectlyToward::UBTTask_MoveDirectlyToward(const class FPostConstructInitializeProperties& PCIP) 
+	: Super(PCIP)
+	, AcceptableRadius(50.f)
+	, bProjectVectorGoalToNavigation(true)
+	, bAllowStrafe(true)
+{
+	NodeName = "MoveDirectlyToward";
+
+	// accept only actors and vectors
+	BlackboardKey.AddObjectFilter(this, AActor::StaticClass());
+	BlackboardKey.AddVectorFilter(this);
+}
+
+EBTNodeResult::Type UBTTask_MoveDirectlyToward::ExecuteTask(class UBehaviorTreeComponent* OwnerComp, uint8* NodeMemory)
+{
+	const UBlackboardComponent* MyBlackboard = OwnerComp->GetBlackboardComponent();
+	FBTMoveDirectlyTowardMemory* MyMemory = (FBTMoveDirectlyTowardMemory*)NodeMemory;
+	AAIController* MyController = OwnerComp ? Cast<AAIController>(OwnerComp->GetOwner()) : NULL;
+	EBTNodeResult::Type NodeResult = EBTNodeResult::Failed;
+
+	if (MyController && MyBlackboard)
+	{
+		EPathFollowingRequestResult::Type RequestResult = EPathFollowingRequestResult::Failed;
+
+		if (BlackboardKey.SelectedKeyType == UBlackboardKeyType_Object::StaticClass())
+		{
+			UObject* KeyValue = MyBlackboard->GetValueAsObject(BlackboardKey.GetSelectedKeyID());
+			AActor* TargetActor = Cast<AActor>(KeyValue);
+			if (TargetActor)
+			{
+				RequestResult = bForceMoveToLocation ?
+					MyController->MoveToLocation(TargetActor->GetActorLocation(), AcceptableRadius, /*bStopOnOverlap=*/true, /*bUsePathfinding=*/false, /*bProjectDestinationToNavigation=*/bProjectVectorGoalToNavigation, bAllowStrafe) :
+					MyController->MoveToActor(TargetActor, AcceptableRadius, /*bStopOnOverlap=*/true, /*bUsePathfinding=*/false, bAllowStrafe);
+			}
+		}
+		else if (BlackboardKey.SelectedKeyType == UBlackboardKeyType_Vector::StaticClass())
+		{
+			const FVector TargetLocation = MyBlackboard->GetValueAsVector(BlackboardKey.GetSelectedKeyID());
+			RequestResult = MyController->MoveToLocation(TargetLocation, AcceptableRadius, /*bStopOnOverlap=*/true, /*bUsePathfinding=*/false, /*bProjectDestinationToNavigation=*/bProjectVectorGoalToNavigation, bAllowStrafe);
+		}
+
+		if (RequestResult == EPathFollowingRequestResult::RequestSuccessful)
+		{
+			const FAIRequestID RequestID = MyController->GetCurrentMoveRequestID();
+
+			MyMemory->MoveRequestID = RequestID;
+			WaitForMessage(OwnerComp, UBrainComponent::AIMessage_MoveFinished, RequestID);
+
+			NodeResult = EBTNodeResult::InProgress;
+		}
+		else if (RequestResult == EPathFollowingRequestResult::AlreadyAtGoal)
+		{
+			NodeResult = EBTNodeResult::Succeeded;
+		}
+	}
+
+	return NodeResult;
+}
+
+EBTNodeResult::Type UBTTask_MoveDirectlyToward::AbortTask(class UBehaviorTreeComponent* OwnerComp, uint8* NodeMemory)
+{
+	FBTMoveDirectlyTowardMemory* MyMemory = (FBTMoveDirectlyTowardMemory*)NodeMemory;
+	AAIController* MyController = OwnerComp ? Cast<AAIController>(OwnerComp->GetOwner()) : NULL;
+
+	if (MyController && MyController->PathFollowingComponent)
+	{
+		MyController->PathFollowingComponent->AbortMove(TEXT("BehaviorTree abort"), MyMemory->MoveRequestID);
+	}
+
+	return Super::AbortTask(OwnerComp, NodeMemory);
+}
+
+FString UBTTask_MoveDirectlyToward::GetStaticDescription() const
+{
+	FString KeyDesc("invalid");
+	if (BlackboardKey.SelectedKeyType == UBlackboardKeyType_Object::StaticClass() ||
+		BlackboardKey.SelectedKeyType == UBlackboardKeyType_Vector::StaticClass())
+	{
+		KeyDesc = BlackboardKey.SelectedKeyName.ToString();
+	}
+
+	return FString::Printf(TEXT("%s: %s"), *Super::GetStaticDescription(), *KeyDesc);
+}
+
+void UBTTask_MoveDirectlyToward::DescribeRuntimeValues(const class UBehaviorTreeComponent* OwnerComp, uint8* NodeMemory, EBTDescriptionVerbosity::Type Verbosity, TArray<FString>& Values) const
+{
+	Super::DescribeRuntimeValues(OwnerComp, NodeMemory, Verbosity, Values);
+
+	const UBlackboardComponent* BlackboardComp = OwnerComp->GetBlackboardComponent();
+	FBTMoveDirectlyTowardMemory* MyMemory = (FBTMoveDirectlyTowardMemory*)NodeMemory;
+
+	if (MyMemory->MoveRequestID && BlackboardComp)
+	{
+		FString KeyValue = BlackboardComp->DescribeKeyValue(BlackboardKey.GetSelectedKeyID(), EBlackboardDescription::OnlyValue);
+		Values.Add(FString::Printf(TEXT("move target: %s"), *KeyValue));
+	}
+}
+
+uint16 UBTTask_MoveDirectlyToward::GetInstanceMemorySize() const
+{
+	return sizeof(FBTMoveDirectlyTowardMemory);
+}
diff --git a/Engine/Source/Runtime/AIModule/Private/BehaviorTree/Tasks/BTTask_MoveTo.cpp b/Engine/Source/Runtime/AIModule/Private/BehaviorTree/Tasks/BTTask_MoveTo.cpp
new file mode 100644
--- /dev/null
+++ b/Engine/Source/Runtime/AIModule/Private/BehaviorTree/Tasks/BTTask_MoveTo.cpp
@@ -0,0 +1,165 @@
+// Copyright 1998-2014 Epic Games, Inc. All Rights Reserved.
+
+#include "AIModulePrivate.h"
+#include "BehaviorTree/Tasks/BTTask_MoveTo.h"
+#include "BehaviorTree/BehaviorTree.h"
+#include "BehaviorTree/BlackboardComponent.h"
+#include "BehaviorTree/Blackboard/BlackboardKeyType_Object.h"
+#include "BehaviorTree/Blackboard/BlackboardKeyType_Vector.h"
+
+UBTTask_MoveTo::UBTTask_MoveTo(const class FPostConstructInitializeProperties& PCIP) 
+	: Super(PCIP)
+	, AcceptableRadius(50.f)
+	, bAllowStrafe(false)
+{
+	NodeName = "Move To";
+	bNotifyTick = true;
+
+	// accept only actors and vectors
+	BlackboardKey.AddObjectFilter(this, AActor::StaticClass());
+	BlackboardKey.AddVectorFilter(this);
+}
+
+EBTNodeResult::Type UBTTask_MoveTo::ExecuteTask(class UBehaviorTreeComponent* OwnerComp, uint8* NodeMemory)
+{
+	EBTNodeResult::Type NodeResult = EBTNodeResult::InProgress;
+
+	FBTMoveToTaskMemory* MyMemory = (FBTMoveToTaskMemory*)NodeMemory;
+	AAIController* MyController = OwnerComp ? Cast<AAIController>(OwnerComp->GetOwner()) : NULL;
+
+	MyMemory->bWaitingForPath = MyController->ShouldPostponePathUpdates();
+	if (!MyMemory->bWaitingForPath)
+	{
+		NodeResult = PerformMoveTask(OwnerComp, NodeMemory);
+	}
+	else
+	{
+		UE_VLOG(MyController, LogBehaviorTree, Log, TEXT("Pathfinding requests are freezed, waiting..."));
+	}
+	
+	return NodeResult;
+}
+
+EBTNodeResult::Type UBTTask_MoveTo::PerformMoveTask(class UBehaviorTreeComponent* OwnerComp, uint8* NodeMemory)
+{
+	const UBlackboardComponent* MyBlackboard = OwnerComp->GetBlackboardComponent();
+	FBTMoveToTaskMemory* MyMemory = (FBTMoveToTaskMemory*)NodeMemory;
+	AAIController* MyController = OwnerComp ? Cast<AAIController>(OwnerComp->GetOwner()) : NULL;
+
+	EBTNodeResult::Type NodeResult = EBTNodeResult::Failed;
+
+	if (MyController && MyBlackboard)
+	{
+		EPathFollowingRequestResult::Type RequestResult = EPathFollowingRequestResult::Failed;
+
+		if (BlackboardKey.SelectedKeyType == UBlackboardKeyType_Object::StaticClass())
+		{
+			UObject* KeyValue = MyBlackboard->GetValueAsObject(BlackboardKey.GetSelectedKeyID());
+			AActor* TargetActor = Cast<AActor>(KeyValue);
+			if (TargetActor)
+			{
+				RequestResult = MyController->MoveToActor(TargetActor, AcceptableRadius, true, true, bAllowStrafe, FilterClass);
+			}
+			else
+			{
+				UE_VLOG(MyController, LogBehaviorTree, Warning, TEXT("UBTTask_MoveTo::ExecuteTask tried to go to actor while BB %s entry was empty"), *BlackboardKey.SelectedKeyName.ToString());
+			}
+		}
+		else if (BlackboardKey.SelectedKeyType == UBlackboardKeyType_Vector::StaticClass())
+		{
+			const FVector TargetLocation = MyBlackboard->GetValueAsVector(BlackboardKey.GetSelectedKeyID());
+			RequestResult = MyController->MoveToLocation(TargetLocation, AcceptableRadius, true, true, false, bAllowStrafe, FilterClass);
+		}
+
+		if (RequestResult == EPathFollowingRequestResult::RequestSuccessful)
+		{
+			const FAIRequestID RequestID = MyController->GetCurrentMoveRequestID();
+
+			MyMemory->MoveRequestID = RequestID;
+			WaitForMessage(OwnerComp, UBrainComponent::AIMessage_MoveFinished, RequestID);
+			WaitForMessage(OwnerComp, UBrainComponent::AIMessage_RepathFailed);
+
+			NodeResult = EBTNodeResult::InProgress;
+		}
+		else if (RequestResult == EPathFollowingRequestResult::AlreadyAtGoal)
+		{
+			NodeResult = EBTNodeResult::Succeeded;
+		}
+	}
+
+	return NodeResult;
+}
+
+EBTNodeResult::Type UBTTask_MoveTo::AbortTask(class UBehaviorTreeComponent* OwnerComp, uint8* NodeMemory)
+{
+	FBTMoveToTaskMemory* MyMemory = (FBTMoveToTaskMemory*)NodeMemory;
+	if (!MyMemory->bWaitingForPath)
+	{
+		AAIController* MyController = OwnerComp ? Cast<AAIController>(OwnerComp->GetOwner()) : NULL;
+
+		if (MyController && MyController->PathFollowingComponent)
+		{
+			MyController->PathFollowingComponent->AbortMove(TEXT("BehaviorTree abort"), MyMemory->MoveRequestID);
+		}
+	}
+
+	return Super::AbortTask(OwnerComp, NodeMemory);
+}
+
+void UBTTask_MoveTo::TickTask(class UBehaviorTreeComponent* OwnerComp, uint8* NodeMemory, float DeltaSeconds)
+{
+	FBTMoveToTaskMemory* MyMemory = (FBTMoveToTaskMemory*)NodeMemory;
+	if (MyMemory->bWaitingForPath && !OwnerComp->IsPaused())
+	{
+		AAIController* MyController = OwnerComp ? Cast<AAIController>(OwnerComp->GetOwner()) : NULL;
+		if (MyController && !MyController->ShouldPostponePathUpdates())
+		{
+			UE_VLOG(MyController, LogBehaviorTree, Log, TEXT("Pathfinding requests are unlocked!"));
+			MyMemory->bWaitingForPath = false;
+
+			const EBTNodeResult::Type NodeResult = PerformMoveTask(OwnerComp, NodeMemory);
+			if (NodeResult != EBTNodeResult::InProgress)
+			{
+				FinishLatentTask(OwnerComp, NodeResult);
+			}
+		}
+	}
+}
+
+void UBTTask_MoveTo::OnMessage(class UBehaviorTreeComponent* OwnerComp, uint8* NodeMemory, FName Message, int32 SenderID, bool bSuccess)
+{
+	// AIMessage_RepathFailed means task has failed
+	bSuccess &= (Message != UBrainComponent::AIMessage_RepathFailed);
+	Super::OnMessage(OwnerComp, NodeMemory, Message, SenderID, bSuccess);
+}
+
+FString UBTTask_MoveTo::GetStaticDescription() const
+{
+	FString KeyDesc("invalid");
+	if (BlackboardKey.SelectedKeyType == UBlackboardKeyType_Object::StaticClass() ||
+		BlackboardKey.SelectedKeyType == UBlackboardKeyType_Vector::StaticClass())
+	{
+		KeyDesc = BlackboardKey.SelectedKeyName.ToString();
+	}
+
+	return FString::Printf(TEXT("%s: %s"), *Super::GetStaticDescription(), *KeyDesc);
+}
+
+void UBTTask_MoveTo::DescribeRuntimeValues(const class UBehaviorTreeComponent* OwnerComp, uint8* NodeMemory, EBTDescriptionVerbosity::Type Verbosity, TArray<FString>& Values) const
+{
+	Super::DescribeRuntimeValues(OwnerComp, NodeMemory, Verbosity, Values);
+
+	const UBlackboardComponent* BlackboardComp = OwnerComp->GetBlackboardComponent();
+	FBTMoveToTaskMemory* MyMemory = (FBTMoveToTaskMemory*)NodeMemory;
+
+	if (MyMemory->MoveRequestID && BlackboardComp)
+	{
+		FString KeyValue = BlackboardComp->DescribeKeyValue(BlackboardKey.GetSelectedKeyID(), EBlackboardDescription::OnlyValue);
+		Values.Add(FString::Printf(TEXT("move target: %s"), *KeyValue));
+	}
+}
+
+uint16 UBTTask_MoveTo::GetInstanceMemorySize() const
+{
+	return sizeof(FBTMoveToTaskMemory);
+}
diff --git a/Engine/Source/Runtime/AIModule/Private/BehaviorTree/Tasks/BTTask_PlaySound.cpp b/Engine/Source/Runtime/AIModule/Private/BehaviorTree/Tasks/BTTask_PlaySound.cpp
new file mode 100644
--- /dev/null
+++ b/Engine/Source/Runtime/AIModule/Private/BehaviorTree/Tasks/BTTask_PlaySound.cpp
@@ -0,0 +1,30 @@
+// Copyright 1998-2014 Epic Games, Inc. All Rights Reserved.
+
+#include "AIModulePrivate.h"
+#include "SoundDefinitions.h"
+#include "BehaviorTree/Tasks/BTTask_PlaySound.h"
+
+UBTTask_PlaySound::UBTTask_PlaySound(const class FPostConstructInitializeProperties& PCIP) : Super(PCIP)
+{
+	NodeName = "PlaySound";
+}
+
+EBTNodeResult::Type UBTTask_PlaySound::ExecuteTask(UBehaviorTreeComponent* OwnerComp, uint8* NodeMemory)
+{
+	const AAIController* MyController = OwnerComp ? Cast<AAIController>(OwnerComp->GetOwner()) : NULL;
+
+	UAudioComponent* AC = NULL;
+	if (SoundToPlay && MyController)
+	{
+		if (const APawn* MyPawn = MyController->GetPawn())
+		{
+			AC = UGameplayStatics::PlaySoundAttached(SoundToPlay, MyPawn->GetRootComponent());
+		}
+	}
+	return AC ? EBTNodeResult::Succeeded : EBTNodeResult::Failed;
+}
+
+FString UBTTask_PlaySound::GetStaticDescription() const
+{
+	return FString::Printf(TEXT("%s: '%s'"), *Super::GetStaticDescription(), SoundToPlay ? *SoundToPlay->GetName() : TEXT(""));
+}
diff --git a/Engine/Source/Runtime/AIModule/Private/BehaviorTree/Tasks/BTTask_RunBehavior.cpp b/Engine/Source/Runtime/AIModule/Private/BehaviorTree/Tasks/BTTask_RunBehavior.cpp
new file mode 100644
--- /dev/null
+++ b/Engine/Source/Runtime/AIModule/Private/BehaviorTree/Tasks/BTTask_RunBehavior.cpp
@@ -0,0 +1,226 @@
+// Copyright 1998-2014 Epic Games, Inc. All Rights Reserved.
+
+#include "AIModulePrivate.h"
+#include "BehaviorTree/Tasks/BTTask_RunBehavior.h"
+#include "BehaviorTree/BehaviorTreeManager.h"
+
+UBTTask_RunBehavior::UBTTask_RunBehavior(const class FPostConstructInitializeProperties& PCIP) : Super(PCIP)
+{
+	NodeName = "Run Behavior";
+}
+
+EBTNodeResult::Type UBTTask_RunBehavior::ExecuteTask(UBehaviorTreeComponent* OwnerComp, uint8* NodeMemory)
+{
+	const bool bPushed = OwnerComp->PushInstance(BehaviorAsset);
+	return bPushed ? EBTNodeResult::InProgress : EBTNodeResult::Failed;
+}
+
+FString UBTTask_RunBehavior::GetStaticDescription() const
+{
+	return FString::Printf(TEXT("%s: %s"), *Super::GetStaticDescription(), *GetNameSafe(BehaviorAsset));
+}
+
+uint16 UBTTask_RunBehavior::GetInstanceMemorySize() const
+{
+	// memory required by root level decorators is provided by this task
+	// it's not possible to have dynamic subtrees supporting root level decorators
+	// (unnecessary complication of accessing node's memory)
+
+	int32 MemorySize = 0;
+	if (BehaviorAsset)
+	{
+		TArray<uint16> MemoryOffsets;
+		UBehaviorTreeManager::InitializeMemoryHelper(BehaviorAsset->RootDecorators, MemoryOffsets, MemorySize);
+	}
+
+	return MemorySize;
+}
+
+void UBTTask_RunBehavior::InjectNodes(UBehaviorTreeComponent* OwnerComp, uint8* NodeMemory, int32& InstancedIndex) const
+{
+	if (BehaviorAsset == NULL || BehaviorAsset->RootDecorators.Num() == 0)
+	{
+		return;
+	}
+
+	const int32 NumInjectedDecorators = BehaviorAsset->RootDecorators.Num();
+	const int32 FirstNodeIdx = InstancedIndex;
+
+	// initialize on first access
+	if (!OwnerComp->NodeInstances.IsValidIndex(InstancedIndex))
+	{
+		TArray<uint16> MemoryOffsets;
+		int32 MemorySize = 0;
+
+		UBehaviorTreeManager::InitializeMemoryHelper(BehaviorAsset->RootDecorators, MemoryOffsets, MemorySize);
+		const int32 AlignedInstanceMemorySize = UBehaviorTreeManager::GetAlignedDataSize(sizeof(FBTInstancedNodeMemory));
+
+		// prepare dummy memory block for nodes that won't require instancing and offset it by special data size
+		// InitializeForInstance will read it through GetSpecialNodeMemory function, which moves pointer back 
+		FBTInstancedNodeMemory DummyMemory;
+		uint8* RawDummyMemory = ((uint8*)&DummyMemory) + AlignedInstanceMemorySize;
+
+		// create nodes
+		for (int32 Idx = 0; Idx < NumInjectedDecorators; Idx++)
+		{
+			UBTDecorator* DecoratorOb = BehaviorAsset->RootDecorators[Idx];
+
+			const bool bUsesInstancing = DecoratorOb->HasInstance();
+			if (bUsesInstancing)
+			{
+				DecoratorOb->InitializeForInstance(OwnerComp, NodeMemory + MemoryOffsets[Idx], InstancedIndex);
+			}
+			else
+			{
+				DecoratorOb->ForceInstancing(true);
+				DecoratorOb->InitializeForInstance(OwnerComp, RawDummyMemory, InstancedIndex);
+				DecoratorOb->ForceInstancing(false);
+
+				UBTDecorator* InstancedOb = Cast<UBTDecorator>(OwnerComp->NodeInstances.Last());
+				InstancedOb->InitializeMemory(OwnerComp, NodeMemory + MemoryOffsets[Idx]);
+			}
+		}
+
+		// setup their properties
+		uint16 NewExecutionIdx = GetExecutionIndex() - GetInjectedNodesCount();
+		for (int32 Idx = 0; Idx < NumInjectedDecorators; Idx++)
+		{
+			UBTDecorator* InstancedOb = Cast<UBTDecorator>(OwnerComp->NodeInstances[FirstNodeIdx + Idx]);
+			InstancedOb->InitializeNode(GetParentNode(), NewExecutionIdx, GetMemoryOffset() + MemoryOffsets[Idx], GetTreeDepth() - 1);
+			InstancedOb->MarkInjectedNode();
+
+			NewExecutionIdx++;
+		}
+	}
+	else
+	{
+		// initialize memory for injected nodes
+		for (int32 Idx = 0; Idx < NumInjectedDecorators; Idx++)
+		{
+			UBTDecorator* DecoratorOb = BehaviorAsset->RootDecorators[Idx];
+			UBTDecorator* InstancedOb = Cast<UBTDecorator>(OwnerComp->NodeInstances[FirstNodeIdx + Idx]);
+
+			const bool bUsesInstancing = DecoratorOb->HasInstance();
+			if (bUsesInstancing)
+			{
+				InstancedOb->OnInstanceCreated(OwnerComp);
+			}
+			else
+			{
+				uint8* InjectedNodeMemory = NodeMemory + (InstancedOb->GetMemoryOffset() - GetMemoryOffset());
+				InstancedOb->InitializeMemory(OwnerComp, InjectedNodeMemory);
+			}
+		}
+
+		InstancedIndex += NumInjectedDecorators;
+	}
+
+	// inject nodes
+	if (GetParentNode())
+	{
+		const int32 ChildIdx = GetParentNode()->GetChildIndex(this);
+		if (ChildIdx != INDEX_NONE)
+		{
+			FBTCompositeChild& LinkData = GetParentNode()->Children[ChildIdx];
+
+			// check if link already has injected decorators
+			bool bAlreadyInjected = false;
+			for (int32 Idx = 0; Idx < LinkData.Decorators.Num(); Idx++)
+			{
+				if (LinkData.Decorators[Idx] && LinkData.Decorators[Idx]->IsInjected())
+				{
+					bAlreadyInjected = true;
+					break;
+				}
+			}
+
+			// add decorators to link
+			const int32 NumOriginalDecorators = LinkData.Decorators.Num();
+			for (int32 Idx = 0; Idx < NumInjectedDecorators; Idx++)
+			{
+				UBTDecorator* InstancedOb = Cast<UBTDecorator>(OwnerComp->NodeInstances[FirstNodeIdx + Idx]);
+				InstancedOb->InitializeFromAsset(BehaviorAsset);
+				InstancedOb->InitializeDecorator(ChildIdx);
+
+				if (!bAlreadyInjected)
+				{
+					LinkData.Decorators.Add(InstancedOb);
+				}
+			}
+
+			// update composite logic operators
+			if (!bAlreadyInjected && (LinkData.DecoratorOps.Num() || BehaviorAsset->RootDecoratorOps.Num()))
+			{
+				const int32 NumOriginalOps = LinkData.DecoratorOps.Num();
+				if (NumOriginalDecorators > 0)
+				{
+					// and operator for two groups of composites: original and new one
+					FBTDecoratorLogic MasterAndOp(EBTDecoratorLogic::And, LinkData.DecoratorOps.Num() ? 2 : (NumOriginalDecorators + 1));
+					LinkData.DecoratorOps.Insert(MasterAndOp, 0);
+
+					if (NumOriginalOps == 0)
+					{
+						// add Test operations, original link didn't have composite operators
+						for (int32 Idx = 0; Idx < NumOriginalDecorators; Idx++)
+						{
+							FBTDecoratorLogic TestOp(EBTDecoratorLogic::Test, Idx);
+							LinkData.DecoratorOps.Add(TestOp);
+						}
+					}
+				}
+
+				// add injected operators
+				if (BehaviorAsset->RootDecoratorOps.Num() == 0)
+				{
+					FBTDecoratorLogic InjectedAndOp(EBTDecoratorLogic::And, NumInjectedDecorators);
+					LinkData.DecoratorOps.Add(InjectedAndOp);
+
+					for (int32 Idx = 0; Idx < NumInjectedDecorators; Idx++)
+					{
+						FBTDecoratorLogic TestOp(EBTDecoratorLogic::Test, NumOriginalDecorators + Idx);
+						LinkData.DecoratorOps.Add(TestOp);
+					}
+				}
+				else
+				{
+					for (int32 Idx = 0; Idx < BehaviorAsset->RootDecoratorOps.Num(); Idx++)
+					{
+						FBTDecoratorLogic InjectedOp = BehaviorAsset->RootDecoratorOps[Idx];
+						if (InjectedOp.Operation == EBTDecoratorLogic::Test)
+						{
+							InjectedOp.Number += NumOriginalDecorators;
+						}
+
+						LinkData.DecoratorOps.Add(InjectedOp);
+					}
+				}
+			}
+
+#if USE_BEHAVIORTREE_DEBUGGER
+			if (!bAlreadyInjected)
+			{
+				// insert to NextExecutionNode list for debugger
+				UBTNode* NodeIt = GetParentNode();
+				while (NodeIt && NodeIt->GetNextNode() != this)
+				{
+					NodeIt = NodeIt->GetNextNode();
+				}
+
+				if (NodeIt)
+				{
+					NodeIt->InitializeExecutionOrder(OwnerComp->NodeInstances[FirstNodeIdx]);
+					NodeIt = NodeIt->GetNextNode();
+
+					for (int32 Idx = 1; Idx < NumInjectedDecorators; Idx++)
+					{
+						NodeIt->InitializeExecutionOrder(OwnerComp->NodeInstances[FirstNodeIdx + Idx]);
+						NodeIt = NodeIt->GetNextNode();
+					}
+
+					NodeIt->InitializeExecutionOrder((UBTNode*)this);
+				}
+			}
+#endif
+		}
+	}
+}
diff --git a/Engine/Source/Runtime/AIModule/Private/BehaviorTree/Tasks/BTTask_RunEQSQuery.cpp b/Engine/Source/Runtime/AIModule/Private/BehaviorTree/Tasks/BTTask_RunEQSQuery.cpp
new file mode 100644
--- /dev/null
+++ b/Engine/Source/Runtime/AIModule/Private/BehaviorTree/Tasks/BTTask_RunEQSQuery.cpp
@@ -0,0 +1,142 @@
+// Copyright 1998-2014 Epic Games, Inc. All Rights Reserved.
+
+#include "AIModulePrivate.h"
+#include "BehaviorTree/Tasks/BTTask_RunEQSQuery.h"
+#include "EnvironmentQuery/Items/EnvQueryItemType.h"
+#include "EnvironmentQuery/EnvQueryManager.h"
+#include "EnvironmentQuery/EnvQuery.h"
+
+
+UBTTask_RunEQSQuery::UBTTask_RunEQSQuery(const class FPostConstructInitializeProperties& PCIP) : Super(PCIP)
+{
+	NodeName = "Run EQS Query";
+
+	// filter with keys that have matching env query values, only for instances
+	// CDO won't be able to access types from game dlls
+	if (GIsEditor && !HasAnyFlags(RF_ClassDefaultObject))
+	{
+		CollectKeyFilters();
+	}
+}
+
+EBTNodeResult::Type UBTTask_RunEQSQuery::ExecuteTask(class UBehaviorTreeComponent* OwnerComp, uint8* NodeMemory)
+{
+	AActor* QueryOwner = OwnerComp->GetOwner();
+	if (AController* ControllerOwner = Cast<AController>(QueryOwner))
+	{
+		QueryOwner = ControllerOwner->GetPawn();
+	}
+
+	FEnvQueryRequest QueryRequest(QueryTemplate, QueryOwner);
+	QueryRequest.SetNamedParams(QueryParams);
+
+	FBTEnvQueryTaskMemory* MyMemory = (FBTEnvQueryTaskMemory*)NodeMemory;
+	MyMemory->RequestID = QueryRequest.Execute(EEnvQueryRunMode::SingleResult, this, &UBTTask_RunEQSQuery::OnQueryFinished);
+
+	const bool bValid = (MyMemory->RequestID >= 0);
+	if (bValid)
+	{
+		WaitForMessage(OwnerComp, UBrainComponent::AIMessage_QueryFinished, MyMemory->RequestID);
+		return EBTNodeResult::InProgress;
+	}
+
+	return EBTNodeResult::Failed;
+}
+
+EBTNodeResult::Type UBTTask_RunEQSQuery::AbortTask(class UBehaviorTreeComponent* OwnerComp, uint8* NodeMemory)
+{
+	UWorld* MyWorld = OwnerComp->GetWorld();
+	UEnvQueryManager* QueryManager = MyWorld->GetEnvironmentQueryManager();
+	if (QueryManager)
+	{
+		FBTEnvQueryTaskMemory* MyMemory = (FBTEnvQueryTaskMemory*)NodeMemory;
+		QueryManager->AbortQuery(MyMemory->RequestID);
+	}
+
+	return EBTNodeResult::Aborted;
+}
+
+FString UBTTask_RunEQSQuery::GetStaticDescription() const
+{
+	return FString::Printf(TEXT("%s: %s\nBlackboard key: %s"), *Super::GetStaticDescription(),
+		*GetNameSafe(QueryTemplate), *BlackboardKey.SelectedKeyName.ToString());
+}
+
+void UBTTask_RunEQSQuery::DescribeRuntimeValues(const class UBehaviorTreeComponent* OwnerComp, uint8* NodeMemory, EBTDescriptionVerbosity::Type Verbosity, TArray<FString>& Values) const
+{
+	Super::DescribeRuntimeValues(OwnerComp, NodeMemory, Verbosity, Values);
+
+	if (Verbosity == EBTDescriptionVerbosity::Detailed)
+	{
+		FBTEnvQueryTaskMemory* MyMemory = (FBTEnvQueryTaskMemory*)NodeMemory;
+		Values.Add(FString::Printf(TEXT("request: %d"), MyMemory->RequestID));
+	}
+}
+
+uint16 UBTTask_RunEQSQuery::GetInstanceMemorySize() const
+{
+	return sizeof(FBTEnvQueryTaskMemory);
+}
+
+void UBTTask_RunEQSQuery::OnQueryFinished(TSharedPtr<struct FEnvQueryResult> Result)
+{
+	if (Result->Status == EEnvQueryStatus::Aborted)
+	{
+		return;
+	}
+
+	AActor* MyOwner = Cast<AActor>(Result->Owner.Get());
+	if (APawn* PawnOwner = Cast<APawn>(MyOwner))
+	{
+		MyOwner = PawnOwner->GetController();
+	}
+
+	UBehaviorTreeComponent* MyComp = MyOwner ? MyOwner->FindComponentByClass<UBehaviorTreeComponent>() : NULL;
+	if (!MyComp)
+	{
+		UE_LOG(LogBehaviorTree, Warning, TEXT("Unable to find behavior tree to notify about finished query from %s!"), *GetNameSafe(MyOwner));
+		return;
+	}
+
+	bool bSuccess = (Result->Items.Num() >= 1);
+	if (bSuccess)
+	{
+		UBlackboardComponent* MyBlackboard = MyComp->GetBlackboardComponent();
+		UEnvQueryItemType* ItemTypeCDO = (UEnvQueryItemType*)Result->ItemType->GetDefaultObject();
+
+		bSuccess = ItemTypeCDO->StoreInBlackboard(BlackboardKey, MyBlackboard, Result->RawData.GetTypedData() + Result->Items[0].DataOffset);		
+		if (!bSuccess)
+		{
+			UE_VLOG(MyOwner, LogBehaviorTree, Warning, TEXT("Failed to store query result! item:%s key:%s"),
+				*UEnvQueryTypes::GetShortTypeName(Result->ItemType).ToString(),
+				*UBehaviorTreeTypes::GetShortTypeName(BlackboardKey.SelectedKeyType));
+		}
+	}
+
+	FAIMessage::Send(MyComp, FAIMessage(UBrainComponent::AIMessage_QueryFinished, this, Result->QueryID, bSuccess));
+}
+
+void UBTTask_RunEQSQuery::CollectKeyFilters()
+{
+	for (int32 i = 0; i < UEnvQueryManager::RegisteredItemTypes.Num(); i++)
+	{
+		UEnvQueryItemType* ItemTypeCDO = (UEnvQueryItemType*)UEnvQueryManager::RegisteredItemTypes[i]->GetDefaultObject();
+		ItemTypeCDO->AddBlackboardFilters(BlackboardKey, this);
+	}
+}
+
+#if WITH_EDITOR
+void UBTTask_RunEQSQuery::PostEditChangeProperty(struct FPropertyChangedEvent& PropertyChangedEvent)
+{
+	Super::PostEditChangeProperty(PropertyChangedEvent);
+
+	if (PropertyChangedEvent.Property &&
+		PropertyChangedEvent.Property->GetFName() == GET_MEMBER_NAME_CHECKED(UBTTask_RunEQSQuery, QueryTemplate))
+	{
+		if (QueryTemplate)
+		{
+			QueryTemplate->CollectQueryParams(QueryParams);
+		}
+	}
+}
+#endif
diff --git a/Engine/Source/Runtime/AIModule/Private/BehaviorTree/Tasks/BTTask_Wait.cpp b/Engine/Source/Runtime/AIModule/Private/BehaviorTree/Tasks/BTTask_Wait.cpp
new file mode 100644
--- /dev/null
+++ b/Engine/Source/Runtime/AIModule/Private/BehaviorTree/Tasks/BTTask_Wait.cpp
@@ -0,0 +1,52 @@
+// Copyright 1998-2014 Epic Games, Inc. All Rights Reserved.
+
+#include "AIModulePrivate.h"
+#include "BehaviorTree/Tasks/BTTask_Wait.h"
+
+UBTTask_Wait::UBTTask_Wait(const class FPostConstructInitializeProperties& PCIP) : Super(PCIP)
+{
+	NodeName = "Wait";
+	WaitTime = 5.0f;
+	bNotifyTick = true;
+}
+
+EBTNodeResult::Type UBTTask_Wait::ExecuteTask(class UBehaviorTreeComponent* OwnerComp, uint8* NodeMemory)
+{
+	FBTWaitTaskMemory* MyMemory = (FBTWaitTaskMemory*)NodeMemory;
+	MyMemory->RemainingWaitTime = WaitTime;
+	
+	return EBTNodeResult::InProgress;
+}
+
+void UBTTask_Wait::TickTask(class UBehaviorTreeComponent* OwnerComp, uint8* NodeMemory, float DeltaSeconds)
+{
+	FBTWaitTaskMemory* MyMemory = (FBTWaitTaskMemory*)NodeMemory;
+	MyMemory->RemainingWaitTime -= DeltaSeconds;
+
+	if (MyMemory->RemainingWaitTime <= 0.0f)
+	{
+		// continue execution from this node
+		FinishLatentTask(OwnerComp, EBTNodeResult::Succeeded);
+	}
+}
+
+FString UBTTask_Wait::GetStaticDescription() const
+{
+	return FString::Printf(TEXT("%s: %.1fs"), *Super::GetStaticDescription(), WaitTime);
+}
+
+void UBTTask_Wait::DescribeRuntimeValues(const class UBehaviorTreeComponent* OwnerComp, uint8* NodeMemory, EBTDescriptionVerbosity::Type Verbosity, TArray<FString>& Values) const
+{
+	Super::DescribeRuntimeValues(OwnerComp, NodeMemory, Verbosity, Values);
+
+	FBTWaitTaskMemory* MyMemory = (FBTWaitTaskMemory*)NodeMemory;
+	if (MyMemory->RemainingWaitTime)
+	{
+		Values.Add(FString::Printf(TEXT("remaining: %ss"), *FString::SanitizeFloat(MyMemory->RemainingWaitTime)));
+	}
+}
+
+uint16 UBTTask_Wait::GetInstanceMemorySize() const
+{
+	return sizeof(FBTWaitTaskMemory);
+}
diff --git a/Engine/Source/Runtime/AIModule/Private/Blueprint/AIBlueprintHelperLibrary.cpp b/Engine/Source/Runtime/AIModule/Private/Blueprint/AIBlueprintHelperLibrary.cpp
new file mode 100644
--- /dev/null
+++ b/Engine/Source/Runtime/AIModule/Private/Blueprint/AIBlueprintHelperLibrary.cpp
@@ -0,0 +1,207 @@
+// Copyright 1998-2014 Epic Games, Inc. All Rights Reserved.
+
+#include "AIModulePrivate.h"
+#include "Blueprint/AIAsyncTaskBlueprintProxy.h"
+#include "Blueprint/AIBlueprintHelperLibrary.h"
+#include "BehaviorTree/BlackboardComponent.h"
+#include "Animation/AnimInstance.h"
+
+//----------------------------------------------------------------------//
+// UAIAsyncTaskBlueprintProxy
+//----------------------------------------------------------------------//
+
+UAIAsyncTaskBlueprintProxy::UAIAsyncTaskBlueprintProxy(const class FPostConstructInitializeProperties& PCIP)
+	: Super(PCIP)
+{
+}
+
+void UAIAsyncTaskBlueprintProxy::OnMoveCompleted(FAIRequestID RequestID, EPathFollowingResult::Type MovementResult)
+{
+	if (RequestID.IsEquivalent(MoveRequestId) && AIController.IsValid())
+	{
+		if (!AIController->IsPendingKill() && AIController->ReceiveMoveCompleted.IsBound())
+		{
+			AIController->ReceiveMoveCompleted.RemoveDynamic(this, &UAIAsyncTaskBlueprintProxy::OnMoveCompleted);
+		}
+
+		if (MovementResult == EPathFollowingResult::Success)
+		{
+			OnSuccess.Broadcast(MovementResult);
+		}
+		else
+		{
+			OnFail.Broadcast(MovementResult);
+		}
+	}
+}
+
+void UAIAsyncTaskBlueprintProxy::OnNoPath()
+{
+	OnFail.Broadcast(EPathFollowingResult::Aborted);
+}
+
+void UAIAsyncTaskBlueprintProxy::BeginDestroy()
+{
+	if (AIController.IsValid() && !AIController->IsPendingKill() && AIController->ReceiveMoveCompleted.IsBound())
+	{
+		AIController->ReceiveMoveCompleted.RemoveDynamic(this, &UAIAsyncTaskBlueprintProxy::OnMoveCompleted);
+	}
+
+	Super::BeginDestroy();
+}
+
+//----------------------------------------------------------------------//
+// UAIAsyncTaskBlueprintProxy
+//----------------------------------------------------------------------//
+
+UAIBlueprintHelperLibrary::UAIBlueprintHelperLibrary(const class FPostConstructInitializeProperties& PCIP)
+	: Super(PCIP)
+{
+}
+
+class UAIAsyncTaskBlueprintProxy* UAIBlueprintHelperLibrary::CreateMoveToProxyObject(class UObject* WorldContextObject, APawn* Pawn, FVector Destination, AActor* TargetActor, float AcceptanceRadius, bool bStopOnOverlap)
+{
+	check(WorldContextObject);
+	if (!Pawn)
+	{
+		return NULL;
+	}
+	UAIAsyncTaskBlueprintProxy* MyObj = NULL;
+	AAIController* AIController = Cast<AAIController>(Pawn->GetController());
+	if (AIController)
+	{
+		MyObj = NewObject<UAIAsyncTaskBlueprintProxy>();
+		FNavPathSharedPtr Path = TargetActor ? AIController->FindPath(TargetActor, true) : AIController->FindPath(Destination, true);
+		if (Path.IsValid())
+		{
+			MyObj->AIController = AIController;
+			MyObj->AIController->ReceiveMoveCompleted.AddDynamic(MyObj, &UAIAsyncTaskBlueprintProxy::OnMoveCompleted);
+			MyObj->MoveRequestId = MyObj->AIController->RequestMove(Path, TargetActor, AcceptanceRadius, bStopOnOverlap);
+		}
+		else
+		{
+			UWorld* World = GEngine->GetWorldFromContextObject( WorldContextObject );
+			World->GetTimerManager().SetTimer(MyObj, &UAIAsyncTaskBlueprintProxy::OnNoPath, 0.1, false);
+		}
+	}
+	return MyObj;
+}
+
+void UAIBlueprintHelperLibrary::SendAIMessage(APawn* Target, FName Message, UObject* MessageSource, bool bSuccess)
+{
+	FAIMessage::Send(Target, FAIMessage(Message, MessageSource, bSuccess));
+}
+
+APawn* UAIBlueprintHelperLibrary::SpawnAIFromClass(class UObject* WorldContextObject, TSubclassOf<APawn> PawnClass, class UBehaviorTree* BehaviorTree, FVector Location, FRotator Rotation, bool bNoCollisionFail)
+{
+	APawn* NewPawn = NULL;
+
+	UWorld* World = GEngine->GetWorldFromContextObject(WorldContextObject);
+	if (World && *PawnClass)
+	{
+		FActorSpawnParameters ActorSpawnParams;
+		ActorSpawnParams.bNoCollisionFail = bNoCollisionFail;
+		NewPawn = World->SpawnActor<APawn>(*PawnClass, Location, Rotation, ActorSpawnParams);
+
+		if (NewPawn != NULL)
+		{
+			if (NewPawn->Controller == NULL)
+			{	// NOTE: SpawnDefaultController ALSO calls Possess() to possess the pawn (if a controller is successfully spawned).
+				NewPawn->SpawnDefaultController();
+			}
+
+			if (BehaviorTree != NULL)
+			{
+				AAIController* AIController = Cast<AAIController>(NewPawn->Controller);
+
+				if (AIController != NULL)
+				{
+					AIController->RunBehaviorTree(BehaviorTree);
+				}
+			}
+		}
+	}
+
+	return NewPawn;
+}
+
+APawn* UAIBlueprintHelperLibrary::SpawnAI(class UObject* WorldContextObject, UBlueprint* Pawn, class UBehaviorTree* BehaviorTree, FVector Location, FRotator Rotation, bool bNoCollisionFail)
+{
+	APawn* NewPawn = NULL;
+
+	const bool bGoodBPGeneratedClass = Pawn && Pawn->GeneratedClass && Pawn->GeneratedClass->IsChildOf(APawn::StaticClass());
+	if (bGoodBPGeneratedClass)
+	{
+		NewPawn = SpawnAIFromClass(WorldContextObject, *(Pawn->GeneratedClass), BehaviorTree, Location, Rotation, bNoCollisionFail);
+	}
+
+	return NewPawn;
+}
+
+UBlackboardComponent* UAIBlueprintHelperLibrary::GetBlackboard(AActor* Target)
+{
+	UBlackboardComponent* BlackboardComp = NULL;
+
+	APawn* TargetPawn = Cast<APawn>(Target);
+	if (TargetPawn && TargetPawn->GetController())
+	{
+		BlackboardComp = TargetPawn->GetController()->FindComponentByClass<UBlackboardComponent>();
+	}
+
+	if (BlackboardComp == NULL && Target)
+	{
+		BlackboardComp = Target->FindComponentByClass<UBlackboardComponent>();
+	}
+
+	return BlackboardComp;
+}
+
+void UAIBlueprintHelperLibrary::LockAIResourcesWithAnimation(class UAnimInstance* AnimInstance, bool bLockMovement, bool LockAILogic)
+{
+	if (AnimInstance == NULL)
+	{
+		return;
+	}
+
+	APawn* PawnOwner = AnimInstance->TryGetPawnOwner();
+	if (PawnOwner)
+	{
+		AAIController* OwningAI = Cast<AAIController>(PawnOwner->Controller);
+		if (OwningAI)
+		{
+			if (bLockMovement && OwningAI->PathFollowingComponent)
+			{
+				OwningAI->PathFollowingComponent->LockResource(EAILockSource::Animation);
+			}
+			if (LockAILogic && OwningAI->BrainComponent)
+			{
+				OwningAI->BrainComponent->LockResource(EAILockSource::Animation);
+			}
+		}
+	}
+}
+
+void UAIBlueprintHelperLibrary::UnlockAIResourcesWithAnimation(class UAnimInstance* AnimInstance, bool bUnlockMovement, bool UnlockAILogic)
+{
+	if (AnimInstance == NULL)
+	{
+		return;
+	}
+
+	APawn* PawnOwner = AnimInstance->TryGetPawnOwner();
+	if (PawnOwner)
+	{
+		AAIController* OwningAI = Cast<AAIController>(PawnOwner->Controller);
+		if (OwningAI)
+		{
+			if (bUnlockMovement && OwningAI->PathFollowingComponent)
+			{
+				OwningAI->PathFollowingComponent->ClearResourceLock(EAILockSource::Animation);
+			}
+			if (UnlockAILogic && OwningAI->BrainComponent)
+			{
+				OwningAI->BrainComponent->ClearResourceLock(EAILockSource::Animation);
+			}
+		}
+	}
+}
\ No newline at end of file
diff --git a/Engine/Source/Runtime/AIModule/Private/EnvironmentQuery/Contexts/EnvQueryContext_Item.cpp b/Engine/Source/Runtime/AIModule/Private/EnvironmentQuery/Contexts/EnvQueryContext_Item.cpp
new file mode 100644
--- /dev/null
+++ b/Engine/Source/Runtime/AIModule/Private/EnvironmentQuery/Contexts/EnvQueryContext_Item.cpp
@@ -0,0 +1,9 @@
+// Copyright 1998-2014 Epic Games, Inc. All Rights Reserved.
+
+#include "AIModulePrivate.h"
+#include "EnvironmentQuery/Contexts/EnvQueryContext_Item.h"
+
+UEnvQueryContext_Item::UEnvQueryContext_Item(const class FPostConstructInitializeProperties& PCIP) : Super(PCIP)
+{
+	// special case, don't set delegate
+}
diff --git a/Engine/Source/Runtime/AIModule/Private/EnvironmentQuery/Contexts/EnvQueryContext_Querier.cpp b/Engine/Source/Runtime/AIModule/Private/EnvironmentQuery/Contexts/EnvQueryContext_Querier.cpp
new file mode 100644
--- /dev/null
+++ b/Engine/Source/Runtime/AIModule/Private/EnvironmentQuery/Contexts/EnvQueryContext_Querier.cpp
@@ -0,0 +1,15 @@
+// Copyright 1998-2014 Epic Games, Inc. All Rights Reserved.
+
+#include "AIModulePrivate.h"
+#include "EnvironmentQuery/Contexts/EnvQueryContext_Querier.h"
+#include "EnvironmentQuery/Items/EnvQueryItemType_Actor.h"
+
+UEnvQueryContext_Querier::UEnvQueryContext_Querier(const class FPostConstructInitializeProperties& PCIP) : Super(PCIP)
+{
+}
+
+void UEnvQueryContext_Querier::ProvideContext(struct FEnvQueryInstance& QueryInstance, struct FEnvQueryContextData& ContextData) const
+{
+	AActor* QueryOwner = Cast<AActor>(QueryInstance.Owner.Get());
+	UEnvQueryItemType_Actor::SetContextHelper(ContextData, QueryOwner);
+}
diff --git a/Engine/Source/Runtime/AIModule/Private/EnvironmentQuery/EQSRenderingComponent.cpp b/Engine/Source/Runtime/AIModule/Private/EnvironmentQuery/EQSRenderingComponent.cpp
new file mode 100644
--- /dev/null
+++ b/Engine/Source/Runtime/AIModule/Private/EnvironmentQuery/EQSRenderingComponent.cpp
@@ -0,0 +1,280 @@
+// Copyright 1998-2014 Epic Games, Inc. All Rights Reserved.
+
+#include "AIModulePrivate.h"
+#include "DebugRenderSceneProxy.h"
+#include "EnvironmentQuery/EQSRenderingComponent.h"
+#include "EnvironmentQuery/EQSQueryResultSourceInterface.h"
+#include "EnvironmentQuery/Items/EnvQueryItemType_VectorBase.h"
+
+static const int32 EQSMaxItemsDrawn = 10000;
+
+namespace FEQSRenderingHelper
+{
+	FVector ExtractLocation(TSubclassOf<class UEnvQueryItemType> ItemType, TArray<uint8> RawData, const TArray<FEnvQueryItem>& Items, int32 Index)
+	{
+		if (Items.IsValidIndex(Index) &&
+			ItemType->IsChildOf(UEnvQueryItemType_VectorBase::StaticClass()))
+		{
+			UEnvQueryItemType_VectorBase* DefTypeOb = (UEnvQueryItemType_VectorBase*)ItemType->GetDefaultObject();
+			return DefTypeOb->GetLocation(RawData.GetTypedData() + Items[Index].DataOffset);
+		}
+		return FVector::ZeroVector;
+	}
+}
+
+//----------------------------------------------------------------------//
+// FEQSSceneProxy
+//----------------------------------------------------------------------//
+const FVector FEQSSceneProxy::ItemDrawRadius(30,30,30);
+
+FEQSSceneProxy::FEQSSceneProxy(const UPrimitiveComponent* InComponent, const FString& InViewFlagName, bool bInDrawOnlyWhenSelected) 
+	: FDebugRenderSceneProxy(InComponent)
+	, ActorOwner(NULL)
+	, QueryDataSource(NULL)
+	, bDrawOnlyWhenSelected(bInDrawOnlyWhenSelected)
+{
+	TextWithoutShadowDistance = 1500;
+	ViewFlagIndex = uint32(FEngineShowFlags::FindIndexByName(*InViewFlagName));
+	ViewFlagName = InViewFlagName;
+
+	if (InComponent == NULL)
+	{
+		return;
+	}
+
+	ActorOwner = InComponent ? InComponent->GetOwner() : NULL;
+	QueryDataSource = InterfaceCast<const IEQSQueryResultSourceInterface>(ActorOwner);
+	if (QueryDataSource == NULL)
+	{
+		QueryDataSource = InterfaceCast<const IEQSQueryResultSourceInterface>(InComponent);
+		if (QueryDataSource == NULL)
+		{
+			return;
+		}
+	}
+
+	FEQSSceneProxy::CollectEQSData(InComponent, QueryDataSource, SolidSpheres, Texts);
+}
+
+FEQSSceneProxy::FEQSSceneProxy(const UPrimitiveComponent* InComponent, const FString& InViewFlagName, bool bInDrawOnlyWhenSelected, const TArray<FSphere>& Spheres, const TArray<FText3d>& InTexts)
+	: FDebugRenderSceneProxy(InComponent)
+	, ActorOwner(NULL)
+	, QueryDataSource(NULL)
+	, bDrawOnlyWhenSelected(bInDrawOnlyWhenSelected)
+{
+	TextWithoutShadowDistance = 1500;
+	ViewFlagIndex = uint32(FEngineShowFlags::FindIndexByName(*InViewFlagName));
+	ViewFlagName = InViewFlagName;
+
+	if (InComponent == NULL)
+	{
+		return;
+	}
+
+	ActorOwner = InComponent ? InComponent->GetOwner() : NULL;
+	QueryDataSource = InterfaceCast<const IEQSQueryResultSourceInterface>(ActorOwner);
+	if (QueryDataSource == NULL)
+	{
+		QueryDataSource = InterfaceCast<const IEQSQueryResultSourceInterface>(InComponent);
+		if (QueryDataSource == NULL)
+		{
+			return;
+		}
+	}
+
+	SolidSpheres = Spheres;
+	Texts = InTexts;
+}
+
+void FEQSSceneProxy::CollectEQSData(const UPrimitiveComponent* InComponent, const class IEQSQueryResultSourceInterface* InQueryDataSource, TArray<FSphere>& Spheres, TArray<FText3d>& Texts)
+{
+	AActor* ActorOwner = InComponent ? InComponent->GetOwner() : NULL;
+	IEQSQueryResultSourceInterface* QueryDataSource = const_cast<IEQSQueryResultSourceInterface*>(InQueryDataSource);
+	if (QueryDataSource == NULL)
+	{
+		QueryDataSource = InterfaceCast<IEQSQueryResultSourceInterface>(ActorOwner);
+		if (QueryDataSource == NULL)
+		{
+			QueryDataSource = InterfaceCast<IEQSQueryResultSourceInterface>(const_cast<UPrimitiveComponent*>(InComponent));
+			if (QueryDataSource == NULL)
+			{
+				return;
+			}
+		}
+	}
+
+	const FEnvQueryResult* ResultItems = QueryDataSource->GetQueryResult();
+	const FEnvQueryInstance* QueryInstance = QueryDataSource->GetQueryInstance();
+
+	if (ResultItems == NULL)
+	{
+		if (QueryInstance == NULL)
+		{
+			return;
+		}
+		else
+		{
+			ResultItems = QueryInstance;
+		}
+	}
+
+#if USE_EQS_DEBUGGER
+	// using "mid-results" requires manual normalization
+	const bool bUseMidResults = QueryInstance && (QueryInstance->Items.Num() < QueryInstance->DebugData.DebugItems.Num());
+	const TArray<FEnvQueryItem>& Items = bUseMidResults ? QueryInstance->DebugData.DebugItems : ResultItems->Items;
+	const int32 ItemCountLimit = FMath::Clamp(Items.Num(), 0, EQSMaxItemsDrawn);
+	const bool bNoTestsPerformed = QueryInstance != NULL && QueryInstance->CurrentTest <= 0;
+
+	//if (!bUseMidResults && Items.Num() > 0 && Items[0].IsValid())
+	//{
+	//	const FVector Loc = FEQSRenderingHelper::ExtractLocation(ResultItems->ItemType, ResultItems->RawData, Items, 0);
+	//	Cylinders.Add(FWireCylinder(Loc, ItemDrawRadius.X, 100, FColor::White));
+	//}
+
+	float MinScore = 0.f;
+	float MaxScore = -BIG_NUMBER;
+	if (bUseMidResults)
+	{
+		const FEnvQueryItem* ItemInfo = Items.GetTypedData();
+		for (int32 ItemIndex = 0; ItemIndex < Items.Num(); ItemIndex++, ItemInfo++)
+		{
+			if (ItemInfo->IsValid())
+			{
+				MinScore = FMath::Min(MinScore, ItemInfo->Score);
+				MaxScore = FMath::Max(MaxScore, ItemInfo->Score);
+			}
+		}
+	}
+	const float ScoreNormalizer = bUseMidResults && (MaxScore != MinScore) ? 1.f / (MaxScore - MinScore) : 1.f;
+
+	for (int32 ItemIndex = 0; ItemIndex < ItemCountLimit; ++ItemIndex)
+	{
+		if (Items[ItemIndex].IsValid())
+		{
+			const float Score = bNoTestsPerformed ? 1 : Items[ItemIndex].Score * ScoreNormalizer;
+			const FVector Loc = FEQSRenderingHelper::ExtractLocation(ResultItems->ItemType, ResultItems->RawData, Items, ItemIndex);
+			Spheres.Add(FSphere(ItemDrawRadius.X, Loc, !bNoTestsPerformed ? FLinearColor(FColor::MakeRedToGreenColorFromScalar(Score)) : FLinearColor(0.2, 1.0, 1.0, 1)));
+
+			const FString Label = bNoTestsPerformed ? TEXT("") : FString::Printf(TEXT("%.2f"), Score);
+			Texts.Add(FText3d(Label, Loc, FLinearColor::White));
+		}
+	}
+
+	if (QueryDataSource->GetShouldDrawFailedItems() && QueryInstance)
+	{
+		const FEQSQueryDebugData& InstanceDebugData = QueryInstance->DebugData;
+		const TArray<FEnvQueryItem>& DebugQueryItems = InstanceDebugData.DebugItems;
+		const TArray<FEnvQueryItemDetails>& Details = InstanceDebugData.DebugItemDetails;
+
+		const int32 DebugItemCountLimit = DebugQueryItems.Num() == Details.Num() ? FMath::Clamp(DebugQueryItems.Num(), 0, EQSMaxItemsDrawn) : 0;
+
+		for (int32 ItemIndex = 0; ItemIndex < DebugItemCountLimit; ++ItemIndex)
+		{
+			if (DebugQueryItems[ItemIndex].IsValid())
+			{
+				continue;
+			}
+
+			const float Score = bNoTestsPerformed ? 1 : Items[ItemIndex].Score * ScoreNormalizer;
+			const FVector Loc = FEQSRenderingHelper::ExtractLocation(QueryInstance->ItemType, InstanceDebugData.RawData, DebugQueryItems, ItemIndex);
+			Spheres.Add(FSphere(ItemDrawRadius.X, Loc, FLinearColor(0.0, 0.0, 0.6, 0.6)));
+
+			const FString Label = InstanceDebugData.PerformedTestNames[Details[ItemIndex].FailedTestIndex];
+			Texts.Add(FText3d(Label, Loc, FLinearColor::White));
+		}
+	}
+#endif // USE_EQS_DEBUGGER
+
+	if (ActorOwner && Spheres.Num() > EQSMaxItemsDrawn)
+	{
+		UE_VLOG(ActorOwner, LogEQS, Warning, TEXT("EQS drawing: too much items to draw! Drawing first %d from set of %d")
+			, EQSMaxItemsDrawn
+			, Spheres.Num()
+			);
+	}
+}
+
+void FEQSSceneProxy::DrawDebugLabels(UCanvas* Canvas, APlayerController* PC)
+{
+	if ( !ActorOwner
+		 || (ActorOwner->IsSelected() == false && bDrawOnlyWhenSelected == true)
+		 || QueryDataSource->GetShouldDebugDrawLabels() == false)
+	{
+		return;
+	}
+
+	FDebugRenderSceneProxy::DrawDebugLabels(Canvas, PC);
+}
+
+bool FEQSSceneProxy::SafeIsActorSelected() const
+{
+	if(ActorOwner)
+	{
+		return ActorOwner->IsSelected();
+	}
+
+	return false;
+}
+
+FPrimitiveViewRelevance FEQSSceneProxy::GetViewRelevance(const FSceneView* View)
+{
+	FPrimitiveViewRelevance Result;
+	Result.bDrawRelevance = View->Family->EngineShowFlags.GetSingleFlag(ViewFlagIndex) && IsShown(View) 
+		&& (!bDrawOnlyWhenSelected || SafeIsActorSelected());
+	Result.bDynamicRelevance = true;
+	Result.bNormalTranslucencyRelevance = IsShown(View);
+	return Result;
+}
+
+//----------------------------------------------------------------------//
+// UEQSRenderingComponent
+//----------------------------------------------------------------------//
+UEQSRenderingComponent::UEQSRenderingComponent(const class FPostConstructInitializeProperties& PCIP) 
+	: Super(PCIP)
+	, DrawFlagName("GameplayDebug")
+	, bDrawOnlyWhenSelected(true)
+{
+}
+
+FPrimitiveSceneProxy* UEQSRenderingComponent::CreateSceneProxy()
+{
+	return new FEQSSceneProxy(this, DrawFlagName, bDrawOnlyWhenSelected);
+}
+
+FBoxSphereBounds UEQSRenderingComponent::CalcBounds(const FTransform & LocalToWorld) const
+{
+	static FBox BoundingBox(FVector(-HALF_WORLD_MAX), FVector(HALF_WORLD_MAX));
+	return FBoxSphereBounds(BoundingBox);
+}
+
+void UEQSRenderingComponent::CreateRenderState_Concurrent()
+{
+	Super::CreateRenderState_Concurrent();
+
+#if USE_EQS_DEBUGGER
+	if (SceneProxy)
+	{
+		static_cast<FEQSSceneProxy*>(SceneProxy)->RegisterDebugDrawDelgate();
+	}
+#endif
+}
+
+void UEQSRenderingComponent::DestroyRenderState_Concurrent()
+{
+#if USE_EQS_DEBUGGER
+	if (SceneProxy)
+	{
+		static_cast<FEQSSceneProxy*>(SceneProxy)->UnregisterDebugDrawDelgate();
+	}
+#endif
+
+	Super::DestroyRenderState_Concurrent();
+}
+
+//----------------------------------------------------------------------//
+// UEQSQueryResultSourceInterface
+//----------------------------------------------------------------------//
+UEQSQueryResultSourceInterface::UEQSQueryResultSourceInterface(const class FPostConstructInitializeProperties& PCIP)
+	: Super(PCIP)
+{
+}
diff --git a/Engine/Source/Runtime/AIModule/Private/EnvironmentQuery/EQSTestingPawn.cpp b/Engine/Source/Runtime/AIModule/Private/EnvironmentQuery/EQSTestingPawn.cpp
new file mode 100644
--- /dev/null
+++ b/Engine/Source/Runtime/AIModule/Private/EnvironmentQuery/EQSTestingPawn.cpp
@@ -0,0 +1,272 @@
+// Copyright 1998-2014 Epic Games, Inc. All Rights Reserved.
+
+#include "AIModulePrivate.h"
+#include "EnvironmentQuery/EnvQuery.h"
+#include "EnvironmentQuery/EQSTestingPawn.h"
+#include "EnvironmentQuery/EnvQueryTypes.h"
+#include "EnvironmentQuery/EnvQueryManager.h"
+#include "EnvironmentQuery/EQSRenderingComponent.h"
+
+#if WITH_EDITOR
+#include "UnrealEd.h"
+#endif
+
+//----------------------------------------------------------------------//
+// AEQSTestingPawn
+//----------------------------------------------------------------------//
+AEQSTestingPawn::AEQSTestingPawn(const class FPostConstructInitializeProperties& PCIP)
+	: Super(PCIP)
+	, TimeLimitPerStep(-1)
+	, StepToDebugDraw(0)
+	, bDrawLabels(true)
+	, bDrawFailedItems(true)
+	, bReRunQueryOnlyOnFinishedMove(true)
+{
+	static FName CollisionProfileName(TEXT("NoCollision"));
+	CapsuleComponent->SetCollisionProfileName(CollisionProfileName);
+
+#if WITH_EDITORONLY_DATA
+	UArrowComponent* ArrowComponent = FindComponentByClass<UArrowComponent>();
+	if (ArrowComponent != NULL)
+	{
+		ArrowComponent->SetRelativeScale3D(FVector(2,2,2));
+		ArrowComponent->bIsScreenSizeScaled = true;
+	}
+
+	TSubobjectPtr<UBillboardComponent> SpriteComponent = PCIP.CreateEditorOnlyDefaultSubobject<UBillboardComponent>(this, TEXT("Sprite"));
+	if (!IsRunningCommandlet() && (SpriteComponent != nullptr))
+	{
+		struct FConstructorStatics
+		{
+			ConstructorHelpers::FObjectFinderOptional<UTexture2D> TextureObject;
+			FName ID_Misc;
+			FText NAME_Misc;
+			FConstructorStatics()
+				: TextureObject(TEXT("/Engine/EditorResources/S_Pawn"))
+				, ID_Misc(TEXT("Misc"))
+				, NAME_Misc(NSLOCTEXT("SpriteCategory", "Misc", "Misc"))
+			{
+			}
+		};
+		static FConstructorStatics ConstructorStatics;
+
+		SpriteComponent->Sprite = ConstructorStatics.TextureObject.Get();
+		SpriteComponent->RelativeScale3D = FVector(0.5f, 0.5f, 0.5f);
+		SpriteComponent->bHiddenInGame = true;
+		//SpriteComponent->Mobility = EComponentMobility::Static;
+		SpriteComponent->SpriteInfo.Category = ConstructorStatics.ID_Misc;
+		SpriteComponent->SpriteInfo.DisplayName = ConstructorStatics.NAME_Misc;
+		SpriteComponent->AttachParent = RootComponent;
+		SpriteComponent->bIsScreenSizeScaled = true;
+	}
+
+	EdRenderComp = PCIP.CreateEditorOnlyDefaultSubobject<UEQSRenderingComponent>(this, TEXT("EQSRender"));
+#endif // WITH_EDITORONLY_DATA
+
+	// Default to no tick function, but if we set 'never ticks' to false (so there is a tick function) it is enabled by default
+	PrimaryActorTick.bCanEverTick = true;
+	PrimaryActorTick.bStartWithTickEnabled = true;
+
+#if WITH_EDITOR
+	if (HasAnyFlags(RF_ClassDefaultObject) && GetClass() == StaticClass())
+	{
+		USelection::SelectObjectEvent.AddStatic(&AEQSTestingPawn::OnEditorSelectionChanged);
+	}
+#endif // WITH_EDITOR
+}
+
+void AEQSTestingPawn::OnEditorSelectionChanged(UObject* NewSelection)
+{
+	bool bEQSPawnSelected = Cast<AEQSTestingPawn>(NewSelection) != NULL;
+	if (bEQSPawnSelected == false)
+	{
+		USelection* Selection = Cast<USelection>(NewSelection);
+		if (Selection != NULL)
+		{
+			TArray<AEQSTestingPawn*> SelectedPawns;
+			Selection->GetSelectedObjects<AEQSTestingPawn>(SelectedPawns);
+			bEQSPawnSelected = SelectedPawns.Num() > 0;
+		}
+	}
+
+#if WITH_EDITOR
+	if (GCurrentLevelEditingViewportClient != NULL && UGameplayDebuggingComponent::ShowFlagIndex != INDEX_NONE)
+	{
+		GCurrentLevelEditingViewportClient->EngineShowFlags.SetSingleFlag(UGameplayDebuggingComponent::ShowFlagIndex, bEQSPawnSelected);
+	}
+#endif // WITH_EDITOR
+}
+
+void AEQSTestingPawn::TickActor(float DeltaTime, enum ELevelTick TickType, FActorTickFunction& ThisTickFunction)
+{
+	Super::Tick(DeltaTime);
+
+	if (QueryInstance.IsValid() && QueryInstance->Status == EEnvQueryStatus::Processing)
+	{
+		MakeOneStep();
+	}
+}
+
+void AEQSTestingPawn::PostLoad() 
+{
+	Super::PostLoad();
+
+	UBillboardComponent* SpriteComponent = FindComponentByClass<UBillboardComponent>();
+	if (SpriteComponent != NULL)
+	{
+		SpriteComponent->bHiddenInGame = !bShouldBeVisibleInGame;
+	}
+
+	UEnvQueryManager* EQS = GetWorld() ? GetWorld()->GetEnvironmentQueryManager() : NULL;
+	if (EQS)
+	{
+		RunEQSQuery();
+	}
+	else
+	{
+		QueryTemplate = NULL;
+		QueryParams.Reset();
+	}
+}
+
+void AEQSTestingPawn::RunEQSQuery()
+{
+	Reset();
+	
+	// make one step if TimeLimitPerStep > 0.f, else all steps
+	do
+	{
+		MakeOneStep();
+	} while (TimeLimitPerStep <= 0.f && QueryInstance.IsValid() == true && QueryInstance->Status == EEnvQueryStatus::Processing);	
+}
+
+void AEQSTestingPawn::Reset()
+{
+	QueryInstance = NULL;
+	StepToDebugDraw = 0;
+	StepResults.Reset();
+
+	if (Controller == NULL)
+	{
+		SpawnDefaultController();
+	}
+}
+
+void AEQSTestingPawn::MakeOneStep()
+{
+	UEnvQueryManager* EQS = GetWorld()->GetEnvironmentQueryManager();
+	if (EQS == NULL)
+	{
+		return;
+	}
+
+	if (QueryInstance.IsValid() == false && QueryTemplate != NULL)
+	{
+		FEnvQueryRequest QueryRequest(QueryTemplate, this);
+		QueryRequest.SetNamedParams(QueryParams);
+		QueryInstance = EQS->PrepareQueryInstance(QueryRequest, EEnvQueryRunMode::AllMatching);
+	}
+
+	// possible still not valid 
+	if (QueryInstance.IsValid() == true && QueryInstance->Status == EEnvQueryStatus::Processing)
+	{
+		QueryInstance->ExecuteOneStep(double(TimeLimitPerStep));
+		StepResults.Add(*(QueryInstance.Get()));
+
+		if (QueryInstance->Status != EEnvQueryStatus::Processing)
+		{
+			StepToDebugDraw = StepResults.Num() - 1;
+			UpdateDrawing();
+		}
+	}
+}
+
+void AEQSTestingPawn::UpdateDrawing()
+{
+#if WITH_EDITORONLY_DATA
+
+	UBillboardComponent* SpriteComponent = FindComponentByClass<UBillboardComponent>();
+	if (SpriteComponent != NULL)
+	{
+		SpriteComponent->MarkRenderStateDirty();
+	}	
+
+	if (EdRenderComp != NULL && EdRenderComp->bVisible)
+	{
+		EdRenderComp->MarkRenderStateDirty();
+
+#if WITH_EDITOR
+		if (GEditor != NULL)
+		{
+			GEditor->RedrawLevelEditingViewports();
+		}
+#endif // WITH_EDITOR
+	}
+#endif // WITH_EDITORONLY_DATA
+}
+
+const FEnvQueryResult* AEQSTestingPawn::GetQueryResult() const 
+{
+	return StepResults.Num() > 0 ? &StepResults[StepToDebugDraw] : NULL;
+}
+
+const FEnvQueryInstance* AEQSTestingPawn::GetQueryInstance() const 
+{
+	return StepResults.Num() > 0 ? &StepResults[StepToDebugDraw] : NULL;
+}
+
+#if WITH_EDITOR
+
+void AEQSTestingPawn::PostEditChangeProperty( struct FPropertyChangedEvent& PropertyChangedEvent)
+{
+	static const FName NAME_QueryTemplate = GET_MEMBER_NAME_CHECKED(AEQSTestingPawn, QueryTemplate);
+	static const FName NAME_StepToDebugDraw = GET_MEMBER_NAME_CHECKED(AEQSTestingPawn, StepToDebugDraw);
+	static const FName NAME_QueryParams = GET_MEMBER_NAME_CHECKED(AEQSTestingPawn, QueryParams);
+	static const FName NAME_DrawFailedItems = GET_MEMBER_NAME_CHECKED(AEQSTestingPawn, bDrawFailedItems);
+	static const FName NAME_ShouldBeVisibleInGame = GET_MEMBER_NAME_CHECKED(AEQSTestingPawn, bShouldBeVisibleInGame);
+
+	if (PropertyChangedEvent.Property != NULL)
+	{
+		const FName PropName = PropertyChangedEvent.MemberProperty->GetFName();
+		if (PropName == NAME_QueryTemplate || PropName == NAME_QueryParams)
+		{
+			if (QueryTemplate)
+			{
+				QueryTemplate->CollectQueryParams(QueryParams);
+			}
+
+			RunEQSQuery();
+		}
+		else if (PropName == NAME_StepToDebugDraw)
+		{
+			StepToDebugDraw  = FMath::Clamp(StepToDebugDraw, 0, StepResults.Num() - 1 );
+			UpdateDrawing();
+		}
+		else if (PropName == NAME_DrawFailedItems)
+		{
+			UpdateDrawing();
+		}
+		else if (PropName == NAME_ShouldBeVisibleInGame)
+		{
+			UBillboardComponent* SpriteComponent = FindComponentByClass<UBillboardComponent>();
+			if (SpriteComponent != NULL)
+			{
+				SpriteComponent->bHiddenInGame = !bShouldBeVisibleInGame;
+			}
+		}
+	}
+
+	Super::PostEditChangeProperty(PropertyChangedEvent);
+}
+
+void AEQSTestingPawn::PostEditMove(bool bFinished)
+{
+	Super::PostEditMove(bFinished);
+
+	if (bFinished || !bReRunQueryOnlyOnFinishedMove)
+	{
+		RunEQSQuery();
+	}
+}
+
+#endif // WITH_EDITOR
\ No newline at end of file
diff --git a/Engine/Source/Runtime/AIModule/Private/EnvironmentQuery/EnvQuery.cpp b/Engine/Source/Runtime/AIModule/Private/EnvironmentQuery/EnvQuery.cpp
new file mode 100644
--- /dev/null
+++ b/Engine/Source/Runtime/AIModule/Private/EnvironmentQuery/EnvQuery.cpp
@@ -0,0 +1,106 @@
+// Copyright 1998-2014 Epic Games, Inc. All Rights Reserved.
+
+#include "AIModulePrivate.h"
+#include "EnvironmentQuery/EnvQuery.h"
+#include "EnvironmentQuery/EnvQueryGenerator.h"
+#include "EnvironmentQuery/EnvQueryOption.h"
+#include "EnvironmentQuery/EnvQueryTest.h"
+
+UEnvQuery::UEnvQuery(const class FPostConstructInitializeProperties& PCIP) : Super(PCIP)
+{
+}
+
+static bool HasNamedValue(const FName& ParamName, const TArray<FEnvNamedValue>& NamedValues)
+{
+	for (int32 i = 0; i < NamedValues.Num(); i++)
+	{
+		if (NamedValues[i].ParamName == ParamName)
+		{
+			return true;
+		}
+	}
+
+	return false;
+}
+
+static void AddNamedValue(const FName& ParamName, const EEnvQueryParam::Type& ParamType, float Value,
+						  TArray<FEnvNamedValue>& NamedValues, TArray<FName>& RequiredParams)
+{
+	if (ParamName != NAME_None)
+	{
+		if (!HasNamedValue(ParamName, NamedValues))
+		{
+			FEnvNamedValue NewValue;
+			NewValue.ParamName = ParamName;
+			NewValue.ParamType = ParamType;
+			NewValue.Value = Value;
+			NamedValues.Add(NewValue);
+		}
+
+		RequiredParams.AddUnique(ParamName);
+	}
+}
+
+#define GET_STRUCT_NAME_CHECKED(StructName) \
+	((void)sizeof(StructName), TEXT(#StructName))
+
+static void AddNamedValuesFromObject(const UObject* Ob, TArray<FEnvNamedValue>& NamedValues, TArray<FName>& RequiredParams)
+{
+	if (Ob == NULL)
+		return;
+
+	for (UProperty* TestProperty = Ob->GetClass()->PropertyLink; TestProperty; TestProperty = TestProperty->PropertyLinkNext)
+	{
+		UStructProperty* TestStruct = Cast<UStructProperty>(TestProperty);
+		if (TestStruct == NULL)
+		{
+			continue;
+		}
+
+		FString TypeDesc = TestStruct->GetCPPType(NULL, CPPF_None);
+		if (TypeDesc.Contains(GET_STRUCT_NAME_CHECKED(FEnvIntParam)))
+		{
+			const FEnvIntParam* PropertyValue = TestStruct->ContainerPtrToValuePtr<FEnvIntParam>(Ob);
+			AddNamedValue(PropertyValue->ParamName, EEnvQueryParam::Int, *((float*)&PropertyValue->Value), NamedValues, RequiredParams);
+		}
+		else if (TypeDesc.Contains(GET_STRUCT_NAME_CHECKED(FEnvFloatParam)))
+		{
+			const FEnvFloatParam* PropertyValue = TestStruct->ContainerPtrToValuePtr<FEnvFloatParam>(Ob);
+			AddNamedValue(PropertyValue->ParamName, EEnvQueryParam::Float, PropertyValue->Value, NamedValues, RequiredParams);
+		}
+		else if (TypeDesc.Contains(GET_STRUCT_NAME_CHECKED(FEnvBoolParam)))
+		{
+			const FEnvBoolParam* PropertyValue = TestStruct->ContainerPtrToValuePtr<FEnvBoolParam>(Ob);
+			AddNamedValue(PropertyValue->ParamName, EEnvQueryParam::Bool, PropertyValue->Value ? 1.0f : -1.0f, NamedValues, RequiredParams);
+		}
+	}
+}
+
+#undef GET_STRUCT_NAME_CHECKED
+
+void UEnvQuery::CollectQueryParams(TArray<FEnvNamedValue>& NamedValues) const
+{
+	TArray<FName> RequiredParams;
+
+	// collect all params
+	for (int32 iOption = 0; iOption < Options.Num(); iOption++)
+	{
+		const UEnvQueryOption* Option = Options[iOption];
+		AddNamedValuesFromObject(Option->Generator, NamedValues, RequiredParams);
+
+		for (int32 iTest = 0; iTest < Option->Tests.Num(); iTest++)
+		{
+			const UEnvQueryTest* TestOb = Option->Tests[iTest];
+			AddNamedValuesFromObject(TestOb, NamedValues, RequiredParams);
+		}
+	}
+
+	// remove unnecessary params
+	for (int32 i = NamedValues.Num() - 1; i >= 0; i--)
+	{
+		if (!RequiredParams.Contains(NamedValues[i].ParamName))
+		{
+			NamedValues.RemoveAt(i);
+		}
+	}
+}
diff --git a/Engine/Source/Runtime/AIModule/Private/EnvironmentQuery/EnvQueryContext.cpp b/Engine/Source/Runtime/AIModule/Private/EnvironmentQuery/EnvQueryContext.cpp
new file mode 100644
--- /dev/null
+++ b/Engine/Source/Runtime/AIModule/Private/EnvironmentQuery/EnvQueryContext.cpp
@@ -0,0 +1,13 @@
+// Copyright 1998-2014 Epic Games, Inc. All Rights Reserved.
+
+#include "AIModulePrivate.h"
+#include "EnvironmentQuery/EnvQueryContext.h"
+
+UEnvQueryContext::UEnvQueryContext(const class FPostConstructInitializeProperties& PCIP) : Super(PCIP)
+{
+}
+
+void UEnvQueryContext::ProvideContext(struct FEnvQueryInstance& QueryInstance, struct FEnvQueryContextData& ContextData) const
+{
+	// empty in base class
+}
diff --git a/Engine/Source/Runtime/AIModule/Private/EnvironmentQuery/EnvQueryGenerator.cpp b/Engine/Source/Runtime/AIModule/Private/EnvironmentQuery/EnvQueryGenerator.cpp
new file mode 100644
--- /dev/null
+++ b/Engine/Source/Runtime/AIModule/Private/EnvironmentQuery/EnvQueryGenerator.cpp
@@ -0,0 +1,27 @@
+// Copyright 1998-2014 Epic Games, Inc. All Rights Reserved.
+
+#include "AIModulePrivate.h"
+#include "EnvironmentQuery/EnvQueryManager.h"
+#include "EnvironmentQuery/EnvQueryGenerator.h"
+
+UEnvQueryGenerator::UEnvQueryGenerator(const class FPostConstructInitializeProperties& PCIP) : Super(PCIP)
+{
+}
+
+FText UEnvQueryGenerator::GetDescriptionTitle() const
+{
+	return UEnvQueryTypes::GetShortTypeName(this);
+}
+
+FText UEnvQueryGenerator::GetDescriptionDetails() const
+{
+	return FText::GetEmpty();
+}
+
+#if WITH_EDITOR
+void UEnvQueryGenerator::PostEditChangeProperty(FPropertyChangedEvent& PropertyChangedEvent) 
+{
+	Super::PostEditChangeProperty(PropertyChangedEvent);
+	UEnvQueryManager::NotifyAssetUpdate(NULL);
+}
+#endif //WITH_EDITOR
