Extracting AI code from the engine into a separate module, pass 1 #UE4

diff --git a/Engine/Source/Runtime/AIModule/Private/EnvironmentQuery/EnvQueryInstance.cpp b/Engine/Source/Runtime/AIModule/Private/EnvironmentQuery/EnvQueryInstance.cpp
new file mode 100644
--- /dev/null
+++ b/Engine/Source/Runtime/AIModule/Private/EnvironmentQuery/EnvQueryInstance.cpp
@@ -0,0 +1,678 @@
+// Copyright 1998-2014 Epic Games, Inc. All Rights Reserved.
+
+#include "AIModulePrivate.h"
+#include "EnvironmentQuery/EnvQueryTypes.h"
+#include "EnvironmentQuery/EnvQueryTest.h"
+#include "EnvironmentQuery/Contexts/EnvQueryContext_Item.h"
+#include "EnvironmentQuery/Items/EnvQueryItemType_ActorBase.h"
+#include "EnvironmentQuery/Items/EnvQueryItemType_VectorBase.h"
+
+//----------------------------------------------------------------------//
+// FEQSQueryDebugData
+//----------------------------------------------------------------------//
+
+void FEQSQueryDebugData::Store(const FEnvQueryInstance* QueryInstance)
+{
+	DebugItemDetails = QueryInstance->ItemDetails;
+	DebugItems = QueryInstance->Items;
+	RawData = QueryInstance->RawData;
+}
+
+//----------------------------------------------------------------------//
+// FEnvQueryInstance
+//----------------------------------------------------------------------//
+#if USE_EQS_DEBUGGER
+bool FEnvQueryInstance::bDebuggingInfoEnabled = true;
+
+struct FItemSortPredicate
+{
+	TArray<FEnvQueryItem>* PtrItems;
+	FItemSortPredicate(TArray<FEnvQueryItem>* InItems) : PtrItems(InItems) {}
+
+	FORCEINLINE bool operator()(const FEnvQueryItemDetails& A, const FEnvQueryItemDetails& B) const
+	{
+		return (*PtrItems)[B.ItemIndex] < (*PtrItems)[A.ItemIndex];
+	}
+};
+#endif // USE_EQS_DEBUGGER
+
+bool FEnvQueryInstance::PrepareContext(UClass* ContextClass, FEnvQueryContextData& ContextData)
+{
+	if (ContextClass == NULL)
+	{
+		return false;
+	}
+
+	if (ContextClass != UEnvQueryContext_Item::StaticClass())
+	{
+		FEnvQueryContextData* CachedData = ContextCache.Find(ContextClass);
+		if (CachedData == NULL)
+		{
+			UEnvQueryContext* ContextOb = ContextClass->GetDefaultObject<UEnvQueryContext>();
+			ContextOb->ProvideContext(*this, ContextData);
+
+			DEC_MEMORY_STAT_BY(STAT_AI_EQS_InstanceMemory, GetContextAllocatedSize());
+
+			ContextCache.Add(ContextClass, ContextData);
+
+			INC_MEMORY_STAT_BY(STAT_AI_EQS_InstanceMemory, GetContextAllocatedSize());
+		}
+		else
+		{
+			ContextData = *CachedData;
+		}
+	}
+
+	if (ContextData.NumValues == 0)
+	{
+		UE_LOG(LogEQS, Log, TEXT("Query [%s] is missing values for context [%s], skipping test %d:%d [%s]"),
+			*QueryName, *UEnvQueryTypes::GetShortTypeName(ContextClass).ToString(),
+			OptionIndex, CurrentTest,
+			CurrentTest >=0 ? *UEnvQueryTypes::GetShortTypeName(Options[OptionIndex].TestDelegates[CurrentTest].GetUObject()).ToString() : TEXT("Generator")
+			);
+
+		return false;
+	}
+
+	return true;
+}
+
+bool FEnvQueryInstance::PrepareContext(UClass* Context, TArray<FEnvQuerySpatialData>& Data)
+{
+	if (Context == NULL)
+	{
+		return false;
+	}
+
+	FEnvQueryContextData ContextData;
+	const bool bSuccess = PrepareContext(Context, ContextData);
+
+	if (bSuccess && ContextData.ValueType && ContextData.ValueType->IsChildOf(UEnvQueryItemType_VectorBase::StaticClass()))
+	{
+		UEnvQueryItemType_VectorBase* DefTypeOb = (UEnvQueryItemType_VectorBase*)ContextData.ValueType->GetDefaultObject();
+		const uint16 DefTypeValueSize = DefTypeOb->GetValueSize();
+		uint8* RawData = ContextData.RawData.GetTypedData();
+
+		Data.Init(ContextData.NumValues);
+		for (int32 i = 0; i < ContextData.NumValues; i++)
+		{
+			Data[i].Location = DefTypeOb->GetLocation(RawData);
+			Data[i].Rotation = DefTypeOb->GetRotation(RawData);
+			RawData += DefTypeValueSize;
+		}
+	}
+
+	return bSuccess;
+}
+
+bool FEnvQueryInstance::PrepareContext(UClass* Context, TArray<FVector>& Data)
+{
+	if (Context == NULL)
+	{
+		return false;
+	}
+
+	FEnvQueryContextData ContextData;
+	const bool bSuccess = PrepareContext(Context, ContextData);
+
+	if (bSuccess && ContextData.ValueType && ContextData.ValueType->IsChildOf(UEnvQueryItemType_VectorBase::StaticClass()))
+	{
+		UEnvQueryItemType_VectorBase* DefTypeOb = (UEnvQueryItemType_VectorBase*)ContextData.ValueType->GetDefaultObject();
+		const uint16 DefTypeValueSize = DefTypeOb->GetValueSize();
+		uint8* RawData = (uint8*)ContextData.RawData.GetTypedData();
+
+		Data.Init(ContextData.NumValues);
+		for (int32 i = 0; i < ContextData.NumValues; i++)
+		{
+			Data[i] = DefTypeOb->GetLocation(RawData);
+			RawData += DefTypeValueSize;
+		}
+	}
+
+	return bSuccess;
+}
+
+bool FEnvQueryInstance::PrepareContext(UClass* Context, TArray<FRotator>& Data)
+{
+	if (Context == NULL)
+	{
+		return false;
+	}
+
+	FEnvQueryContextData ContextData;
+	const bool bSuccess = PrepareContext(Context, ContextData);
+
+	if (bSuccess && ContextData.ValueType && ContextData.ValueType->IsChildOf(UEnvQueryItemType_VectorBase::StaticClass()))
+	{
+		UEnvQueryItemType_VectorBase* DefTypeOb = (UEnvQueryItemType_VectorBase*)ContextData.ValueType->GetDefaultObject();
+		const uint16 DefTypeValueSize = DefTypeOb->GetValueSize();
+		uint8* RawData = ContextData.RawData.GetTypedData();
+
+		Data.Init(ContextData.NumValues);
+		for (int32 i = 0; i < ContextData.NumValues; i++)
+		{
+			Data[i] = DefTypeOb->GetRotation(RawData);
+			RawData += DefTypeValueSize;
+		}
+	}
+
+	return bSuccess;
+}
+
+bool FEnvQueryInstance::PrepareContext(UClass* Context, TArray<AActor*>& Data)
+{
+	if (Context == NULL)
+	{
+		return false;
+	}
+
+	FEnvQueryContextData ContextData;
+	const bool bSuccess = PrepareContext(Context, ContextData);
+
+	if (bSuccess && ContextData.ValueType && ContextData.ValueType->IsChildOf(UEnvQueryItemType_ActorBase::StaticClass()))
+	{
+		UEnvQueryItemType_ActorBase* DefTypeOb = (UEnvQueryItemType_ActorBase*)ContextData.ValueType->GetDefaultObject();
+		const uint16 DefTypeValueSize = DefTypeOb->GetValueSize();
+		uint8* RawData = ContextData.RawData.GetTypedData();
+
+		Data.Init(ContextData.NumValues);
+		for (int32 i = 0; i < ContextData.NumValues; i++)
+		{
+			Data[i] = DefTypeOb->GetActor(RawData);
+			RawData += DefTypeValueSize;
+		}
+	}
+
+	return bSuccess;
+}
+
+void FEnvQueryInstance::ExecuteOneStep(double InTimeLimit)
+{
+	if (!Owner.IsValid())
+	{
+		Status = EEnvQueryStatus::OwnerLost;
+		return;
+	}
+
+	FEnvQueryOptionInstance& OptionItem = Options[OptionIndex];
+	CONDITIONAL_SCOPE_CYCLE_COUNTER(STAT_AI_EQS_GeneratorTime, CurrentTest < 0);
+	CONDITIONAL_SCOPE_CYCLE_COUNTER(STAT_AI_EQS_TestTime, CurrentTest >= 0);
+
+	bool bStepDone = true;
+	TimeLimit = InTimeLimit;
+
+	if (CurrentTest < 0)
+	{
+		DEC_DWORD_STAT_BY(STAT_AI_EQS_NumItems, Items.Num());
+
+		RawData.Reset();
+		Items.Reset();
+		ItemType = OptionItem.ItemType;
+		ValueSize = ((UEnvQueryItemType*)ItemType->GetDefaultObject())->GetValueSize();
+
+		OptionItem.GenerateDelegate.Execute(*this);
+		FinalizeGeneration();
+	}
+	else
+	{
+		const int32 ItemsAlreadyProcessed = CurrentTestStartingItem;
+		OptionItem.TestDelegates[CurrentTest].Execute(*this);
+		bStepDone = CurrentTestStartingItem >= Items.Num() || bFoundSingleResult
+			// or no items processed ==> this means error
+			|| (ItemsAlreadyProcessed == CurrentTestStartingItem);
+
+		if (bStepDone)
+		{
+			FinalizeTest();
+		}
+	}
+	
+	if (bStepDone)
+	{
+#if USE_EQS_DEBUGGER
+		if (bStoreDebugInfo)
+		{
+			DebugData.Store(this);
+		}
+#endif // USE_EQS_DEBUGGER
+
+		CurrentTest++;
+		CurrentTestStartingItem = 0;
+	}
+
+	// sort results or switch to next option when all tests are performed
+	if (Status == EEnvQueryStatus::Processing &&
+		(OptionItem.TestDelegates.Num() == CurrentTest || NumValidItems <= 0))
+	{
+		if (NumValidItems > 0)
+		{
+			// found items, sort and finish
+			FinalizeQuery();
+		}
+		else
+		{
+			// no items here, go to next option or finish			
+			if (OptionIndex + 1 >= Options.Num())
+			{
+				// out of options, finish processing without errors
+				FinalizeQuery();
+			}
+			else
+			{
+				// not doing it always for debugging purposes
+				OptionIndex++;
+				CurrentTest = -1;
+			}
+		}
+	}
+}
+
+#if !NO_LOGGING
+void FEnvQueryInstance::Log(const FString Msg) const
+{
+	UE_LOG(LogEQS, Warning, TEXT("%s"), *Msg);
+}
+#endif // !NO_LOGGING
+
+void FEnvQueryInstance::ReserveItemData(int32 NumAdditionalItems)
+{
+	DEC_MEMORY_STAT_BY(STAT_AI_EQS_InstanceMemory, RawData.GetAllocatedSize());
+
+	RawData.Reserve((RawData.Num() + NumAdditionalItems) * ValueSize);
+
+	INC_MEMORY_STAT_BY(STAT_AI_EQS_InstanceMemory, RawData.GetAllocatedSize());
+}
+
+
+FEnvQueryInstance::ItemIterator::ItemIterator(class UEnvQueryTest* QueryTest, FEnvQueryInstance& QueryInstance)
+	: Instance(&QueryInstance), CurrentItem(QueryInstance.CurrentTestStartingItem)
+{
+	Deadline = QueryInstance.TimeLimit > 0.0 ? (FPlatformTime::Seconds() + QueryInstance.TimeLimit) : -1.0;
+	bDiscardFailed = QueryTest && QueryTest->bDiscardFailedItems;
+	InitItemScore();
+}
+
+void FEnvQueryInstance::ItemIterator::StoreTestResult()
+{
+	if (Instance->Mode == EEnvQueryRunMode::SingleResult &&	Instance->bPassOnSingleResult)
+	{
+		// handle SingleResult mode
+		if (bPassed && !bSkipped)
+		{
+			Instance->PickSingleItem(CurrentItem);
+			Instance->bFoundSingleResult = true;
+		}
+	}
+	else
+	{
+		if (bSkipped)
+		{
+			ItemScore = UEnvQueryTypes::SkippedItemValue;
+		}
+		else if (!bPassed || NumPartialScores == 0)
+		{
+			ItemScore = -1.f;
+			Instance->Items[CurrentItem].Discard();
+#if USE_EQS_DEBUGGER
+			Instance->ItemDetails[CurrentItem].FailedTestIndex = Instance->CurrentTest;
+#endif
+			Instance->NumValidItems--;
+		}
+		else
+		{
+			ItemScore /= NumPartialScores;
+		}
+
+		Instance->ItemDetails[CurrentItem].TestResults[Instance->CurrentTest] = ItemScore;
+	}
+}
+
+void FEnvQueryInstance::NormalizeScores()
+{
+	// @note this function assumes results have been already sorted and all first NumValidItems
+	// items in Items are valid (and the rest is not).
+	check(NumValidItems <= Items.Num())
+
+	float MinScore = 0.f;
+	float MaxScore = -BIG_NUMBER;
+
+	FEnvQueryItem* ItemInfo = Items.GetTypedData();
+	for (int32 ItemIndex = 0; ItemIndex < NumValidItems; ItemIndex++, ItemInfo++)
+	{
+		ensure(ItemInfo->IsValid());
+
+		MinScore = FMath::Min(MinScore, ItemInfo->Score);
+		MaxScore = FMath::Max(MaxScore, ItemInfo->Score);
+	}
+
+	ItemInfo = Items.GetTypedData();
+	if (MinScore == MaxScore)
+	{
+		for (int32 ItemIndex = 0; ItemIndex < NumValidItems; ItemIndex++, ItemInfo++)
+		{
+			ItemInfo->Score = 1.0f;
+		}
+	}
+	else
+	{
+		const float ScoreRange = MaxScore - MinScore;
+		for (int32 ItemIndex = 0; ItemIndex < NumValidItems; ItemIndex++, ItemInfo++)
+		{
+			ItemInfo->Score = (ItemInfo->Score - MinScore) / ScoreRange;
+		}
+	}
+}
+
+void FEnvQueryInstance::SortScores()
+{
+	if (Options[OptionIndex].bShuffleItems)
+	{
+		for (int32 i = 0; i < Items.Num(); i++)
+		{
+			const int32 Idx1 = FMath::RandHelper(Items.Num());
+			const int32 Idx2 = FMath::RandHelper(Items.Num());
+			Items.Swap(Idx1, Idx2);
+
+#if USE_EQS_DEBUGGER
+			ItemDetails.Swap(Idx1, Idx2);
+#endif
+		}
+	}
+
+#if USE_EQS_DEBUGGER
+	if (bStoreDebugInfo)
+	{
+		ItemDetails.Sort(FItemSortPredicate(&Items));
+	}
+	else
+	{
+#endif
+		ItemDetails.Reset();
+#if USE_EQS_DEBUGGER
+	}
+#endif
+
+	Items.Sort(TGreater<FEnvQueryItem>());
+}
+
+void FEnvQueryInstance::StripRedundantData()
+{
+#if USE_EQS_DEBUGGER
+	if (bStoreDebugInfo)
+	{
+		DebugData.Reset();
+	}
+#endif // USE_EQS_DEBUGGER
+	Items.SetNum(NumValidItems);
+}
+
+void FEnvQueryInstance::PickBestItem()
+{
+	// find first valid item with score worse than best one
+	int32 NumBestItems = NumValidItems;
+	for (int32 i = 1; i < NumValidItems; i++)
+	{
+		if (Items[i].Score < Items[0].Score)
+		{
+			NumBestItems = i;
+			break;
+		}
+	}
+
+	// pick only one, discard others
+	PickSingleItem(FMath::RandHelper(NumBestItems));
+}
+
+void FEnvQueryInstance::PickSingleItem(int32 ItemIndex)
+{
+#if USE_EQS_DEBUGGER
+	if (bStoreDebugInfo)
+	{
+		Items.Swap(0, ItemIndex);
+		ItemDetails.Swap(0, ItemIndex);
+
+		// do not limit valid items amount for debugger purposes!
+		// bFoundSingleResult can be used to determine if more than one item is valid
+
+		// normalize all scores for debugging
+		NormalizeScores();
+	}
+	else
+	{
+#endif
+		FEnvQueryItem BestItem;
+		BestItem.Score = 1.0f;
+		BestItem.DataOffset = Items[ItemIndex].DataOffset;
+
+		DEC_MEMORY_STAT_BY(STAT_AI_EQS_InstanceMemory, Items.GetAllocatedSize());
+
+		Items.Empty(1);
+		Items.Add(BestItem);
+		NumValidItems = 1;
+
+		INC_MEMORY_STAT_BY(STAT_AI_EQS_InstanceMemory, Items.GetAllocatedSize());
+#if USE_EQS_DEBUGGER
+	}
+#endif
+}
+
+void FEnvQueryInstance::FinalizeQuery()
+{
+	if (NumValidItems > 0)
+	{
+		if (Mode == EEnvQueryRunMode::SingleResult)
+		{
+			// if last test was not pure condition: sort and pick one of best items
+			if (!bFoundSingleResult)
+			{
+				SortScores();
+				PickBestItem();
+				bFoundSingleResult = true;
+			}
+		}
+		else
+		{
+			SortScores();
+
+			DEC_MEMORY_STAT_BY(STAT_AI_EQS_InstanceMemory, Items.GetAllocatedSize());
+
+			// remove failed ones from Items
+			Items.SetNum(NumValidItems);
+
+			INC_MEMORY_STAT_BY(STAT_AI_EQS_InstanceMemory, Items.GetAllocatedSize());
+
+			// normalizing after scoring and reducing number of elements to not 
+			// do anything for discarded items
+			NormalizeScores();
+		}
+	}
+	else
+	{
+		Items.Reset();
+		ItemDetails.Reset();
+		RawData.Reset();
+	}
+
+	Status = EEnvQueryStatus::Success;
+}
+
+void FEnvQueryInstance::FinalizeGeneration()
+{
+	FEnvQueryOptionInstance& OptionInstance = Options[OptionIndex];
+	const int32 NumTests = OptionInstance.TestDelegates.Num();
+
+	DEC_MEMORY_STAT_BY(STAT_AI_EQS_InstanceMemory, ItemDetails.GetAllocatedSize());
+	INC_DWORD_STAT_BY(STAT_AI_EQS_NumItems, Items.Num());
+
+	NumValidItems = Items.Num();
+	ItemDetails.Reset();
+	bPassOnSingleResult = false;
+	bFoundSingleResult = false;
+
+	if (NumValidItems > 0)
+	{
+		ItemDetails.Reserve(NumValidItems);
+		for (int32 i = 0; i < NumValidItems; i++)
+		{
+			ItemDetails.Add(FEnvQueryItemDetails(NumTests, i));
+		}
+	}
+
+	INC_MEMORY_STAT_BY(STAT_AI_EQS_InstanceMemory, ItemDetails.GetAllocatedSize());
+
+	ItemTypeVectorCDO = (ItemType && ItemType->IsChildOf(UEnvQueryItemType_VectorBase::StaticClass())) ?
+		(UEnvQueryItemType_VectorBase*)ItemType->GetDefaultObject() :	NULL;
+
+	ItemTypeActorCDO = (ItemType && ItemType->IsChildOf(UEnvQueryItemType_ActorBase::StaticClass())) ?
+		(UEnvQueryItemType_ActorBase*)ItemType->GetDefaultObject() : NULL;
+
+	if (Mode == EEnvQueryRunMode::SingleResult && NumTests == 1)
+	{
+		UEnvQueryTest* DefTestOb = Cast<UEnvQueryTest>(OptionInstance.TestDelegates[0].GetUObject());
+		if ((DefTestOb->TestPurpose != EEnvTestPurpose::Score) && // We are filtering and...
+			((DefTestOb->TestPurpose == EEnvTestPurpose::Filter) ||			// Either we are NOT scoring at ALL or...
+			 (DefTestOb->ScoringEquation == EEnvTestScoreEquation::Constant)// We are giving a constant score value for passing.
+			)
+		   )
+		{
+			OnFinalCondition();
+		}
+	}
+}
+
+void FEnvQueryInstance::FinalizeTest()
+{
+	FEnvQueryOptionInstance& OptionInstance = Options[OptionIndex];
+	const int32 NumTests = OptionInstance.TestDelegates.Num();
+
+	UEnvQueryTest* DefTestOb = (UEnvQueryTest*)(OptionInstance.TestDelegates[CurrentTest].GetUObject());
+	DefTestOb->NormalizeItemScores(*this);
+
+#if USE_EQS_DEBUGGER
+	if (bStoreDebugInfo)
+	{
+		DebugData.PerformedTestNames.Add(UEnvQueryTypes::GetShortTypeName(DefTestOb).ToString());
+	}
+#endif
+
+	// check if next test is the last one
+	if (Mode == EEnvQueryRunMode::SingleResult && CurrentTest == NumTests - 2)
+	{
+		DefTestOb = (UEnvQueryTest*)(OptionInstance.TestDelegates[NumTests - 1].GetUObject());
+		if ((DefTestOb->TestPurpose != EEnvTestPurpose::Score) && // We are filtering and...
+			((DefTestOb->TestPurpose == EEnvTestPurpose::Filter) ||			// Either we are NOT scoring at ALL or...
+			 (DefTestOb->ScoringEquation == EEnvTestScoreEquation::Constant)// We are giving a constant score value for passing.
+			)
+		   )
+		{
+			OnFinalCondition();
+		}
+	}
+}
+
+void FEnvQueryInstance::OnFinalCondition()
+{
+	// set flag for item iterator
+	bPassOnSingleResult = true;
+
+	UEnvQueryTest* DefTestOb = (UEnvQueryTest*)(Options[OptionIndex].TestDelegates[0].GetUObject());
+	const bool bShouldSort = (DefTestOb->WeightModifier < EEnvTestWeight::Constant);
+
+#if USE_EQS_DEBUGGER
+	if (bStoreDebugInfo)
+	{
+		DebugData.Store(this);
+
+		// randomize before sorting, so items with the same score can be chosen randomly
+		for (int32 i = 0; i < Items.Num(); i++)
+		{
+			const int32 Idx1 = FMath::RandHelper(Items.Num());
+			const int32 Idx2 = FMath::RandHelper(Items.Num());
+			Items.Swap(Idx1, Idx2);
+			ItemDetails.Swap(Idx1, Idx2);
+		}
+
+		// keep ItemDetails in sync for debugger purposes
+		if (bShouldSort)
+		{
+			ItemDetails.Sort(FItemSortPredicate(&Items));
+		}
+	}
+	else
+	{
+#endif
+		// don't update ItemDetails here - it won't be needed anymore, just reset it
+		ItemDetails.Reset();
+
+		// randomize before sorting, so items with the same score can be chosen randomly
+		for (int32 i = 0; i < Items.Num(); i++)
+		{
+			Items.Swap(FMath::RandHelper(Items.Num()), FMath::RandHelper(Items.Num()));
+		}
+#if USE_EQS_DEBUGGER
+	}
+#endif	
+
+	if (bShouldSort)
+	{
+		Items.Sort(TGreater<FEnvQueryItem>());
+	}
+}
+
+#if STATS
+uint32 FEnvQueryInstance::GetAllocatedSize() const
+{
+	uint32 MemSize = sizeof(*this) + Items.GetAllocatedSize() + RawData.GetAllocatedSize();
+	MemSize += GetContextAllocatedSize();
+	MemSize += NamedParams.GetAllocatedSize();
+	MemSize += ItemDetails.GetAllocatedSize();
+	MemSize += Options.GetAllocatedSize();
+
+	for (int32 i = 0; i < Options.Num(); i++)
+	{
+		MemSize += Options[i].GetAllocatedSize();
+	}
+
+	return MemSize;
+}
+
+uint32 FEnvQueryInstance::GetContextAllocatedSize() const
+{
+	uint32 MemSize = ContextCache.GetAllocatedSize();
+	for (TMap<UClass*, FEnvQueryContextData>::TConstIterator It(ContextCache); It; ++It)
+	{
+		MemSize += It.Value().GetAllocatedSize();
+	}
+
+	return MemSize;
+}
+#endif // STATS
+
+FBox FEnvQueryInstance::GetBoundingBox() const
+{
+	const TArray<FEnvQueryItem>& Items = 
+#if USE_EQS_DEBUGGER
+	DebugData.DebugItems.Num() > 0 ? DebugData.DebugItems : 
+#endif // USE_EQS_DEBUGGER
+		this->Items;
+
+	const TArray<uint8>& RawData = 
+#if USE_EQS_DEBUGGER
+		DebugData.RawData.Num() > 0 ? DebugData.RawData : 
+#endif // USE_EQS_DEBUGGER
+		this->RawData;
+
+	FBox BBox(0);
+
+	if (ItemType->IsChildOf(UEnvQueryItemType_VectorBase::StaticClass()))
+	{
+		UEnvQueryItemType_VectorBase* DefTypeOb = (UEnvQueryItemType_VectorBase*)ItemType->GetDefaultObject();
+
+		for (int32 Index = 0; Index < Items.Num(); ++Index)
+		{		
+			BBox += DefTypeOb->GetLocation(RawData.GetTypedData() + Items[Index].DataOffset);
+		}
+	}
+
+	return BBox;
+}
\ No newline at end of file
diff --git a/Engine/Source/Runtime/AIModule/Private/EnvironmentQuery/EnvQueryManager.cpp b/Engine/Source/Runtime/AIModule/Private/EnvironmentQuery/EnvQueryManager.cpp
new file mode 100644
--- /dev/null
+++ b/Engine/Source/Runtime/AIModule/Private/EnvironmentQuery/EnvQueryManager.cpp
@@ -0,0 +1,520 @@
+// Copyright 1998-2014 Epic Games, Inc. All Rights Reserved.
+
+#include "AIModulePrivate.h"
+#include "EnvironmentQuery/EnvQuery.h"
+#include "EnvironmentQuery/EnvQueryTest.h"
+#include "EnvironmentQuery/EnvQueryGenerator.h"
+#include "EnvironmentQuery/EnvQueryOption.h"
+#include "EnvironmentQuery/EnvQueryManager.h"
+#include "EnvironmentQuery/EQSTestingPawn.h"
+
+#if WITH_EDITOR
+#include "UnrealEd.h"
+#endif
+
+DEFINE_LOG_CATEGORY(LogEQS);
+
+DEFINE_STAT(STAT_AI_EQS_Tick);
+DEFINE_STAT(STAT_AI_EQS_LoadTime);
+DEFINE_STAT(STAT_AI_EQS_GeneratorTime);
+DEFINE_STAT(STAT_AI_EQS_TestTime);
+DEFINE_STAT(STAT_AI_EQS_NumInstances);
+DEFINE_STAT(STAT_AI_EQS_NumItems);
+DEFINE_STAT(STAT_AI_EQS_InstanceMemory);
+
+//////////////////////////////////////////////////////////////////////////
+// FEnvQueryRequest
+
+FEnvQueryRequest& FEnvQueryRequest::SetNamedParams(const TArray<struct FEnvNamedValue>& Params)
+{
+	for (int32 ParamIndex = 0; ParamIndex < Params.Num(); ParamIndex++)
+	{
+		NamedParams.Add(Params[ParamIndex].ParamName, Params[ParamIndex].Value);
+	}
+
+	return *this;
+}
+
+int32 FEnvQueryRequest::Execute(EEnvQueryRunMode::Type RunMode, FQueryFinishedSignature const& FinishDelegate)
+{
+	if (Owner == NULL)
+	{
+		Owner = FinishDelegate.GetUObject();
+		if (Owner == NULL)
+		{
+			UE_LOG(LogEQS, Warning, TEXT("Unknown owner of request: %s"), *GetNameSafe(QueryTemplate));
+			return INDEX_NONE;
+		}
+	}
+
+	if (World == NULL)
+	{
+		World = GEngine->GetWorldFromContextObject(Owner);
+		if (World == NULL)
+		{
+			UE_LOG(LogEQS, Warning, TEXT("Unable to access world with owner: %s"), *GetNameSafe(Owner));
+			return INDEX_NONE;
+		}
+	}
+
+	UEnvQueryManager* EnvQueryManager = World->GetEnvironmentQueryManager();
+	if (EnvQueryManager == NULL)
+	{
+		UE_LOG(LogEQS, Warning, TEXT("Missing EQS manager!"));
+		return INDEX_NONE;
+	}
+
+	return EnvQueryManager->RunQuery(*this, RunMode, FinishDelegate);
+}
+
+
+//////////////////////////////////////////////////////////////////////////
+// UEnvQueryManager
+
+TArray<TSubclassOf<UEnvQueryItemType> > UEnvQueryManager::RegisteredItemTypes;
+
+UEnvQueryManager::UEnvQueryManager(const class FPostConstructInitializeProperties& PCIP) : Super(PCIP)
+{
+	if (!HasAnyFlags(RF_ClassDefaultObject))
+	{
+		FCoreDelegates::PreLoadMap.AddUObject(this, &UEnvQueryManager::OnPreLoadMap);
+	}
+
+	NextQueryID = 0;
+}
+
+void UEnvQueryManager::FinishDestroy()
+{
+	FCoreDelegates::PreLoadMap.RemoveAll(this);
+	Super::FinishDestroy();
+}
+
+UEnvQueryManager* UEnvQueryManager::GetCurrent(UWorld* World)
+{
+	return World ? World->GetEnvironmentQueryManager() : NULL;
+}
+
+UEnvQueryManager* UEnvQueryManager::GetCurrent(class UObject* WorldContextObject)
+{
+	UWorld* World = WorldContextObject ? GEngine->GetWorldFromContextObject(WorldContextObject) : NULL;
+	return World ? World->GetEnvironmentQueryManager() : NULL;
+}
+
+#if WITH_EDITOR
+void UEnvQueryManager::NotifyAssetUpdate(UEnvQuery* Query)
+{
+	if (GEditor == NULL)
+	{
+		return;
+	}
+
+	UWorld* World = GEditor->GetEditorWorldContext().World();
+	if (World)
+	{
+		UEnvQueryManager* EQS = World->GetEnvironmentQueryManager();
+		if (EQS)
+		{
+			EQS->InstanceCache.Reset();
+		}
+
+		// was as follows, but got broken with changes to actor iterator (FActorIteratorBase::SpawnedActorArray)
+		// for (TActorIterator<AEQSTestingPawn> It(World); It; ++It)
+		for (TActorIterator<AActor> It(World); It; ++It)
+		{
+			AEQSTestingPawn* EQSPawn = Cast<AEQSTestingPawn>(*It);
+			if (EQSPawn == NULL)
+			{
+				continue;
+			}
+
+			if (EQSPawn->QueryTemplate == Query || Query == NULL)
+			{
+				EQSPawn->RunEQSQuery();
+			}
+		}
+	}
+}
+#endif // WITH_EDITOR
+
+TStatId UEnvQueryManager::GetStatId() const
+{
+	RETURN_QUICK_DECLARE_CYCLE_STAT(UEnvQueryManager, STATGROUP_Tickables);
+}
+
+int32 UEnvQueryManager::RunQuery(const FEnvQueryRequest& Request, EEnvQueryRunMode::Type RunMode, FQueryFinishedSignature const& FinishDelegate)
+{
+	TSharedPtr<FEnvQueryInstance> QueryInstance = PrepareQueryInstance(Request, RunMode);
+	if (QueryInstance.IsValid() == false)
+	{
+		return INDEX_NONE;
+	}
+
+	QueryInstance->FinishDelegate = FinishDelegate;
+	RunningQueries.Add(QueryInstance);
+
+	return QueryInstance->QueryID;
+}
+
+TSharedPtr<FEnvQueryResult> UEnvQueryManager::RunInstantQuery(const FEnvQueryRequest& Request, EEnvQueryRunMode::Type RunMode)
+{
+	TSharedPtr<FEnvQueryInstance> QueryInstance = PrepareQueryInstance(Request, RunMode);
+	if (!QueryInstance.IsValid())
+	{
+		return NULL;
+	}
+
+	while (QueryInstance->Status == EEnvQueryStatus::Processing)
+	{
+		QueryInstance->ExecuteOneStep((double)FLT_MAX);
+	}
+
+#if USE_EQS_DEBUGGER
+	EQSDebugger.StoreQuery(QueryInstance);
+#endif // USE_EQS_DEBUGGER
+
+	return QueryInstance;
+}
+
+TSharedPtr<struct FEnvQueryInstance> UEnvQueryManager::PrepareQueryInstance(const FEnvQueryRequest& Request, EEnvQueryRunMode::Type RunMode)
+{
+	TSharedPtr<FEnvQueryInstance> QueryInstance = CreateQueryInstance(Request.QueryTemplate, RunMode);
+	if (!QueryInstance.IsValid())
+	{
+		return NULL;
+	}
+
+	QueryInstance->World = Cast<UWorld>(GetOuter());
+	QueryInstance->Owner = Request.Owner;
+
+	DEC_MEMORY_STAT_BY(STAT_AI_EQS_InstanceMemory, QueryInstance->NamedParams.GetAllocatedSize());
+
+	// @TODO: interface for providing default named params (like custom ranges in AI)
+	QueryInstance->NamedParams = Request.NamedParams;
+
+	INC_MEMORY_STAT_BY(STAT_AI_EQS_InstanceMemory, QueryInstance->NamedParams.GetAllocatedSize());
+
+	QueryInstance->QueryID = NextQueryID++;
+
+	return QueryInstance;
+}
+
+bool UEnvQueryManager::AbortQuery(int32 RequestID)
+{
+	for (int32 QueryIndex = 0; QueryIndex < RunningQueries.Num(); QueryIndex++)
+	{
+		TSharedPtr<FEnvQueryInstance>& QueryInstance = RunningQueries[QueryIndex];
+		if (QueryInstance->QueryID == RequestID &&
+			QueryInstance->Status == EEnvQueryStatus::Processing)
+		{
+			QueryInstance->Status = EEnvQueryStatus::Aborted;
+			QueryInstance->FinishDelegate.ExecuteIfBound(QueryInstance);
+			
+			RunningQueries.RemoveAt(QueryIndex);
+			return true;
+		}
+	}
+
+	return false;
+}
+
+void UEnvQueryManager::Tick(float DeltaTime)
+{
+	SCOPE_CYCLE_COUNTER(STAT_AI_EQS_Tick);
+	SET_DWORD_STAT(STAT_AI_EQS_NumInstances, RunningQueries.Num());
+	// @TODO: threads?
+
+	const double StartTime = FPlatformTime::Seconds();
+	const double MaxAllowedSeconds = 0.010;
+	double TimeLeft = MaxAllowedSeconds;
+		
+	while (TimeLeft > 0.0 && RunningQueries.Num() > 0)
+	{
+		for (int32 Index = 0; Index < RunningQueries.Num() && TimeLeft > 0.0; Index++)
+		{
+			TSharedPtr<FEnvQueryInstance>& QueryInstance = RunningQueries[Index];
+			QueryInstance->ExecuteOneStep(TimeLeft);
+			
+			if (QueryInstance->Status != EEnvQueryStatus::Processing)
+			{
+#if USE_EQS_DEBUGGER
+				EQSDebugger.StoreQuery(QueryInstance);
+#endif // USE_EQS_DEBUGGER
+
+				QueryInstance->FinishDelegate.ExecuteIfBound(QueryInstance);
+
+				RunningQueries.RemoveAt(Index);
+				Index--;
+			}
+
+			TimeLeft -= FPlatformTime::Seconds() - StartTime;
+		}
+	}
+}
+
+void UEnvQueryManager::OnPreLoadMap()
+{
+	if (RunningQueries.Num() > 0)
+	{
+		for (int32 Index = 0; Index < RunningQueries.Num(); Index++)
+		{
+			TSharedPtr<FEnvQueryInstance>& QueryInstance = RunningQueries[Index];
+			if (QueryInstance->Status == EEnvQueryStatus::Processing)
+			{
+				QueryInstance->Status = EEnvQueryStatus::Failed;
+				QueryInstance->FinishDelegate.ExecuteIfBound(QueryInstance);
+			}
+		}
+
+		RunningQueries.Reset();
+	}
+}
+
+namespace EnvQueryTestSort
+{
+	struct FAllMatching
+	{
+		FORCEINLINE bool operator()(const UEnvQueryTest& TestA, const UEnvQueryTest& TestB) const
+		{
+			// cheaper tests go first
+			if (TestB.Cost > TestA.Cost)
+			{
+				return true;
+			}
+
+			// conditions go first
+			const bool bConditionA = (TestA.TestPurpose != EEnvTestPurpose::Score); // Is Test A Filtering?
+			const bool bConditionB = (TestB.TestPurpose != EEnvTestPurpose::Score); // Is Test B Filtering?
+			if (bConditionA && !bConditionB)
+			{
+				return true;
+			}
+
+			// keep connection order (sort stability)
+			return (TestB.TestOrder > TestA.TestOrder);
+		}
+	};
+
+	struct FSingleResult
+	{
+		FSingleResult(EEnvTestCost::Type InHighestCost) : HighestCost(InHighestCost) {}
+
+		FORCEINLINE bool operator()(const UEnvQueryTest& TestA, const UEnvQueryTest& TestB) const
+		{
+			// cheaper tests go first
+			if (TestB.Cost > TestA.Cost)
+			{
+				return true;
+			}
+
+			const bool bConditionA = (TestA.TestPurpose != EEnvTestPurpose::Score); // Is Test A Filtering?
+			const bool bConditionB = (TestB.TestPurpose != EEnvTestPurpose::Score); // Is Test B Filtering?
+			if (TestA.Cost == HighestCost)
+			{
+				// highest cost: weights go first, conditions later (first match will return result)
+				if (!bConditionA && bConditionB)
+				{
+					return true;
+				}
+
+				// conditions with weights before pure conditions
+				const bool bNoWeightA = (TestA.WeightModifier >= EEnvTestWeight::Constant);
+				const bool bNoWeightB = (TestB.WeightModifier >= EEnvTestWeight::Constant);
+				if (!bNoWeightA && bNoWeightB)
+				{
+					return true;
+				}
+			}
+			else
+			{
+				// lower costs: conditions go first to reduce amount of items
+				if (bConditionA && !bConditionB)
+				{
+					return true;
+				}
+			}
+
+			// keep connection order (sort stability)
+			return (TestB.TestOrder > TestA.TestOrder);
+		}
+
+	protected:
+		TEnumAsByte<EEnvTestCost::Type> HighestCost;
+	};
+}
+
+UEnvQuery* UEnvQueryManager::FindQueryTemplate(const FString& QueryName) const
+{
+	for (int32 InstanceIndex = 0; InstanceIndex < InstanceCache.Num(); InstanceIndex++)
+	{
+		const UEnvQuery* QueryTemplate = InstanceCache[InstanceIndex].Template.Get();
+
+		if (QueryTemplate != NULL && QueryTemplate->GetName() == QueryName)
+		{
+			return const_cast<UEnvQuery*>(QueryTemplate);
+		}
+	}
+
+	for (FObjectIterator It(UEnvQuery::StaticClass()); It; ++It)
+	{
+		if (It->GetName() == QueryName)
+		{
+			return Cast<UEnvQuery>(*It);
+		}
+	}
+
+	return NULL;
+}
+
+TSharedPtr<FEnvQueryInstance> UEnvQueryManager::CreateQueryInstance(const class UEnvQuery* Template, EEnvQueryRunMode::Type RunMode)
+{
+	if (Template == NULL)
+	{
+		return NULL;
+	}
+
+	// try to find entry in cache
+	FEnvQueryInstance* InstanceTemplate = NULL;
+	for (int32 InstanceIndex = 0; InstanceIndex < InstanceCache.Num(); InstanceIndex++)
+	{
+		if (InstanceCache[InstanceIndex].Template == Template &&
+			InstanceCache[InstanceIndex].Instance.Mode == RunMode)
+		{
+			InstanceTemplate = &InstanceCache[InstanceIndex].Instance;
+			break;
+		}
+	}
+
+	// and create one if can't be found
+	if (InstanceTemplate == NULL)
+	{
+		SCOPE_CYCLE_COUNTER(STAT_AI_EQS_LoadTime);
+		
+		{
+			// memory stat tracking: temporary variable will exist only inside this section
+			FEnvQueryInstanceCache NewCacheEntry;
+			NewCacheEntry.Template = Template;
+			NewCacheEntry.Instance.QueryName = Template->GetName();
+			NewCacheEntry.Instance.Mode = RunMode;
+
+			const int32 Idx = InstanceCache.Add(NewCacheEntry);
+			InstanceTemplate = &InstanceCache[Idx].Instance;
+		}
+
+		for (int32 OptionIndex = 0; OptionIndex < Template->Options.Num(); OptionIndex++)
+		{
+			UEnvQueryOption* MyOption = Template->Options[OptionIndex];
+			if (MyOption == NULL || MyOption->Generator == NULL)
+			{
+				UE_LOG(LogEQS, Error, TEXT("Trying to spawn a query with broken Template (empty generator): %s, option %d"),
+					*GetNameSafe(Template), OptionIndex);
+
+				continue;
+			}
+
+			EEnvTestCost::Type HighestCost(EEnvTestCost::Low);
+			TArray<UEnvQueryTest*> SortedTests = MyOption->Tests;
+			TSubclassOf<UEnvQueryItemType> GeneratedType = MyOption->Generator->ItemType;
+			for (int32 TestIndex = SortedTests.Num() - 1; TestIndex >= 0; TestIndex--)
+			{
+				UEnvQueryTest* TestOb = SortedTests[TestIndex];
+				if (TestOb == NULL || !TestOb->IsSupportedItem(GeneratedType))
+				{
+					UE_LOG(LogEQS, Warning, TEXT("Query [%s] can't use test [%s] in option %d [%s], removing it"),
+						*GetNameSafe(Template), *GetNameSafe(TestOb), OptionIndex, *MyOption->Generator->OptionName);
+
+					SortedTests.RemoveAt(TestIndex);
+				}
+				else if (HighestCost < TestOb->Cost)
+				{
+					HighestCost = TestOb->Cost;
+				}
+			}
+
+			if (SortedTests.Num())
+			{
+				switch (RunMode)
+				{
+				case EEnvQueryRunMode::SingleResult:
+					SortedTests.Sort(EnvQueryTestSort::FSingleResult(HighestCost));
+					break;
+
+				case EEnvQueryRunMode::AllMatching:
+					SortedTests.Sort(EnvQueryTestSort::FAllMatching());
+					break;
+
+				default:
+					{
+						UEnum* RunModeEnum = FindObject<UEnum>(ANY_PACKAGE, TEXT("EEnvQueryRunMode"));
+						UE_LOG(LogEQS, Warning, TEXT("Query [%s] can't be sorted for RunMode: %d [%s]"),
+							*GetNameSafe(Template), (int32)RunMode, RunModeEnum ? *RunModeEnum->GetEnumName(RunMode) : TEXT("??"));
+					}
+				}
+			}
+
+			CreateOptionInstance(MyOption, SortedTests, *InstanceTemplate);
+		}
+
+		if (InstanceTemplate->Options.Num() == 0)
+		{
+			UE_LOG(LogEQS, Warning, TEXT("Query [%s] doesn't have any valid options!"), *GetNameSafe(Template));
+			return NULL;
+		}
+	}
+
+	// create new instance
+	TSharedPtr<FEnvQueryInstance> NewInstance(new FEnvQueryInstance(*InstanceTemplate));
+	return NewInstance;
+}
+
+void UEnvQueryManager::CreateOptionInstance(class UEnvQueryOption* OptionTemplate, const TArray<UEnvQueryTest*>& SortedTests, struct FEnvQueryInstance& Instance)
+{
+	FEnvQueryOptionInstance OptionInstance;
+	OptionInstance.GenerateDelegate = OptionTemplate->Generator->GenerateDelegate;
+	OptionInstance.ItemType = OptionTemplate->Generator->ItemType;
+	OptionInstance.bShuffleItems = true;
+
+	OptionInstance.TestDelegates.AddZeroed(SortedTests.Num());
+	for (int32 TestIndex = 0; TestIndex < SortedTests.Num(); TestIndex++)
+	{
+		const UEnvQueryTest* TestOb = SortedTests[TestIndex];
+		OptionInstance.TestDelegates[TestIndex] = TestOb->ExecuteDelegate;
+
+		// HACK!  TODO: Is this the correct replacement here?  or should it check just if SCORING ONLY?
+		// always randomize when asking for single result
+		// otherwise, can skip randomization if test wants to score every item
+		if ((TestOb->TestPurpose != EEnvTestPurpose::Filter) && // We ARE scoring, regardless of whether or not we're filtering
+			(Instance.Mode != EEnvQueryRunMode::SingleResult))
+		{
+			OptionInstance.bShuffleItems = false;
+		}
+	}
+
+	DEC_MEMORY_STAT_BY(STAT_AI_EQS_InstanceMemory, Instance.Options.GetAllocatedSize());
+
+	const int32 AddedIdx = Instance.Options.Add(OptionInstance);
+
+	INC_MEMORY_STAT_BY(STAT_AI_EQS_InstanceMemory, Instance.Options.GetAllocatedSize() + Instance.Options[AddedIdx].GetAllocatedSize());
+}
+
+//----------------------------------------------------------------------//
+// FEQSDebugger
+//----------------------------------------------------------------------//
+#if USE_EQS_DEBUGGER
+
+void FEQSDebugger::StoreQuery(TSharedPtr<FEnvQueryInstance>& Query)
+{
+	TSharedPtr<FEnvQueryInstance>& StoredQuery = StoredQueries.FindOrAdd(Query->Owner.Get());
+	StoredQuery = Query;
+}
+
+TSharedPtr<FEnvQueryInstance> FEQSDebugger::GetQueryForOwner(const UObject* Owner)
+{
+	return StoredQueries.FindRef(Owner);
+}
+
+const TSharedPtr<FEnvQueryInstance> FEQSDebugger::GetQueryForOwner(const UObject* Owner) const
+{
+	return StoredQueries.FindRef(Owner);
+}
+
+#endif // USE_EQS_DEBUGGER
diff --git a/Engine/Source/Runtime/AIModule/Private/EnvironmentQuery/EnvQueryOption.cpp b/Engine/Source/Runtime/AIModule/Private/EnvironmentQuery/EnvQueryOption.cpp
new file mode 100644
--- /dev/null
+++ b/Engine/Source/Runtime/AIModule/Private/EnvironmentQuery/EnvQueryOption.cpp
@@ -0,0 +1,19 @@
+// Copyright 1998-2014 Epic Games, Inc. All Rights Reserved.
+
+#include "AIModulePrivate.h"
+#include "EnvironmentQuery/EnvQueryGenerator.h"
+#include "EnvironmentQuery/EnvQueryOption.h"
+
+UEnvQueryOption::UEnvQueryOption(const class FPostConstructInitializeProperties& PCIP) : Super(PCIP)
+{
+}
+
+FText UEnvQueryOption::GetDescriptionTitle() const
+{
+	return Generator ? Generator->GetDescriptionTitle() : FText::GetEmpty();
+}
+
+FText UEnvQueryOption::GetDescriptionDetails() const
+{
+	return Generator ? Generator->GetDescriptionDetails() : FText::GetEmpty();
+}
diff --git a/Engine/Source/Runtime/AIModule/Private/EnvironmentQuery/EnvQueryTest.cpp b/Engine/Source/Runtime/AIModule/Private/EnvironmentQuery/EnvQueryTest.cpp
new file mode 100644
--- /dev/null
+++ b/Engine/Source/Runtime/AIModule/Private/EnvironmentQuery/EnvQueryTest.cpp
@@ -0,0 +1,455 @@
+// Copyright 1998-2014 Epic Games, Inc. All Rights Reserved.
+
+#include "AIModulePrivate.h"
+#include "EnvironmentQuery/EnvQueryOption.h"
+#include "EnvironmentQuery/EnvQueryContext.h"
+#include "EnvironmentQuery/EnvQueryManager.h"
+#include "EnvironmentQuery/Contexts/EnvQueryContext_Item.h"
+#include "EnvironmentQuery/Items/EnvQueryItemType_VectorBase.h"
+#include "EnvironmentQuery/Items/EnvQueryItemType_ActorBase.h"
+#include "EnvironmentQuery/EnvQueryTest.h"
+
+#define LOCTEXT_NAMESPACE "EnvQueryGenerator"
+
+UEnvQueryTest::UEnvQueryTest(const class FPostConstructInitializeProperties& PCIP) : Super(PCIP)
+{
+	TestPurpose = EEnvTestPurpose::FilterAndScore;
+	FilterType = EEnvTestFilterType::Range;
+	BoolFilter.Value = true;
+	bFormatUpdated = false;
+	Condition = EEnvTestCondition::NoCondition;	// DEPRECATED
+	Weight.Value = 1.0f;						// DEPRECATED
+	WeightModifier = EEnvTestWeight::None;		// DEPRECATED
+	Cost = EEnvTestCost::Low;
+// 	bUseAbsoluteValueBeforeClamping = false;
+	ClampMinType = EEnvQueryTestClamping::None;
+	ClampMaxType = EEnvQueryTestClamping::None;
+// 	bMirrorNormalizedScore = false;
+	ScoringEquation = EEnvTestScoreEquation::Linear;
+	
+	bNormalizeFromZero = false;					// DEPRECATED
+	bWorkOnFloatValues = true;
+	bDiscardFailedItems = true;					// DEPRECATED
+}
+
+void UEnvQueryTest::NormalizeItemScores(struct FEnvQueryInstance& QueryInstance)
+{
+// 	if (ScoringEquation == EEnvTestScoreEquation::Skip)
+	if (!IsScoring())
+	{
+		return;
+	}
+
+	float WeightValue = 0.0f;
+	if (!QueryInstance.GetParamValue(Weight, WeightValue, TEXT("Weight")))
+	{
+		return;
+	}
+
+// 	float ScoringFactorValue = 0.0f;
+// 	if (!QueryInstance.GetParamValue(ScoringFactor, ScoringFactorValue, TEXT("Factor")))
+// 	{
+// 		return;
+// 	}
+	
+	float MinScore = 0;
+	float MaxScore = -BIG_NUMBER;
+
+	if (ClampMinType == EEnvQueryTestClamping::FilterThreshold)
+	{
+		bool bSuccess = QueryInstance.GetParamValue(FloatFilterMin, MinScore, TEXT("FloatFilterMin"));
+		if (!bSuccess)
+		{
+			UE_LOG(LogEQS, Warning, TEXT("Unable to get FloatFilterMin parameter value from EnvQueryInstance %s"), FloatFilterMin.IsNamedParam() ? *FloatFilterMin.ParamName.ToString() : TEXT("<No name specified>"));
+		}
+	}
+	else if (ClampMinType == EEnvQueryTestClamping::SpecifiedValue)
+	{
+		bool bSuccess = QueryInstance.GetParamValue(ScoreClampingMin, MinScore, TEXT("ScoreClampingMin"));
+		if (!bSuccess)
+		{
+			UE_LOG(LogEQS, Warning, TEXT("Unable to get ClampMinType parameter value from EnvQueryInstance %s"), ScoreClampingMin.IsNamedParam() ? *ScoreClampingMin.ParamName.ToString() : TEXT("<No name specified>"));
+		}
+	}
+
+	if (ClampMaxType == EEnvQueryTestClamping::FilterThreshold)
+	{
+		bool bSuccess = QueryInstance.GetParamValue(FloatFilterMax, MaxScore, TEXT("FloatFilterMax"));
+		if (!bSuccess)
+		{
+			UE_LOG(LogEQS, Warning, TEXT("Unable to get FloatFilterMax parameter value from EnvQueryInstance %s"), FloatFilterMax.IsNamedParam() ? *FloatFilterMax.ParamName.ToString() : TEXT("<No name specified>"));
+		}
+	}
+	else if (ClampMaxType == EEnvQueryTestClamping::SpecifiedValue)
+	{
+		bool bSuccess = QueryInstance.GetParamValue(ScoreClampingMax, MaxScore, TEXT("ScoreClampingMax"));
+		if (!bSuccess)
+		{
+			UE_LOG(LogEQS, Warning, TEXT("Unable to get ScoreClampingMax parameter value from EnvQueryInstance %s"), ScoreClampingMax.IsNamedParam() ? *ScoreClampingMax.ParamName.ToString() : TEXT("<No name specified>"));
+		}
+	}
+
+	FEnvQueryItemDetails* DetailInfo = QueryInstance.ItemDetails.GetTypedData();
+	if ((ClampMinType == EEnvQueryTestClamping::None) ||
+		(ClampMaxType == EEnvQueryTestClamping::None)
+	   )
+	{
+		for (int32 iItem = 0; iItem < QueryInstance.Items.Num(); iItem++, DetailInfo++)
+		{
+			if (!QueryInstance.Items[iItem].IsValid())
+			{
+				continue;
+			}
+
+			float TestValue = DetailInfo->TestResults[QueryInstance.CurrentTest];
+// 			if (bUseAbsoluteValueBeforeClamping)
+// 			{
+// 				TestValue = FMath::Abs(TestValue);
+// 			}
+
+			if (TestValue != UEnvQueryTypes::SkippedItemValue)
+			{
+				if (ClampMinType == EEnvQueryTestClamping::None)
+				{
+					MinScore = FMath::Min(MinScore, TestValue);
+				}
+
+				if (ClampMaxType == EEnvQueryTestClamping::None)
+				{
+					MaxScore = FMath::Max(MaxScore, TestValue);
+				}
+			}
+		}
+	}
+
+	DetailInfo = QueryInstance.ItemDetails.GetTypedData();
+	if (bNormalizeFromZero && MinScore > 0.0f)
+	{
+		MinScore = 0.0f;
+	}
+
+	if (MinScore != MaxScore)
+	{
+		for (int32 ItemIndex = 0; ItemIndex < QueryInstance.ItemDetails.Num(); ItemIndex++, DetailInfo++)
+		{
+			if (QueryInstance.Items[ItemIndex].IsValid() == false)
+			{
+				continue;
+			}
+
+			float WeightedScore = 0.0f;
+
+			float& TestValue = DetailInfo->TestResults[QueryInstance.CurrentTest];
+			if (TestValue != UEnvQueryTypes::SkippedItemValue)
+			{
+// 				const float TestValueToNormalize = bUseAbsoluteValueBeforeClamping ? FMath::Abs(TestValue) : TestValue;
+				const float NormalizedScore = (TestValue - MinScore) / (MaxScore - MinScore);
+				// TODO: Add an option to invert the normalized score before applying an equation.
+ 				const float NormalizedScoreForEquation = /*bMirrorNormalizedScore ? (1.0f - NormalizedScore) :*/ NormalizedScore;
+				switch (ScoringEquation)
+				{
+					case EEnvTestScoreEquation::Linear:
+						WeightedScore = WeightValue * NormalizedScoreForEquation;
+						break;
+
+					case EEnvTestScoreEquation::InverseLinear:
+					{
+						// For now, we're avoiding having a separate flag for flipping the direction of the curve
+						// because we don't have usage cases yet and want to avoid too complex UI.  If we decide
+						// to add that flag later, we'll need to remove this option, since it should just be "mirror
+						// curve" plus "Linear".
+						float InverseNormalizedScore = (1.0f - NormalizedScoreForEquation);
+						WeightedScore = WeightValue * InverseNormalizedScore;
+						break;
+					}
+
+					case EEnvTestScoreEquation::Square:
+						WeightedScore = WeightValue * (NormalizedScoreForEquation * NormalizedScoreForEquation);
+						break;
+
+					case EEnvTestScoreEquation::Constant:
+						// I know, it's not "constant".  It's "Constant, or zero".  The tooltip should explain that.
+						WeightedScore = (NormalizedScoreForEquation > 0) ? WeightValue : 0.0f;
+						break;
+						
+					default:
+						break;
+				}
+			}
+			else
+			{
+				TestValue = 0.0f;
+				WeightedScore = 0.0f;
+			}
+
+#if USE_EQS_DEBUGGER
+			DetailInfo->TestWeightedScores[QueryInstance.CurrentTest] = WeightedScore;
+#endif
+			QueryInstance.Items[ItemIndex].Score += WeightedScore;
+		}
+	}
+}
+
+bool UEnvQueryTest::IsContextPerItem(TSubclassOf<class UEnvQueryContext> CheckContext) const
+{
+	return CheckContext == UEnvQueryContext_Item::StaticClass();
+}
+
+FVector UEnvQueryTest::GetItemLocation(struct FEnvQueryInstance& QueryInstance, int32 ItemIndex) const
+{
+	return QueryInstance.ItemTypeVectorCDO ?
+		QueryInstance.ItemTypeVectorCDO->GetLocation(QueryInstance.RawData.GetTypedData() + QueryInstance.Items[ItemIndex].DataOffset) :
+		FVector::ZeroVector;
+}
+
+FRotator UEnvQueryTest::GetItemRotation(struct FEnvQueryInstance& QueryInstance, int32 ItemIndex) const
+{
+	return QueryInstance.ItemTypeVectorCDO ?
+		QueryInstance.ItemTypeVectorCDO->GetRotation(QueryInstance.RawData.GetTypedData() + QueryInstance.Items[ItemIndex].DataOffset) :
+		FRotator::ZeroRotator;
+}
+
+AActor* UEnvQueryTest::GetItemActor(struct FEnvQueryInstance& QueryInstance, int32 ItemIndex) const
+{
+	return QueryInstance.ItemTypeActorCDO ?
+		QueryInstance.ItemTypeActorCDO->GetActor(QueryInstance.RawData.GetTypedData() + QueryInstance.Items[ItemIndex].DataOffset) :
+		NULL;
+}
+
+FString UEnvQueryTest::GetDescriptionTitle() const
+{
+	return UEnvQueryTypes::GetShortTypeName(this).ToString();
+}
+
+FText UEnvQueryTest::GetDescriptionDetails() const
+{
+	return FText::GetEmpty();
+}
+
+void UEnvQueryTest::PostLoad()
+{
+	Super::PostLoad();
+
+	if (bFormatUpdated)
+	{
+		// Backwards compatibility already handled.
+		return;
+	}
+
+	switch (Condition)
+	{
+		case EEnvTestCondition::AtLeast:
+			FloatFilterMin = FloatFilter;
+			FilterType = EEnvTestFilterType::Minimum;
+			break;
+		
+		case EEnvTestCondition::UpTo:
+			FloatFilterMax = FloatFilter;
+			FilterType = EEnvTestFilterType::Maximum;
+			break;
+		
+		case EEnvTestCondition::NoCondition:
+			TestPurpose = EEnvTestPurpose::Score;
+			break;
+
+		case EEnvTestCondition::Match:
+			FilterType = EEnvTestFilterType::Match;
+			break;
+
+		default:
+			UE_LOG(LogEQS, Error, TEXT("Invalid Condition type in UEnvQueryTest::PostLoad!"));
+			break;
+	}
+
+	if ((WeightModifier != EEnvTestWeight::Skip) && !Weight.IsNamedParam() && (Weight.Value == 0.f))
+	{	// We ought to be skipping the test if weight is 0!  Zero makes it pointless!
+		WeightModifier = EEnvTestWeight::Skip;
+	}
+
+	if (WeightModifier != EEnvTestWeight::Skip)
+	{
+		if (!bWorkOnFloatValues)
+		{	// Working on booleans, so MUST use constant value!
+			WeightModifier = EEnvTestWeight::Constant;
+		}
+	}
+
+	switch (WeightModifier)
+	{
+		case EEnvTestWeight::None:
+			ScoringEquation = EEnvTestScoreEquation::Linear;
+			break;
+
+		case EEnvTestWeight::Constant:
+			ScoringEquation = EEnvTestScoreEquation::Constant;
+			break;
+
+		case EEnvTestWeight::Inverse:
+			ScoringEquation = EEnvTestScoreEquation::InverseLinear;
+			break;
+
+		case EEnvTestWeight::Square:
+			ScoringEquation = EEnvTestScoreEquation::Square;
+			break;
+
+		case EEnvTestWeight::Absolute:
+			UE_LOG(LogEQS, Error, TEXT("Absolute weight is no longer supported!  Sticking with default weighting."));
+			break;
+
+		case EEnvTestWeight::Skip:
+			TestPurpose = EEnvTestPurpose::Filter;
+			if (Condition == EEnvTestCondition::NoCondition)
+			{
+				UE_LOG(LogEQS, Error, TEXT("Test was set to neither filter nor score!  Now setting to filter!"));	
+			}
+			break;
+
+		default:
+			UE_LOG(LogEQS, Error, TEXT("Invalid Weight Modifier type in UEnvQueryTest::PostLoad!"));
+			break;
+	}
+
+	bFormatUpdated = true;
+}
+
+#if WITH_EDITOR
+void UEnvQueryTest::PostEditChangeProperty(FPropertyChangedEvent& PropertyChangedEvent) 
+{
+	Super::PostEditChangeProperty(PropertyChangedEvent);
+	UEnvQueryManager::NotifyAssetUpdate(NULL);
+}
+#endif //WITH_EDITOR
+
+FText UEnvQueryTest::DescribeFloatTestParams() const
+{
+	FText ParamDesc;
+
+	if (Condition != EEnvTestCondition::NoCondition)
+	{
+		FFormatNamedArguments Args;
+		Args.Add(TEXT("Filter"), FText::FromString(UEnvQueryTypes::DescribeFloatParam(FloatFilter)));
+
+		FText InvalidConditionDesc = LOCTEXT("InvalidCondition", "invalid condition");
+		ParamDesc =
+			(Condition == EEnvTestCondition::AtLeast) ? FText::Format(LOCTEXT("AtLeastFiltered", "at least {Filter}"), Args) :
+			(Condition == EEnvTestCondition::UpTo) ? FText::Format(LOCTEXT("UpToFiltered", "up to {Filter}"), Args) :
+			InvalidConditionDesc;
+	}
+
+	FFormatNamedArguments Args;
+	Args.Add(TEXT("ParmDesc"), ParamDesc);
+
+// 	if (ScoringEquation == EEnvTestScoreEquation::Skip)
+	if (!IsScoring())
+	{
+		ParamDesc = FText::Format(LOCTEXT("ParmDescWithDontScore", "{ParmDesc}, don't score"), Args);
+	}
+	else if (ScoringEquation == EEnvTestScoreEquation::Constant)
+	{
+		FText WeightDesc = Weight.IsNamedParam() ? FText::FromName(Weight.ParamName) : FText::Format( LOCTEXT("ConstantWeightDescPattern", "x{0}"), FText::AsNumber(FMath::Abs(Weight.Value)) );
+		Args.Add(TEXT("WeightDesc"), WeightDesc);
+		ParamDesc = FText::Format(LOCTEXT("ParmDescWithConstantWeight", "{ParmDesc}, constant weight [{WeightDesc}]"), Args);
+	}
+	else if (Weight.IsNamedParam())
+	{
+		Args.Add(TEXT("WeightName"), FText::FromName(Weight.ParamName));
+		ParamDesc = FText::Format(LOCTEXT("ParmDescWithWeight", "{ParmDesc}, weight: {WeightName}"), Args);
+	}
+	else
+	{
+		FNumberFormattingOptions NumberFormattingOptions;
+		NumberFormattingOptions.MaximumFractionalDigits = 2;
+
+		Args.Add(TEXT("WeightCondition"), (Weight.Value > 0) ? LOCTEXT("Greater", "greater") : LOCTEXT("Lesser", "lesser"));
+		Args.Add(TEXT("WeightValue"), FText::AsNumber(FMath::Abs(Weight.Value), &NumberFormattingOptions));
+		ParamDesc = FText::Format(LOCTEXT("DescriptionPreferWeightValue", "{ParmDesc}, prefer {WeightCondition} [x{WeightValue}]"), Args);
+	}
+
+	return ParamDesc;
+}
+
+FText UEnvQueryTest::DescribeBoolTestParams(const FString& ConditionDesc) const
+{
+	FText ParamDesc;
+
+	if (Condition != EEnvTestCondition::NoCondition)
+	{
+		FFormatNamedArguments Args;
+		Args.Add(TEXT("ConditionDesc"), FText::FromString(ConditionDesc));
+
+		if(BoolFilter.IsNamedParam())
+		{
+			Args.Add(TEXT("Filter"), FText::FromString(UEnvQueryTypes::DescribeBoolParam(BoolFilter)));
+		}
+		else
+		{
+			Args.Add(TEXT("Filter"), BoolFilter.Value ? FText::GetEmpty() : LOCTEXT("NotWithSpace", "not "));
+		}
+		FText InvalidConditionDesc = LOCTEXT("InvalidCondition", "invalid condition");
+		ParamDesc =
+			(Condition == EEnvTestCondition::Match) ?
+				(BoolFilter.IsNamedParam() ?
+					FText::Format(LOCTEXT("RequiresCondition", "require {ConditionDesc}: {Filter}"), Args) :
+					FText::Format(LOCTEXT("RequiresFilter", "require {Filter}{ConditionDesc}%s"), Args)
+				) :
+			InvalidConditionDesc;
+	}
+
+	FFormatNamedArguments Args;
+	Args.Add(TEXT("ParmDesc"), ParamDesc);
+
+// 	if (ScoringEquation == EEnvTestScoreEquation::Skip)
+	if (!IsScoring())
+	{
+		ParamDesc = FText::Format(LOCTEXT("ParmDescWithDontScore", "{ParmDesc}, don't score"), Args);
+	}
+	else if (ScoringEquation == EEnvTestScoreEquation::Constant)
+	{
+		FText WeightDesc = Weight.IsNamedParam() ? FText::FromName(Weight.ParamName) : FText::Format(LOCTEXT("ConstantWeightDescPattern", "x{0}"), FText::AsNumber(FMath::Abs(Weight.Value)));
+		Args.Add(TEXT("WeightDesc"), WeightDesc);
+		ParamDesc = FText::Format(LOCTEXT("ParmDescWithConstantWeight", "{ParmDesc}, constant weight [{WeightDesc}]"), Args);
+	}
+	else if (Weight.IsNamedParam())
+	{
+		Args.Add(TEXT("WeightName"), FText::FromName(Weight.ParamName));
+		ParamDesc = FText::Format(LOCTEXT("ParmDescWithWeight", "{ParmDesc}, weight: {WeightName}"), Args);
+	}
+	else
+	{
+		FNumberFormattingOptions NumberFormattingOptions;
+		NumberFormattingOptions.MaximumFractionalDigits = 2;
+
+		Args.Add(TEXT("ConditionDesc"), FText::FromString(ConditionDesc));
+		Args.Add(TEXT("WeightCondition"), (Weight.Value > 0) ? FText::GetEmpty() : LOCTEXT("NotWithSpace", "not "));
+		Args.Add(TEXT("WeightValue"), FText::AsNumber(FMath::Abs(Weight.Value), &NumberFormattingOptions));
+		ParamDesc = FText::Format(LOCTEXT("ParmDescWithPreference", "{ParmDesc}, prefer {WeightCondition}{ConditionDesc} [x{WeightValue}]"), Args);
+	}
+
+	return ParamDesc;
+}
+
+void UEnvQueryTest::SetWorkOnFloatValues(bool bWorkOnFloats)
+{
+	bWorkOnFloatValues = bWorkOnFloats;
+	// Make sure FilterType is set to a valid value.
+	if (bWorkOnFloats)
+	{
+		if (FilterType == EEnvTestFilterType::Match)
+		{
+			FilterType = EEnvTestFilterType::Range;
+		}
+	}
+	else
+	{
+		if (FilterType != EEnvTestFilterType::Match)
+		{
+			FilterType = EEnvTestFilterType::Match;
+		}
+
+		// Scoring MUST be Constant for boolean tests.
+		ScoringEquation = EEnvTestScoreEquation::Constant;
+// 		bMirrorNormalizedScore = false;
+	}
+}
+
+#undef LOCTEXT_NAMESPACE
diff --git a/Engine/Source/Runtime/AIModule/Private/EnvironmentQuery/Generators/EnvQueryGenerator_Composite.cpp b/Engine/Source/Runtime/AIModule/Private/EnvironmentQuery/Generators/EnvQueryGenerator_Composite.cpp
new file mode 100644
--- /dev/null
+++ b/Engine/Source/Runtime/AIModule/Private/EnvironmentQuery/Generators/EnvQueryGenerator_Composite.cpp
@@ -0,0 +1,38 @@
+// Copyright 1998-2014 Epic Games, Inc. All Rights Reserved.
+
+#include "AIModulePrivate.h"
+#include "EnvironmentQuery/Generators/EnvQueryGenerator_Composite.h"
+
+#define LOCTEXT_NAMESPACE "EnvQueryGenerator"
+
+UEnvQueryGenerator_Composite::UEnvQueryGenerator_Composite(const class FPostConstructInitializeProperties& PCIP) : Super(PCIP)
+{
+	GenerateDelegate.BindUObject(this, &UEnvQueryGenerator_Composite::GenerateItems);
+}
+
+void UEnvQueryGenerator_Composite::GenerateItems(struct FEnvQueryInstance& QueryInstance)
+{
+	for (int32 Idx = 0; Idx < Generators.Num(); Idx++)
+	{
+		if (Generators[Idx] && Generators[Idx]->GenerateDelegate.IsBound())
+		{
+			Generators[Idx]->GenerateDelegate.Execute(QueryInstance);
+		}
+	}
+}
+
+FText UEnvQueryGenerator_Composite::GetDescriptionTitle() const
+{
+	FText Desc = Super::GetDescriptionTitle();
+	for (int32 Idx = 0; Idx < Generators.Num(); Idx++)
+	{
+		if (Generators[Idx])
+		{
+			Desc = FText::Format(LOCTEXT("DescTitleExtention", "{0}\n  {1}"), Desc, Generators[Idx]->GetDescriptionTitle());
+		}
+	}
+
+	return Desc;
+};
+
+#undef LOCTEXT_NAMESPACE
\ No newline at end of file
diff --git a/Engine/Source/Runtime/AIModule/Private/EnvironmentQuery/Generators/EnvQueryGenerator_OnCircle.cpp b/Engine/Source/Runtime/AIModule/Private/EnvironmentQuery/Generators/EnvQueryGenerator_OnCircle.cpp
new file mode 100644
--- /dev/null
+++ b/Engine/Source/Runtime/AIModule/Private/EnvironmentQuery/Generators/EnvQueryGenerator_OnCircle.cpp
@@ -0,0 +1,302 @@
+// Copyright 1998-2014 Epic Games, Inc. All Rights Reserved.
+
+#include "AIModulePrivate.h"
+#include "EnvironmentQuery/Contexts/EnvQueryContext_Querier.h"
+#include "EnvironmentQuery/Items/EnvQueryItemType_Point.h"
+#include "EnvironmentQuery/Generators/EnvQueryGenerator_OnCircle.h"
+
+#define LOCTEXT_NAMESPACE "EnvQueryGenerator"
+
+UEnvQueryGenerator_OnCircle::UEnvQueryGenerator_OnCircle(const class FPostConstructInitializeProperties& PCIP) 
+	: Super(PCIP)
+{
+	GenerateDelegate.BindUObject(this, &UEnvQueryGenerator_OnCircle::GenerateItems);
+
+	CircleCenter = UEnvQueryContext_Querier::StaticClass();
+	ItemType = UEnvQueryItemType_Point::StaticClass();
+	Radius.Value = 1000.0f;
+	ItemSpacing.Value = 50.0f;
+	ArcDirection.DirMode = EEnvDirection::TwoPoints;
+	ArcDirection.LineFrom = UEnvQueryContext_Querier::StaticClass();
+	ArcDirection.Rotation = UEnvQueryContext_Querier::StaticClass();
+	Angle.Value = 360.f;
+	AngleRadians = FMath::DegreesToRadians(360.f);
+
+	TraceData.bCanProjectDown = false;
+	TraceData.bCanDisableTrace = false;
+	TraceData.TraceMode = EEnvQueryTrace::Navigation;
+
+	ProjectionData.TraceMode = EEnvQueryTrace::None;
+}
+
+void UEnvQueryGenerator_OnCircle::PostLoad()
+{
+	Super::PostLoad();
+
+	if (Angle.Value <= 0)
+	{
+		Angle.Value = 360.f;
+	}
+	AngleRadians = FMath::DegreesToRadians(Angle.Value);
+}
+
+FVector UEnvQueryGenerator_OnCircle::CalcDirection(FEnvQueryInstance& QueryInstance) const
+{
+	AActor* Querier = Cast<AActor>(QueryInstance.Owner.Get());
+	FVector Direction = Querier->GetActorRotation().Vector();
+
+	if (bDefineArc)
+	{
+		if (ArcDirection.DirMode == EEnvDirection::TwoPoints)
+		{
+			TArray<FVector> Start;
+			TArray<FVector> End;
+			QueryInstance.PrepareContext(ArcDirection.LineFrom, Start);
+			QueryInstance.PrepareContext(ArcDirection.LineTo, End);
+
+			if (Start.Num() > 0 && End.Num() > 0)
+			{
+				Direction = (End[0] - Start[0]).SafeNormal();
+			}
+			else
+			{
+				UE_VLOG(Querier, LogEQS, Warning, TEXT("UEnvQueryGenerator_OnCircle::CalcDirection failed to calc direction in %s. Using querier facing."), *QueryInstance.QueryName);
+			}
+		}
+		else
+		{
+			TArray<FRotator> Rot;
+			QueryInstance.PrepareContext(ArcDirection.Rotation, Rot);
+
+			if (Rot.Num() > 0)
+			{
+				Direction = Rot[0].Vector();
+			}
+			else
+			{
+				UE_VLOG(Querier, LogEQS, Warning, TEXT("UEnvQueryGenerator_OnCircle::CalcDirection failed to calc direction in %s. Using querier facing."), *QueryInstance.QueryName);
+			}
+		}
+	}
+
+	return Direction;
+}
+
+void UEnvQueryGenerator_OnCircle::GenerateItems(FEnvQueryInstance& QueryInstance)
+{
+	float AngleDegree = 360.f;
+	float RadiusValue = 0.f;
+	float ItemSpace = 1.0f;
+
+	if (QueryInstance.GetParamValue(Angle, AngleDegree, TEXT("Angle")) == false
+		|| QueryInstance.GetParamValue(Radius, RadiusValue, TEXT("Radius")) == false
+		|| QueryInstance.GetParamValue(ItemSpacing, ItemSpace, TEXT("ItemSpacing")) == false
+		|| AngleDegree <= 0.f || AngleDegree > 360.f
+		|| RadiusValue <= 0.f)
+	{
+		return;
+	}
+
+	AngleRadians = FMath::DegreesToRadians(Angle.Value);
+
+	// first generate points on a circle
+	const float CircumferenceLength = 2.f * PI * RadiusValue;
+	const float ArcAnglePercentage = Angle.Value / 360.f;
+	const float ArcLength = CircumferenceLength * ArcAnglePercentage;
+	const int32 StepsCount = FMath::CeilToInt(ArcLength / ItemSpace) + 1;
+	const float AngleStep = AngleDegree / (StepsCount - 1);
+
+	FVector StartDirection = CalcDirection(QueryInstance);
+	TArray<FVector> CenterLocationCandidates;
+	QueryInstance.PrepareContext(CircleCenter, CenterLocationCandidates);
+
+	FVector CenterLocation(0);
+	if (CenterLocationCandidates.Num() > 0)
+	{
+		CenterLocation = CenterLocationCandidates[0];
+	}
+	else
+	{
+		AActor* Querier = Cast<AActor>(QueryInstance.Owner.Get());
+		if (Querier)
+		{
+			CenterLocation = Querier->GetActorLocation();
+		}
+	}
+
+	StartDirection = StartDirection.RotateAngleAxis(-AngleDegree/2, FVector::UpVector) * RadiusValue;
+
+	TArray<FVector> ItemCandidates;
+	ItemCandidates.AddZeroed(StepsCount);
+	for (int32 Step = 0; Step < StepsCount; ++Step)
+	{
+		ItemCandidates[Step] = CenterLocation + StartDirection.RotateAngleAxis(AngleStep*Step, FVector::UpVector);
+	}
+
+#if WITH_RECAST
+	// @todo this needs to be optimize to batch raycasts
+	const ARecastNavMesh* NavMesh = 
+		(TraceData.TraceMode == EEnvQueryTrace::Navigation) || (ProjectionData.TraceMode == EEnvQueryTrace::Navigation) ?
+		FEQSHelpers::FindNavMeshForQuery(QueryInstance) : NULL;
+
+	if (NavMesh)
+	{
+		NavMesh->BeginBatchQuery();
+	}
+#endif
+
+	if (TraceData.TraceMode == EEnvQueryTrace::Navigation)
+	{
+#if WITH_RECAST
+		if (NavMesh != NULL)
+		{
+			TSharedPtr<const FNavigationQueryFilter> NavigationFilter = UNavigationQueryFilter::GetQueryFilter(NavMesh, TraceData.NavigationFilter);
+
+			for (int32 Step = 0; Step < StepsCount; ++Step)
+			{
+				NavMesh->Raycast(CenterLocation, ItemCandidates[Step], ItemCandidates[Step], NavigationFilter);
+			}
+		}
+#endif
+	}
+	else
+	{
+		FRunTraceSignature TraceFunc;
+		switch (TraceData.TraceShape)
+		{
+		case EEnvTraceShape::Line:		TraceFunc.BindUObject(this, &UEnvQueryGenerator_OnCircle::RunLineTrace); break;
+		case EEnvTraceShape::Sphere:	TraceFunc.BindUObject(this, &UEnvQueryGenerator_OnCircle::RunSphereTrace); break;
+		case EEnvTraceShape::Capsule:	TraceFunc.BindUObject(this, &UEnvQueryGenerator_OnCircle::RunCapsuleTrace); break;
+		case EEnvTraceShape::Box:		TraceFunc.BindUObject(this, &UEnvQueryGenerator_OnCircle::RunBoxTrace); break;
+		default: break;
+		}
+
+		if (TraceFunc.IsBound())
+		{
+			ECollisionChannel TraceCollisionChannel = UEngineTypes::ConvertToCollisionChannel(TraceData.TraceChannel);	
+			FVector TraceExtent(TraceData.ExtentX, TraceData.ExtentY, TraceData.ExtentZ);
+
+			FCollisionQueryParams TraceParams(TEXT("EnvQueryTrace"), TraceData.bTraceComplex);
+			TraceParams.bTraceAsyncScene = true;
+
+			for (int32 Step = 0; Step < StepsCount; ++Step)
+			{
+				TraceFunc.Execute(CenterLocation, ItemCandidates[Step], QueryInstance.World, TraceCollisionChannel, TraceParams, TraceExtent, ItemCandidates[Step]);
+			}
+		}
+	}
+
+#if WITH_RECAST
+	if (NavMesh)
+	{
+		ProjectAndFilterNavPoints(ItemCandidates, NavMesh);
+		NavMesh->FinishBatchQuery();
+	}
+#endif
+
+	for (int32 Step = 0; Step < ItemCandidates.Num(); ++Step)
+	{
+		QueryInstance.AddItemData<UEnvQueryItemType_Point>(ItemCandidates[Step]);
+	}
+}
+
+void UEnvQueryGenerator_OnCircle::RunLineTrace(const FVector& StartPos, const FVector& EndPos, UWorld* World, enum ECollisionChannel Channel, const FCollisionQueryParams& Params, const FVector& Extent, FVector& HitPos)
+{
+	FHitResult OutHit;
+	const bool bHit = World->LineTraceSingle(OutHit, StartPos, EndPos, Channel, Params);
+	HitPos = bHit ? OutHit.ImpactPoint : EndPos;
+};
+
+void UEnvQueryGenerator_OnCircle::RunSphereTrace(const FVector& StartPos, const FVector& EndPos, UWorld* World, enum ECollisionChannel Channel, const FCollisionQueryParams& Params, const FVector& Extent, FVector& HitPos)
+{
+	FHitResult OutHit;
+	const bool bHit = World->SweepSingle(OutHit, StartPos, EndPos, FQuat::Identity, Channel, FCollisionShape::MakeSphere(Extent.X), Params);
+	HitPos = bHit ? OutHit.ImpactPoint : EndPos;
+};
+
+void UEnvQueryGenerator_OnCircle::RunCapsuleTrace(const FVector& StartPos, const FVector& EndPos, UWorld* World, enum ECollisionChannel Channel, const FCollisionQueryParams& Params, const FVector& Extent, FVector& HitPos)
+{
+	FHitResult OutHit;
+	const bool bHit = World->SweepSingle(OutHit, StartPos, EndPos, FQuat::Identity, Channel, FCollisionShape::MakeCapsule(Extent.X, Extent.Z), Params);
+	HitPos = bHit ? OutHit.ImpactPoint : EndPos;
+};
+
+void UEnvQueryGenerator_OnCircle::RunBoxTrace(const FVector& StartPos, const FVector& EndPos, UWorld* World, enum ECollisionChannel Channel, const FCollisionQueryParams& Params, const FVector& Extent, FVector& HitPos)
+{
+	FHitResult OutHit;
+	const bool bHit = World->SweepSingle(OutHit, StartPos, EndPos, FQuat((EndPos - StartPos).Rotation()), Channel, FCollisionShape::MakeBox(Extent), Params);
+	HitPos = bHit ? OutHit.ImpactPoint : EndPos;
+};
+
+FText UEnvQueryGenerator_OnCircle::GetDescriptionTitle() const
+{
+	FFormatNamedArguments Args;
+	Args.Add(TEXT("DescriptionTitle"), Super::GetDescriptionTitle());
+	Args.Add(TEXT("DescribeContext"), UEnvQueryTypes::DescribeContext(CircleCenter));
+
+	return FText::Format(LOCTEXT("DescriptionGenerateCircleAroundContext", "{DescriptionTitle}: generate items on circle around {DescribeContext}"), Args);
+}
+
+FText UEnvQueryGenerator_OnCircle::GetDescriptionDetails() const
+{
+	FFormatNamedArguments Args;
+	Args.Add(TEXT("Radius"), FText::FromString(UEnvQueryTypes::DescribeFloatParam(Radius)));
+	Args.Add(TEXT("ItemSpacing"), FText::FromString(UEnvQueryTypes::DescribeFloatParam(ItemSpacing)));
+	Args.Add(TEXT("TraceData"), TraceData.ToText(FEnvTraceData::Detailed));
+
+	FText Desc = FText::Format(LOCTEXT("OnCircleDescription", "radius: {Radius}, item span: {ItemSpacing}\n{TraceData}"), Args);
+
+	if (bDefineArc)
+	{
+		FFormatNamedArguments ArcArgs;
+		ArcArgs.Add(TEXT("Description"), Desc);
+		ArcArgs.Add(TEXT("AngleValue"), Angle.Value);
+		ArcArgs.Add(TEXT("ArcDirection"), ArcDirection.ToText());
+		Desc = FText::Format(LOCTEXT("DescriptionWithArc", "{Description}\nLimit to {AngleValue} angle both sides on {ArcDirection}"), ArcArgs);
+	}
+
+	FText ProjDesc = ProjectionData.ToText(FEnvTraceData::Brief);
+	if (!ProjDesc.IsEmpty())
+	{
+		FFormatNamedArguments ProjArgs;
+		ProjArgs.Add(TEXT("Description"), Desc);
+		ProjArgs.Add(TEXT("ProjectionDescription"), ProjDesc);
+		Desc = FText::Format(LOCTEXT("OnCircle_DescriptionWithProjection", "{Description}, {ProjectionDescription}"), ProjArgs);
+	}
+
+	return Desc;
+}
+
+#if WITH_EDITOR
+
+void UEnvQueryGenerator_OnCircle::PostEditChangeProperty( struct FPropertyChangedEvent& PropertyChangedEvent) 
+{
+	static const FName NAME_Angle = GET_MEMBER_NAME_CHECKED(UEnvQueryGenerator_OnCircle, Angle);
+	static const FName NAME_Radius = GET_MEMBER_NAME_CHECKED(UEnvQueryGenerator_OnCircle, Radius);
+
+	if (PropertyChangedEvent.Property != NULL)
+	{
+		const FName PropName = PropertyChangedEvent.MemberProperty->GetFName();
+		if (PropName == NAME_Angle)
+		{
+			Angle.Value = FMath::Clamp(Angle.Value, 0.0f, 360.f);
+			AngleRadians = FMath::DegreesToRadians(Angle.Value);
+			bDefineArc = Angle.Value < 360.f && Angle.Value > 0.f;
+		}
+		else if (PropName == NAME_Radius)
+		{
+			if (Radius.Value <= 0.f)
+			{
+				Radius.Value = 100.f;
+			}
+		}
+
+		ItemSpacing.Value = FMath::Max(1.0f, ItemSpacing.Value);
+	}
+
+	Super::PostEditChangeProperty(PropertyChangedEvent);
+}
+
+#endif // WITH_EDITOR
+
+#undef LOCTEXT_NAMESPACE
\ No newline at end of file
diff --git a/Engine/Source/Runtime/AIModule/Private/EnvironmentQuery/Generators/EnvQueryGenerator_PathingGrid.cpp b/Engine/Source/Runtime/AIModule/Private/EnvironmentQuery/Generators/EnvQueryGenerator_PathingGrid.cpp
new file mode 100644
--- /dev/null
+++ b/Engine/Source/Runtime/AIModule/Private/EnvironmentQuery/Generators/EnvQueryGenerator_PathingGrid.cpp
@@ -0,0 +1,170 @@
+// Copyright 1998-2014 Epic Games, Inc. All Rights Reserved.
+
+#include "AIModulePrivate.h"
+#include "EnvironmentQuery/Items/EnvQueryItemType_Point.h"
+#include "EnvironmentQuery/Contexts/EnvQueryContext_Querier.h"
+#include "EnvironmentQuery/Generators/EnvQueryGenerator_PathingGrid.h"
+#include "AI/Navigation/NavFilters/RecastFilter_UseDefaultArea.h"
+
+
+#define LOCTEXT_NAMESPACE "EnvQueryGenerator"
+
+UEnvQueryGenerator_PathingGrid::UEnvQueryGenerator_PathingGrid(const class FPostConstructInitializeProperties& PCIP) : Super(PCIP)
+{
+	GenerateDelegate.BindUObject(this, &UEnvQueryGenerator_PathingGrid::GenerateItems);
+
+	GenerateAround = UEnvQueryContext_Querier::StaticClass();
+	ItemType = UEnvQueryItemType_Point::StaticClass();
+	MaxPathDistance.Value = 100.0f;
+	Density.Value = 10.0f;
+	PathFromContext.Value = true;
+}
+
+void UEnvQueryGenerator_PathingGrid::GenerateItems(struct FEnvQueryInstance& QueryInstance)
+{
+#if WITH_RECAST
+	const ARecastNavMesh* NavMesh = FEQSHelpers::FindNavMeshForQuery(QueryInstance);
+	if (NavMesh == NULL) 
+	{
+		return;
+	}
+
+	float PathDistanceValue = 0.0f;
+	float DensityValue = 0.0f;
+	bool bFromContextValue = true;
+	if (!QueryInstance.GetParamValue(MaxPathDistance, PathDistanceValue, TEXT("MaxPathDistance")) ||
+		!QueryInstance.GetParamValue(Density, DensityValue, TEXT("Density")) ||
+		!QueryInstance.GetParamValue(PathFromContext, bFromContextValue, TEXT("PathFromContext")))
+	{
+		return;
+	}
+
+	const int32 ItemCount = FPlatformMath::TruncToInt((PathDistanceValue * 2.0f / DensityValue) + 1);
+	const int32 ItemCountHalf = ItemCount / 2;
+
+	TArray<FVector> ContextLocations;
+	QueryInstance.PrepareContext(GenerateAround, ContextLocations);
+	QueryInstance.ReserveItemData(ItemCountHalf * ItemCountHalf * ContextLocations.Num());
+
+	TArray<NavNodeRef> NavNodeRefs;
+	NavMesh->BeginBatchQuery();
+
+	int32 DataOffset = 0;
+	for (int32 ContextIndex = 0; ContextIndex < ContextLocations.Num(); ContextIndex++)
+	{
+		// find all node refs in pathing distance
+		FBox AllowedBounds;
+		NavNodeRefs.Reset();
+		FindNodeRefsInPathDistance(NavMesh, ContextLocations[ContextIndex], PathDistanceValue, bFromContextValue, NavNodeRefs, AllowedBounds);
+
+		// cast 2D grid on generated node refs
+		for (int32 IndexX = 0; IndexX < ItemCount; ++IndexX)
+		{
+			for (int32 IndexY = 0; IndexY < ItemCount; ++IndexY)
+			{
+				const FVector TestPoint = ContextLocations[ContextIndex] - FVector(DensityValue * (IndexX - ItemCountHalf), DensityValue * (IndexY - ItemCountHalf), 0);
+				if (!AllowedBounds.IsInsideXY(TestPoint))
+				{
+					continue;
+				}
+
+				// trace line on navmesh, and process all hits with collected node refs
+				TArray<FNavLocation> Hits;
+				NavMesh->ProjectPointMulti(TestPoint, Hits, FVector::ZeroVector, AllowedBounds.Min.Z, AllowedBounds.Max.Z);
+
+				for (int32 HitIndex = 0; HitIndex < Hits.Num(); HitIndex++)
+				{
+					if (IsNavLocationInPathDistance(NavMesh, Hits[HitIndex], NavNodeRefs))
+					{
+						// store generated point
+						QueryInstance.AddItemData<UEnvQueryItemType_Point>(Hits[HitIndex].Location);
+					}
+				}
+			}
+		}
+	}
+
+	NavMesh->FinishBatchQuery();
+#endif // WITH_RECAST
+}
+
+FText UEnvQueryGenerator_PathingGrid::GetDescriptionTitle() const
+{
+	FFormatNamedArguments Args;
+	Args.Add(TEXT("DescriptionTitle"), Super::GetDescriptionTitle());
+	Args.Add(TEXT("DescribeContext"), UEnvQueryTypes::DescribeContext(GenerateAround));
+
+	return FText::Format(LOCTEXT("DescriptionGenerateAroundContext", "{DescriptionTitle}: generate around {DescribeContext}"), Args);
+};
+
+FText UEnvQueryGenerator_PathingGrid::GetDescriptionDetails() const
+{
+	FFormatNamedArguments Args;
+	Args.Add(TEXT("MaxPathDistance"), FText::FromString(UEnvQueryTypes::DescribeFloatParam(MaxPathDistance)));
+	Args.Add(TEXT("Density"), FText::FromString(UEnvQueryTypes::DescribeFloatParam(Density)));
+	Args.Add(TEXT("PathFromContext"), FText::FromString(UEnvQueryTypes::DescribeBoolParam(PathFromContext)));
+
+	return FText::Format(LOCTEXT("DescriptionDetailsPathingGrid", "max distance: {MaxPathDistance}, density: {Density}, path from context: {PathFromContext}"), Args);
+}
+
+#if WITH_RECAST
+#define ENVQUERY_CLUSTER_SEARCH 0
+
+void UEnvQueryGenerator_PathingGrid::FindNodeRefsInPathDistance(const class ARecastNavMesh* NavMesh, const FVector& ContextLocation, float InMaxPathDistance, bool bPathFromContext, TArray<NavNodeRef>& NodeRefs, FBox& NodeRefsBounds) const
+{
+	FBox MyBounds(0);
+
+#if ENVQUERY_CLUSTER_SEARCH
+	const bool bUseBacktracking = !bPathFromContext;
+	NavMesh->GetClustersWithinPathingDistance(ContextLocation, InMaxPathDistance, NodeRefs, bUseBacktracking);
+
+	for (int32 i = 0; i < NodeRefs.Num(); i++)
+	{
+		FBox ClusterBounds;
+		
+		const bool bSuccess = NavMesh->GetClusterBounds(NodeRefs[i], ClusterBounds);
+		if (bSuccess)
+		{
+			MyBounds += ClusterBounds;
+		}
+	}
+#else
+	TSharedPtr<FNavigationQueryFilter> NavFilterInstance = NavigationFilter != NULL 
+		? UNavigationQueryFilter::GetQueryFilter<URecastFilter_UseDefaultArea>(NavMesh)->GetCopy()
+		: NavMesh->GetDefaultQueryFilter()->GetCopy();
+
+	NavFilterInstance->SetBacktrackingEnabled(!bPathFromContext);
+	NavMesh->GetPolysWithinPathingDistance(ContextLocation, InMaxPathDistance, NodeRefs, NavFilterInstance);
+
+	TArray<FVector> PolyVerts;
+	for (int32 i = 0; i < NodeRefs.Num(); i++)
+	{
+		PolyVerts.Reset();
+		
+		const bool bSuccess = NavMesh->GetPolyVerts(NodeRefs[i], PolyVerts);
+		if (bSuccess)
+		{
+			MyBounds += FBox(PolyVerts);
+		}
+	}
+#endif
+
+	NodeRefsBounds = MyBounds;
+}
+
+bool UEnvQueryGenerator_PathingGrid::IsNavLocationInPathDistance(const class ARecastNavMesh* NavMesh,
+		const struct FNavLocation& NavLocation, const TArray<NavNodeRef>& NodeRefs) const
+{
+#if ENVQUERY_CLUSTER_SEARCH
+	const NavNodeRef ClusterRef = NavMesh->GetClusterRef(NavLocation.NodeRef);
+	return NodeRefs.Contains(ClusterRef);
+#else
+	return NodeRefs.Contains(NavLocation.NodeRef);
+#endif
+}
+
+#undef ENVQUERY_CLUSTER_SEARCH
+
+#endif // WITH_RECAST
+
+#undef LOCTEXT_NAMESPACE
\ No newline at end of file
diff --git a/Engine/Source/Runtime/AIModule/Private/EnvironmentQuery/Generators/EnvQueryGenerator_ProjectedPoints.cpp b/Engine/Source/Runtime/AIModule/Private/EnvironmentQuery/Generators/EnvQueryGenerator_ProjectedPoints.cpp
new file mode 100644
--- /dev/null
+++ b/Engine/Source/Runtime/AIModule/Private/EnvironmentQuery/Generators/EnvQueryGenerator_ProjectedPoints.cpp
@@ -0,0 +1,54 @@
+// Copyright 1998-2013 Epic Games, Inc. All Rights Reserved.
+
+#include "AIModulePrivate.h"
+#include "EnvironmentQuery/Generators/EnvQueryGenerator_ProjectedPoints.h"
+
+UEnvQueryGenerator_ProjectedPoints::UEnvQueryGenerator_ProjectedPoints(const class FPostConstructInitializeProperties& PCIP) : Super(PCIP)
+{
+	ProjectionData.SetNavmeshOnly();
+	ProjectionData.bCanProjectDown = true;
+	ProjectionData.bCanDisableTrace = true;
+	ProjectionData.ExtentX = 0.0f;
+}
+
+void UEnvQueryGenerator_ProjectedPoints::ProjectAndFilterNavPoints(TArray<FVector>& Points, const class ANavigationData* NavData)
+{
+	if (ProjectionData.TraceMode == EEnvQueryTrace::Navigation && NavData)
+	{
+		TSharedPtr<const FNavigationQueryFilter> NavigationFilter = UNavigationQueryFilter::GetQueryFilter(NavData, ProjectionData.NavigationFilter);	
+		const ARecastNavMesh* NavMesh = Cast<const ARecastNavMesh>(NavData);
+
+		const bool bProjectInRange = (ProjectionData.ProjectDown != ProjectionData.ProjectUp);
+		if (bProjectInRange)
+		{
+			TArray<FNavLocation> TempPoints;
+
+			if (NavMesh)
+			{
+				for (int32 i = Points.Num() -1; i >= 0; i--)
+				{
+#if WITH_RECAST
+					const bool bProjected = ProjectNavPointInRange(Points[i], ProjectionData.ExtentX, ProjectionData.ProjectDown, ProjectionData.ProjectUp, NavMesh, NavigationFilter, TempPoints);
+					if (!bProjected)
+					{
+						Points.RemoveAt(i);
+					}
+#endif
+				}
+			}
+		}
+		else
+		{
+			const FVector ProjectionExtent(ProjectionData.ExtentX, ProjectionData.ExtentX, ProjectionData.ProjectDown);
+
+			for (int32 i = Points.Num() -1; i >= 0; i--)
+			{
+				const bool bProjected = ProjectNavPointSimple(Points[i], ProjectionExtent, NavData, NavigationFilter);
+				if (!bProjected)
+				{
+					Points.RemoveAt(i);
+				}
+			}
+		}
+	}
+}
diff --git a/Engine/Source/Runtime/AIModule/Private/EnvironmentQuery/Generators/EnvQueryGenerator_SimpleGrid.cpp b/Engine/Source/Runtime/AIModule/Private/EnvironmentQuery/Generators/EnvQueryGenerator_SimpleGrid.cpp
new file mode 100644
--- /dev/null
+++ b/Engine/Source/Runtime/AIModule/Private/EnvironmentQuery/Generators/EnvQueryGenerator_SimpleGrid.cpp
@@ -0,0 +1,113 @@
+// Copyright 1998-2014 Epic Games, Inc. All Rights Reserved.
+
+#include "AIModulePrivate.h"
+#include "EnvironmentQuery/Contexts/EnvQueryContext_Querier.h"
+#include "EnvironmentQuery/Items/EnvQueryItemType_Point.h"
+#include "EnvironmentQuery/Generators/EnvQueryGenerator_SimpleGrid.h"
+
+#define LOCTEXT_NAMESPACE "EnvQueryGenerator"
+
+UEnvQueryGenerator_SimpleGrid::UEnvQueryGenerator_SimpleGrid(const class FPostConstructInitializeProperties& PCIP) : Super(PCIP)
+{
+	GenerateDelegate.BindUObject(this, &UEnvQueryGenerator_SimpleGrid::GenerateItems);
+
+	GenerateAround = UEnvQueryContext_Querier::StaticClass();
+	ItemType = UEnvQueryItemType_Point::StaticClass();
+	Radius.Value = 500.0f;
+	Density.Value = 100.0f;
+}
+
+void UEnvQueryGenerator_SimpleGrid::GenerateItems(struct FEnvQueryInstance& QueryInstance)
+{
+	float RadiusValue = 0.0f;
+	float DensityValue = 0.0f;
+
+	if (!QueryInstance.GetParamValue(Radius, RadiusValue, TEXT("Radius")) ||
+		!QueryInstance.GetParamValue(Density, DensityValue, TEXT("Density")) )
+	{
+		return;
+	}
+
+	const bool bProjectToNavigation = (ProjectionData.TraceMode == EEnvQueryTrace::Navigation);
+#if WITH_RECAST
+	const ARecastNavMesh* NavMesh = bProjectToNavigation ? FEQSHelpers::FindNavMeshForQuery(QueryInstance) : NULL;
+	if (bProjectToNavigation && NavMesh == NULL) 
+	{
+		return;
+	}
+#endif // WITH_RECAST
+
+	const int32 ItemCount = FPlatformMath::TruncToInt((RadiusValue * 2.0f / DensityValue) + 1);
+	const int32 ItemCountHalf = ItemCount / 2;
+
+	TArray<FVector> ContextLocations;
+	QueryInstance.PrepareContext(GenerateAround, ContextLocations);
+
+#if WITH_RECAST
+	const FVector ProjectExtent(ProjectionData.ExtentX, ProjectionData.ExtentX, 0);
+	if (NavMesh)
+	{
+		NavMesh->BeginBatchQuery();
+	}
+#endif // WITH_RECAST
+
+	TArray<FVector> GridPoints;
+	GridPoints.Reserve(ItemCount * ItemCount * ContextLocations.Num());
+
+	for (int32 ContextIndex = 0; ContextIndex < ContextLocations.Num(); ContextIndex++)
+	{
+		for (int32 IndexX = 0; IndexX <= ItemCount; ++IndexX)
+		{
+			for (int32 IndexY = 0; IndexY <= ItemCount; ++IndexY)
+			{
+				const FVector TestPoint = ContextLocations[ContextIndex] - FVector(DensityValue * (IndexX - ItemCountHalf), DensityValue * (IndexY - ItemCountHalf), 0);
+				GridPoints.Add(TestPoint);
+			}
+		}
+	}
+
+#if WITH_RECAST
+	if (NavMesh)
+	{
+		ProjectAndFilterNavPoints(GridPoints, NavMesh);
+		NavMesh->FinishBatchQuery();
+	}
+#endif // WITH_RECAST
+
+	QueryInstance.ReserveItemData(GridPoints.Num());
+	for (int32 i = 0; i < GridPoints.Num(); i++)
+	{
+		QueryInstance.AddItemData<UEnvQueryItemType_Point>(GridPoints[i]);
+	}
+}
+
+FText UEnvQueryGenerator_SimpleGrid::GetDescriptionTitle() const
+{
+	FFormatNamedArguments Args;
+	Args.Add(TEXT("DescriptionTitle"), Super::GetDescriptionTitle());
+	Args.Add(TEXT("DescribeContext"), UEnvQueryTypes::DescribeContext(GenerateAround));
+
+	return FText::Format(LOCTEXT("DescriptionGenerateSimpleGridAroundContext", "{DescriptionTitle}: generate around {DescribeContext}"), Args);
+};
+
+FText UEnvQueryGenerator_SimpleGrid::GetDescriptionDetails() const
+{
+	FFormatNamedArguments Args;
+	Args.Add(TEXT("Radius"), FText::FromString(UEnvQueryTypes::DescribeFloatParam(Radius)));
+	Args.Add(TEXT("Density"), FText::FromString(UEnvQueryTypes::DescribeFloatParam(Density)));
+
+	FText Desc = FText::Format(LOCTEXT("SimpleGridDescription", "radius: {Radius}, density: {Density}"), Args);
+
+	FText ProjDesc = ProjectionData.ToText(FEnvTraceData::Brief);
+	if (!ProjDesc.IsEmpty())
+	{
+		FFormatNamedArguments ProjArgs;
+		ProjArgs.Add(TEXT("Description"), Desc);
+		ProjArgs.Add(TEXT("ProjectionDescription"), ProjDesc);
+		Desc = FText::Format(LOCTEXT("SimpleGridDescriptionWithProjection", "{Description}, {ProjectionDescription}"), ProjArgs);
+	}
+
+	return Desc;
+}
+
+#undef LOCTEXT_NAMESPACE
\ No newline at end of file
diff --git a/Engine/Source/Runtime/AIModule/Private/EnvironmentQuery/Items/EnvQueryItemType.cpp b/Engine/Source/Runtime/AIModule/Private/EnvironmentQuery/Items/EnvQueryItemType.cpp
new file mode 100644
--- /dev/null
+++ b/Engine/Source/Runtime/AIModule/Private/EnvironmentQuery/Items/EnvQueryItemType.cpp
@@ -0,0 +1,36 @@
+// Copyright 1998-2014 Epic Games, Inc. All Rights Reserved.
+
+#include "AIModulePrivate.h"
+#include "EnvironmentQuery/EnvQueryManager.h"
+#include "EnvironmentQuery/Items/EnvQueryItemType.h"
+
+UEnvQueryItemType::UEnvQueryItemType(const class FPostConstructInitializeProperties& PCIP) : Super(PCIP)
+{
+	// register in known types 
+	if (HasAnyFlags(RF_ClassDefaultObject) && !GetClass()->HasAnyClassFlags(CLASS_Abstract))
+	{
+		UEnvQueryManager::RegisteredItemTypes.Add(GetClass());
+	}
+}
+
+void UEnvQueryItemType::FinishDestroy()
+{
+	// unregister from known types 
+	UEnvQueryManager::RegisteredItemTypes.RemoveSingleSwap(GetClass());
+
+	Super::FinishDestroy();
+}
+
+void UEnvQueryItemType::AddBlackboardFilters(struct FBlackboardKeySelector& KeySelector, UObject* FilterOwner) const
+{
+}
+
+bool UEnvQueryItemType::StoreInBlackboard(struct FBlackboardKeySelector& KeySelector, class UBlackboardComponent* Blackboard, const uint8* RawData) const
+{
+	return false;
+}
+
+FString UEnvQueryItemType::GetDescription(const uint8* RawData) const
+{
+	return TEXT("item");
+}
diff --git a/Engine/Source/Runtime/AIModule/Private/EnvironmentQuery/Items/EnvQueryItemType_Actor.cpp b/Engine/Source/Runtime/AIModule/Private/EnvironmentQuery/Items/EnvQueryItemType_Actor.cpp
new file mode 100644
--- /dev/null
+++ b/Engine/Source/Runtime/AIModule/Private/EnvironmentQuery/Items/EnvQueryItemType_Actor.cpp
@@ -0,0 +1,75 @@
+// Copyright 1998-2014 Epic Games, Inc. All Rights Reserved.
+
+#include "AIModulePrivate.h"
+#include "EnvironmentQuery/Items/EnvQueryItemType_Actor.h"
+
+UEnvQueryItemType_Actor::UEnvQueryItemType_Actor(const class FPostConstructInitializeProperties& PCIP) : Super(PCIP)
+{
+	ValueSize = sizeof(FWeakObjectPtr);
+}
+
+AActor* UEnvQueryItemType_Actor::GetValue(const uint8* RawData)
+{
+	FWeakObjectPtr WeakObjPtr = GetValueFromMemory<FWeakObjectPtr>(RawData);
+	return (AActor*)(WeakObjPtr.Get());
+}
+
+void UEnvQueryItemType_Actor::SetValue(uint8* RawData, const AActor* Value)
+{
+	FWeakObjectPtr WeakObjPtr(Value);
+	SetValueInMemory<FWeakObjectPtr>(RawData, WeakObjPtr);
+}
+
+FVector UEnvQueryItemType_Actor::GetLocation(const uint8* RawData) const
+{
+	AActor* MyActor = UEnvQueryItemType_Actor::GetValue(RawData);
+	return MyActor ? MyActor->GetActorLocation() : FVector::ZeroVector;
+}
+
+FRotator UEnvQueryItemType_Actor::GetRotation(const uint8* RawData) const
+{
+	AActor* MyActor = UEnvQueryItemType_Actor::GetValue(RawData);
+	return MyActor ? MyActor->GetActorRotation() : FRotator::ZeroRotator;
+}
+
+AActor* UEnvQueryItemType_Actor::GetActor(const uint8* RawData) const
+{
+	return UEnvQueryItemType_Actor::GetValue(RawData);
+}
+
+void UEnvQueryItemType_Actor::SetContextHelper(struct FEnvQueryContextData& ContextData, const AActor* SingleActor)
+{
+	ContextData.ValueType = UEnvQueryItemType_Actor::StaticClass();
+	ContextData.NumValues = 1;
+	ContextData.RawData.Init(sizeof(FWeakObjectPtr));
+
+	UEnvQueryItemType_Actor::SetValue(ContextData.RawData.GetTypedData(), SingleActor);
+}
+
+void UEnvQueryItemType_Actor::SetContextHelper(struct FEnvQueryContextData& ContextData, const TArray<const AActor*>& MultipleActors)
+{
+	ContextData.ValueType = UEnvQueryItemType_Actor::StaticClass();
+	ContextData.NumValues = MultipleActors.Num();
+	ContextData.RawData.Init(sizeof(FWeakObjectPtr) * MultipleActors.Num());
+
+	uint8* RawData = (uint8*)ContextData.RawData.GetTypedData();
+	for (int32 ActorIndex = 0; ActorIndex < MultipleActors.Num(); ActorIndex++)
+	{
+		UEnvQueryItemType_Actor::SetValue(RawData, MultipleActors[ActorIndex]);
+		RawData += sizeof(FWeakObjectPtr);
+	}
+}
+
+void UEnvQueryItemType_Actor::SetContextHelper(struct FEnvQueryContextData& ContextData, const TArray<AActor*>& MultipleActors)
+{
+	ContextData.ValueType = UEnvQueryItemType_Actor::StaticClass();
+	ContextData.NumValues = MultipleActors.Num();
+	ContextData.RawData.Init(sizeof(FWeakObjectPtr)* MultipleActors.Num());
+
+	uint8* RawData = (uint8*)ContextData.RawData.GetTypedData();
+	for (int32 ActorIndex = 0; ActorIndex < MultipleActors.Num(); ActorIndex++)
+	{
+		UEnvQueryItemType_Actor::SetValue(RawData, MultipleActors[ActorIndex]);
+		RawData += sizeof(FWeakObjectPtr);
+	}
+}
diff --git a/Engine/Source/Runtime/AIModule/Private/EnvironmentQuery/Items/EnvQueryItemType_ActorBase.cpp b/Engine/Source/Runtime/AIModule/Private/EnvironmentQuery/Items/EnvQueryItemType_ActorBase.cpp
new file mode 100644
--- /dev/null
+++ b/Engine/Source/Runtime/AIModule/Private/EnvironmentQuery/Items/EnvQueryItemType_ActorBase.cpp
@@ -0,0 +1,41 @@
+// Copyright 1998-2014 Epic Games, Inc. All Rights Reserved.
+
+#include "AIModulePrivate.h"
+#include "BehaviorTree/BlackboardComponent.h"
+#include "BehaviorTree/Blackboard/BlackboardKeyType_Object.h"
+#include "EnvironmentQuery/Items/EnvQueryItemType_ActorBase.h"
+
+UEnvQueryItemType_ActorBase::UEnvQueryItemType_ActorBase(const class FPostConstructInitializeProperties& PCIP) : Super(PCIP)
+{
+}
+
+void UEnvQueryItemType_ActorBase::AddBlackboardFilters(struct FBlackboardKeySelector& KeySelector, UObject* FilterOwner) const
+{
+	Super::AddBlackboardFilters(KeySelector, FilterOwner);
+	KeySelector.AddObjectFilter(FilterOwner, AActor::StaticClass());
+}
+
+bool UEnvQueryItemType_ActorBase::StoreInBlackboard(struct FBlackboardKeySelector& KeySelector, class UBlackboardComponent* Blackboard, const uint8* RawData) const
+{
+	bool bStored = Super::StoreInBlackboard(KeySelector, Blackboard, RawData);
+	if (!bStored && KeySelector.SelectedKeyType == UBlackboardKeyType_Object::StaticClass())
+	{
+		UObject* MyObject = GetActor(RawData);
+		Blackboard->SetValueAsObject(KeySelector.GetSelectedKeyID(), MyObject);
+
+		bStored = true;
+	}
+
+	return bStored;
+}
+
+FString UEnvQueryItemType_ActorBase::GetDescription(const uint8* RawData) const
+{
+	const AActor* Actor = GetActor(RawData);
+	return GetNameSafe(Actor);
+}
+
+AActor* UEnvQueryItemType_ActorBase::GetActor(const uint8* RawData) const
+{
+	return NULL;
+}
diff --git a/Engine/Source/Runtime/AIModule/Private/EnvironmentQuery/Items/EnvQueryItemType_Direction.cpp b/Engine/Source/Runtime/AIModule/Private/EnvironmentQuery/Items/EnvQueryItemType_Direction.cpp
new file mode 100644
--- /dev/null
+++ b/Engine/Source/Runtime/AIModule/Private/EnvironmentQuery/Items/EnvQueryItemType_Direction.cpp
@@ -0,0 +1,86 @@
+// Copyright 1998-2013 Epic Games, Inc. All Rights Reserved.
+
+#include "AIModulePrivate.h"
+#include "EnvironmentQuery/Items/EnvQueryItemType_Direction.h"
+
+UEnvQueryItemType_Direction::UEnvQueryItemType_Direction(const class FPostConstructInitializeProperties& PCIP) : Super(PCIP)
+{
+	ValueSize = sizeof(FVector);
+}
+
+FVector UEnvQueryItemType_Direction::GetValue(const uint8* RawData)
+{
+	return GetValueFromMemory<FVector>(RawData);
+}
+
+void UEnvQueryItemType_Direction::SetValue(uint8* RawData, const FVector& Value)
+{
+	return SetValueInMemory<FVector>(RawData, Value);
+}
+
+FRotator UEnvQueryItemType_Direction::GetValueRot(const uint8* RawData)
+{
+	return GetValueFromMemory<FVector>(RawData).Rotation();
+}
+
+void UEnvQueryItemType_Direction::SetValueRot(uint8* RawData, const FRotator& Value)
+{
+	return SetValueInMemory<FVector>(RawData, Value.Vector());
+}
+
+FRotator UEnvQueryItemType_Direction::GetRotation(const uint8* RawData) const
+{
+	return UEnvQueryItemType_Direction::GetValueRot(RawData);
+}
+
+FString UEnvQueryItemType_Direction::GetDescription(const uint8* RawData) const
+{
+	FRotator Rot = GetRotation(RawData);
+	return FString::Printf(TEXT("(P=%.0f,Y=%.0f,R=%.0f)"), Rot.Pitch, Rot.Yaw, Rot.Roll);
+}
+
+void UEnvQueryItemType_Direction::SetContextHelper(struct FEnvQueryContextData& ContextData, const FVector& SingleDirection)
+{
+	ContextData.ValueType = UEnvQueryItemType_Direction::StaticClass();
+	ContextData.NumValues = 1;
+	ContextData.RawData.Init(sizeof(FVector));
+
+	UEnvQueryItemType_Direction::SetValue((uint8*)ContextData.RawData.GetTypedData(), SingleDirection);
+}
+
+void UEnvQueryItemType_Direction::SetContextHelper(struct FEnvQueryContextData& ContextData, const FRotator& SingleRotation)
+{
+	ContextData.ValueType = UEnvQueryItemType_Direction::StaticClass();
+	ContextData.NumValues = 1;
+	ContextData.RawData.Init(sizeof(FVector));
+
+	UEnvQueryItemType_Direction::SetValueRot((uint8*)ContextData.RawData.GetTypedData(), SingleRotation);
+}
+
+void UEnvQueryItemType_Direction::SetContextHelper(struct FEnvQueryContextData& ContextData, const TArray<FVector>& MultipleDirections)
+{
+	ContextData.ValueType = UEnvQueryItemType_Direction::StaticClass();
+	ContextData.NumValues = MultipleDirections.Num();
+	ContextData.RawData.Init(sizeof(FVector) * MultipleDirections.Num());
+
+	uint8* RawData = (uint8*)ContextData.RawData.GetTypedData();
+	for (int32 i = 0; i < MultipleDirections.Num(); i++)
+	{
+		UEnvQueryItemType_Direction::SetValue(RawData, MultipleDirections[i]);
+		RawData += sizeof(FVector);
+	}
+}
+
+void UEnvQueryItemType_Direction::SetContextHelper(struct FEnvQueryContextData& ContextData, const TArray<FRotator>& MultipleRotations)
+{
+	ContextData.ValueType = UEnvQueryItemType_Direction::StaticClass();
+	ContextData.NumValues = MultipleRotations.Num();
+	ContextData.RawData.Init(sizeof(FVector) * MultipleRotations.Num());
+
+	uint8* RawData = (uint8*)ContextData.RawData.GetTypedData();
+	for (int32 i = 0; i < MultipleRotations.Num(); i++)
+	{
+		UEnvQueryItemType_Direction::SetValueRot(RawData, MultipleRotations[i]);
+		RawData += sizeof(FVector);
+	}
+}
diff --git a/Engine/Source/Runtime/AIModule/Private/EnvironmentQuery/Items/EnvQueryItemType_Point.cpp b/Engine/Source/Runtime/AIModule/Private/EnvironmentQuery/Items/EnvQueryItemType_Point.cpp
new file mode 100644
--- /dev/null
+++ b/Engine/Source/Runtime/AIModule/Private/EnvironmentQuery/Items/EnvQueryItemType_Point.cpp
@@ -0,0 +1,47 @@
+// Copyright 1998-2014 Epic Games, Inc. All Rights Reserved.
+
+#include "AIModulePrivate.h"
+#include "EnvironmentQuery/Items/EnvQueryItemType_Point.h"
+
+UEnvQueryItemType_Point::UEnvQueryItemType_Point(const class FPostConstructInitializeProperties& PCIP) : Super(PCIP)
+{
+	ValueSize = sizeof(FVector);
+}
+
+FVector UEnvQueryItemType_Point::GetValue(const uint8* RawData)
+{
+	return GetValueFromMemory<FVector>(RawData);
+}
+
+void UEnvQueryItemType_Point::SetValue(uint8* RawData, const FVector& Value)
+{
+	return SetValueInMemory<FVector>(RawData, Value);
+}
+
+FVector UEnvQueryItemType_Point::GetLocation(const uint8* RawData) const
+{
+	return UEnvQueryItemType_Point::GetValue(RawData);
+}
+
+void UEnvQueryItemType_Point::SetContextHelper(struct FEnvQueryContextData& ContextData, const FVector& SinglePoint)
+{
+	ContextData.ValueType = UEnvQueryItemType_Point::StaticClass();
+	ContextData.NumValues = 1;
+	ContextData.RawData.Init(sizeof(FVector));
+
+	UEnvQueryItemType_Point::SetValue((uint8*)ContextData.RawData.GetTypedData(), SinglePoint);
+}
+
+void UEnvQueryItemType_Point::SetContextHelper(struct FEnvQueryContextData& ContextData, const TArray<FVector>& MultiplePoints)
+{
+	ContextData.ValueType = UEnvQueryItemType_Point::StaticClass();
+	ContextData.NumValues = MultiplePoints.Num();
+	ContextData.RawData.Init(sizeof(FVector) * MultiplePoints.Num());
+
+	uint8* RawData = (uint8*)ContextData.RawData.GetTypedData();
+	for (int32 i = 0; i < MultiplePoints.Num(); i++)
+	{
+		UEnvQueryItemType_Point::SetValue(RawData, MultiplePoints[i]);
+		RawData += sizeof(FVector);
+	}
+}
diff --git a/Engine/Source/Runtime/AIModule/Private/EnvironmentQuery/Items/EnvQueryItemType_VectorBase.cpp b/Engine/Source/Runtime/AIModule/Private/EnvironmentQuery/Items/EnvQueryItemType_VectorBase.cpp
new file mode 100644
--- /dev/null
+++ b/Engine/Source/Runtime/AIModule/Private/EnvironmentQuery/Items/EnvQueryItemType_VectorBase.cpp
@@ -0,0 +1,45 @@
+// Copyright 1998-2014 Epic Games, Inc. All Rights Reserved.
+
+#include "AIModulePrivate.h"
+#include "BehaviorTree/BlackboardComponent.h"
+#include "BehaviorTree/Blackboard/BlackboardKeyType_Vector.h"
+#include "EnvironmentQuery/Items/EnvQueryItemType_VectorBase.h"
+
+UEnvQueryItemType_VectorBase::UEnvQueryItemType_VectorBase(const class FPostConstructInitializeProperties& PCIP) : Super(PCIP)
+{
+}
+
+void UEnvQueryItemType_VectorBase::AddBlackboardFilters(struct FBlackboardKeySelector& KeySelector, UObject* FilterOwner) const
+{
+	KeySelector.AddVectorFilter(FilterOwner);
+}
+
+bool UEnvQueryItemType_VectorBase::StoreInBlackboard(struct FBlackboardKeySelector& KeySelector, class UBlackboardComponent* Blackboard, const uint8* RawData) const
+{
+	bool bStored = false;
+	if (KeySelector.SelectedKeyType == UBlackboardKeyType_Vector::StaticClass())
+	{
+		const FVector MyLocation = GetLocation(RawData);
+		Blackboard->SetValueAsVector(KeySelector.GetSelectedKeyID(), MyLocation);
+
+		bStored = true;
+	}
+
+	return bStored;
+}
+
+FString UEnvQueryItemType_VectorBase::GetDescription(const uint8* RawData) const
+{
+	FVector PointLocation = GetLocation(RawData);
+	return FString::Printf(TEXT("(X=%.0f,Y=%.0f,Z=%.0f)"), PointLocation.X, PointLocation.Y, PointLocation.Z);
+}
+
+FVector UEnvQueryItemType_VectorBase::GetLocation(const uint8* RawData) const
+{
+	return FVector::ZeroVector;
+}
+
+FRotator UEnvQueryItemType_VectorBase::GetRotation(const uint8* RawData) const
+{
+	return FRotator::ZeroRotator;
+}
diff --git a/Engine/Source/Runtime/AIModule/Private/EnvironmentQuery/Tests/EnvQueryTest_Distance.cpp b/Engine/Source/Runtime/AIModule/Private/EnvironmentQuery/Tests/EnvQueryTest_Distance.cpp
new file mode 100644
--- /dev/null
+++ b/Engine/Source/Runtime/AIModule/Private/EnvironmentQuery/Tests/EnvQueryTest_Distance.cpp
@@ -0,0 +1,121 @@
+// Copyright 1998-2014 Epic Games, Inc. All Rights Reserved.
+
+#include "AIModulePrivate.h"
+#include "EnvironmentQuery/Contexts/EnvQueryContext_Querier.h"
+#include "EnvironmentQuery/Items/EnvQueryItemType_VectorBase.h"
+#include "EnvironmentQuery/Tests/EnvQueryTest_Distance.h"
+
+namespace
+{
+	FORCEINLINE float CalcDistance3D(const FVector& PosA, const FVector& PosB)
+	{
+		return (PosB - PosA).Size();
+	}
+
+	FORCEINLINE float CalcDistance2D(const FVector& PosA, const FVector& PosB)
+	{
+		return (PosB - PosA).Size2D();
+	}
+
+	FORCEINLINE float CalcDistanceZ(const FVector& PosA, const FVector& PosB)
+	{
+		return PosB.Z - PosA.Z;
+	}
+}
+
+UEnvQueryTest_Distance::UEnvQueryTest_Distance(const class FPostConstructInitializeProperties& PCIP) : Super(PCIP)
+{
+	ExecuteDelegate.BindUObject(this, &UEnvQueryTest_Distance::RunTest);
+
+	DistanceTo = UEnvQueryContext_Querier::StaticClass();
+	Cost = EEnvTestCost::Low;
+	ValidItemType = UEnvQueryItemType_VectorBase::StaticClass();
+}
+
+void UEnvQueryTest_Distance::RunTest(struct FEnvQueryInstance& QueryInstance)
+{
+// 	float ThresholdValue = 0.0f;
+// 	if (!QueryInstance.GetParamValue(FloatFilter, ThresholdValue, TEXT("FloatFilter")))
+// 	{
+// 		return;
+// 	}
+
+	float MinThresholdValue = 0.0f;
+	if (!QueryInstance.GetParamValue(FloatFilterMin, MinThresholdValue, TEXT("FloatFilterMin")))
+	{
+		return;
+	}
+
+	float MaxThresholdValue = 0.0f;
+	if (!QueryInstance.GetParamValue(FloatFilterMax, MaxThresholdValue, TEXT("FloatFilterMax")))
+	{
+		return;
+	}
+
+	// don't support context Item here, it doesn't make any sense
+	TArray<FVector> ContextLocations;
+	if (!QueryInstance.PrepareContext(DistanceTo, ContextLocations))
+	{
+		return;
+	}
+
+	switch (TestMode)
+	{
+	case EEnvTestDistance::Distance3D:	
+		for (FEnvQueryInstance::ItemIterator It(this, QueryInstance); It; ++It)
+		{
+			const FVector ItemLocation = GetItemLocation(QueryInstance, *It);
+			for (int32 iContext = 0; iContext < ContextLocations.Num(); iContext++)
+			{
+				const float Distance = CalcDistance3D(ItemLocation, ContextLocations[iContext]);
+				It.SetScore(TestPurpose, FilterType, Distance, MinThresholdValue, MaxThresholdValue);
+			}
+		}
+		break;
+	case EEnvTestDistance::Distance2D:	
+		for (FEnvQueryInstance::ItemIterator It(this, QueryInstance); It; ++It)
+		{
+			const FVector ItemLocation = GetItemLocation(QueryInstance, *It);
+			for (int32 iContext = 0; iContext < ContextLocations.Num(); iContext++)
+			{
+				const float Distance = CalcDistance2D(ItemLocation, ContextLocations[iContext]);
+				It.SetScore(TestPurpose, FilterType, Distance, MinThresholdValue, MaxThresholdValue);
+			}
+		}
+		break;
+	case EEnvTestDistance::DistanceZ:	
+		for (FEnvQueryInstance::ItemIterator It(this, QueryInstance); It; ++It)
+		{
+			const FVector ItemLocation = GetItemLocation(QueryInstance, *It);
+			for (int32 iContext = 0; iContext < ContextLocations.Num(); iContext++)
+			{
+				const float Distance = CalcDistanceZ(ItemLocation, ContextLocations[iContext]);
+				It.SetScore(TestPurpose, FilterType, Distance, MinThresholdValue, MaxThresholdValue);
+			}
+		}
+		break;
+	default:
+		return;
+	}
+}
+
+FString UEnvQueryTest_Distance::GetDescriptionTitle() const
+{
+	FString ModeDesc;
+	switch (TestMode)
+	{
+	case EEnvTestDistance::Distance3D:	ModeDesc = TEXT(""); break;
+	case EEnvTestDistance::Distance2D:	ModeDesc = TEXT(" 2D"); break;
+	case EEnvTestDistance::DistanceZ:	ModeDesc = TEXT(" Z"); break;
+	default: break;
+	}
+
+	return FString::Printf(TEXT("%s%s: to %s"), 
+		*Super::GetDescriptionTitle(), *ModeDesc,
+		*UEnvQueryTypes::DescribeContext(DistanceTo).ToString());
+}
+
+FText UEnvQueryTest_Distance::GetDescriptionDetails() const
+{
+	return DescribeFloatTestParams();
+}
