Extracting AI code from the engine into a separate module, pass 1 #UE4

diff --git a/Engine/Source/Runtime/AIModule/Private/EnvironmentQuery/Tests/EnvQueryTest_Dot.cpp b/Engine/Source/Runtime/AIModule/Private/EnvironmentQuery/Tests/EnvQueryTest_Dot.cpp
new file mode 100644
--- /dev/null
+++ b/Engine/Source/Runtime/AIModule/Private/EnvironmentQuery/Tests/EnvQueryTest_Dot.cpp
@@ -0,0 +1,189 @@
+// Copyright 1998-2014 Epic Games, Inc. All Rights Reserved.
+
+#include "AIModulePrivate.h"
+#include "EnvironmentQuery/Items/EnvQueryItemType_VectorBase.h"
+#include "EnvironmentQuery/Contexts/EnvQueryContext_Querier.h"
+#include "EnvironmentQuery/Contexts/EnvQueryContext_Item.h"
+#include "EnvironmentQuery/Tests/EnvQueryTest_Dot.h"
+
+UEnvQueryTest_Dot::UEnvQueryTest_Dot(const class FPostConstructInitializeProperties& PCIP) : Super(PCIP)
+{
+	ExecuteDelegate.BindUObject(this, &UEnvQueryTest_Dot::RunTest);
+
+	Cost = EEnvTestCost::Low;
+	ValidItemType = UEnvQueryItemType_VectorBase::StaticClass();
+	LineA.DirMode = EEnvDirection::Rotation;
+	LineA.Rotation = UEnvQueryContext_Querier::StaticClass();
+	LineB.DirMode = EEnvDirection::TwoPoints;
+	LineB.LineFrom = UEnvQueryContext_Querier::StaticClass();
+	LineB.LineTo = UEnvQueryContext_Item::StaticClass();
+}
+
+void UEnvQueryTest_Dot::RunTest(struct FEnvQueryInstance& QueryInstance)
+{
+// 	float ThresholdValue = 0.0f;
+// 	if (!QueryInstance.GetParamValue(FloatFilter, ThresholdValue, TEXT("FloatFilter")))
+// 	{
+// 		return;
+// 	}
+
+	float MinThresholdValue = 0.0f;
+	if (!QueryInstance.GetParamValue(FloatFilterMin, MinThresholdValue, TEXT("FloatFilterMin")))
+	{
+		return;
+	}
+
+	float MaxThresholdValue = 0.0f;
+	if (!QueryInstance.GetParamValue(FloatFilterMax, MaxThresholdValue, TEXT("FloatFilterMax")))
+	{
+		return;
+	}
+
+	// gather all possible directions: for contexts different than Item
+	TArray<FVector> LineADirs;
+	const bool bUpdateLineAPerItem = RequiresPerItemUpdates(LineA.LineFrom, LineA.LineTo, LineA.Rotation, LineA.DirMode == EEnvDirection::Rotation);
+	if (!bUpdateLineAPerItem)
+	{
+		GatherLineDirections(LineADirs, QueryInstance, LineA.LineFrom, LineA.LineTo, LineA.Rotation, LineA.DirMode == EEnvDirection::Rotation);
+		if (LineADirs.Num() == 0)
+		{
+			return;
+		}
+	}
+
+	TArray<FVector> LineBDirs;
+	const bool bUpdateLineBPerItem = RequiresPerItemUpdates(LineB.LineFrom, LineB.LineTo, LineB.Rotation, LineB.DirMode == EEnvDirection::Rotation);
+	if (!bUpdateLineBPerItem)
+	{
+		GatherLineDirections(LineBDirs, QueryInstance, LineB.LineFrom, LineB.LineTo, LineB.Rotation, LineB.DirMode == EEnvDirection::Rotation);
+		if (LineBDirs.Num() == 0)
+		{
+			return;
+		}
+	}
+
+	// loop through all items
+	for (FEnvQueryInstance::ItemIterator It(this, QueryInstance); It; ++It)
+	{
+		// update lines for contexts using current item
+		if (bUpdateLineAPerItem || bUpdateLineBPerItem)
+		{
+			const FVector ItemLocation = (LineA.DirMode == EEnvDirection::Rotation && LineB.DirMode == EEnvDirection::Rotation) ? FVector::ZeroVector : GetItemLocation(QueryInstance, *It);
+			const FRotator ItemRotation = (LineA.DirMode == EEnvDirection::Rotation || LineB.DirMode == EEnvDirection::Rotation) ? GetItemRotation(QueryInstance, *It) : FRotator::ZeroRotator;
+
+			if (bUpdateLineAPerItem)
+			{
+				LineADirs.Reset();
+				GatherLineDirections(LineADirs, QueryInstance, LineA.LineFrom, LineA.LineTo, LineA.Rotation, LineA.DirMode == EEnvDirection::Rotation, ItemLocation, ItemRotation);
+			}
+
+			if (bUpdateLineBPerItem)
+			{
+				LineBDirs.Reset();
+				GatherLineDirections(LineBDirs, QueryInstance, LineB.LineFrom, LineB.LineTo, LineB.Rotation, LineB.DirMode == EEnvDirection::Rotation, ItemLocation, ItemRotation);
+			}
+		}
+
+		// perform test for each line pair
+		for (int32 iLineA = 0; iLineA < LineADirs.Num(); iLineA++)
+		{
+			for (int32 iLineB = 0; iLineB < LineBDirs.Num(); iLineB++)
+			{
+				const float DotValue = FVector::DotProduct(LineADirs[iLineA], LineBDirs[iLineB]);
+				It.SetScore(TestPurpose, FilterType, DotValue, MinThresholdValue, MaxThresholdValue);
+			}
+		}
+	}
+}
+
+void UEnvQueryTest_Dot::GatherLineDirections(TArray<FVector>& Directions, struct FEnvQueryInstance& QueryInstance, const FVector& ItemLocation,
+	TSubclassOf<class UEnvQueryContext> LineFrom, TSubclassOf<class UEnvQueryContext> LineTo)
+{
+	TArray<FVector> ContextLocationFrom;
+	if (IsContextPerItem(LineFrom))
+	{
+		ContextLocationFrom.Add(ItemLocation);
+	}
+	else
+	{
+		QueryInstance.PrepareContext(LineFrom, ContextLocationFrom);
+	}
+	
+	for (int32 i = 0; i < ContextLocationFrom.Num(); i++)
+	{
+		TArray<FVector> ContextLocationTo;
+		if (IsContextPerItem(LineTo))
+		{
+			ContextLocationTo.Add(ItemLocation);
+		}
+		else
+		{
+			QueryInstance.PrepareContext(LineTo, ContextLocationTo);
+		}
+		
+		for (int32 i2 = 0; i2 < ContextLocationTo.Num(); i2++)
+		{
+			const FVector Dir = (ContextLocationTo[i2] - ContextLocationFrom[i]).SafeNormal();
+			Directions.Add(Dir);
+		}
+	}
+}
+
+void UEnvQueryTest_Dot::GatherLineDirections(TArray<FVector>& Directions, struct FEnvQueryInstance& QueryInstance, const FRotator& ItemRotation,
+	TSubclassOf<class UEnvQueryContext> LineDirection)
+{
+	TArray<FRotator> ContextRotations;
+	if (IsContextPerItem(LineDirection))
+	{
+		ContextRotations.Add(ItemRotation);
+	}
+	else
+	{
+		QueryInstance.PrepareContext(LineDirection, ContextRotations);
+	}
+
+	for (int32 i = 0; i < ContextRotations.Num(); i++)
+	{
+		const FVector Dir = ContextRotations[i].Vector();
+		Directions.Add(Dir);
+	}
+}
+
+void UEnvQueryTest_Dot::GatherLineDirections(TArray<FVector>& Directions, struct FEnvQueryInstance& QueryInstance,
+	TSubclassOf<class UEnvQueryContext> LineFrom, TSubclassOf<class UEnvQueryContext> LineTo, TSubclassOf<class UEnvQueryContext> LineDirection, bool bUseDirectionContext,
+	const FVector& ItemLocation, const FRotator& ItemRotation)
+{
+	if (bUseDirectionContext)
+	{
+		GatherLineDirections(Directions, QueryInstance, ItemRotation, LineDirection);
+	}
+	else
+	{
+		GatherLineDirections(Directions, QueryInstance, ItemLocation, LineFrom, LineTo);
+	}
+}
+
+bool UEnvQueryTest_Dot::RequiresPerItemUpdates(TSubclassOf<class UEnvQueryContext> LineFrom, TSubclassOf<class UEnvQueryContext> LineTo, TSubclassOf<class UEnvQueryContext> LineDirection, bool bUseDirectionContext)
+{
+	bool bRequirePerItemUpdate = false;
+	if (bUseDirectionContext)
+	{
+		bRequirePerItemUpdate = IsContextPerItem(LineDirection);
+	}
+	else
+	{
+		bRequirePerItemUpdate = IsContextPerItem(LineFrom) || IsContextPerItem(LineTo);
+	}
+
+	return bRequirePerItemUpdate;
+}
+
+FString UEnvQueryTest_Dot::GetDescriptionTitle() const
+{
+	return FString::Printf(TEXT("%s: %s and %s"), *Super::GetDescriptionTitle(), *LineA.ToText().ToString(), *LineB.ToText().ToString());
+}
+
+FText UEnvQueryTest_Dot::GetDescriptionDetails() const
+{
+	return DescribeFloatTestParams();
+}
diff --git a/Engine/Source/Runtime/AIModule/Private/EnvironmentQuery/Tests/EnvQueryTest_Pathfinding.cpp b/Engine/Source/Runtime/AIModule/Private/EnvironmentQuery/Tests/EnvQueryTest_Pathfinding.cpp
new file mode 100644
--- /dev/null
+++ b/Engine/Source/Runtime/AIModule/Private/EnvironmentQuery/Tests/EnvQueryTest_Pathfinding.cpp
@@ -0,0 +1,236 @@
+// Copyright 1998-2014 Epic Games, Inc. All Rights Reserved.
+
+#include "AIModulePrivate.h"
+#include "AI/Navigation/NavigationSystem.h"
+#include "AI/Navigation/NavAgentInterface.h"
+#include "EnvironmentQuery/Contexts/EnvQueryContext_Querier.h"
+#include "EnvironmentQuery/Items/EnvQueryItemType_VectorBase.h"
+#include "EnvironmentQuery/Tests/EnvQueryTest_Pathfinding.h"
+
+#define LOCTEXT_NAMESPACE "EnvQueryGenerator"
+
+UEnvQueryTest_Pathfinding::UEnvQueryTest_Pathfinding(const class FPostConstructInitializeProperties& PCIP) : Super(PCIP)
+{
+	ExecuteDelegate.BindUObject(this, &UEnvQueryTest_Pathfinding::RunTest);
+
+	Context = UEnvQueryContext_Querier::StaticClass();
+	Cost = EEnvTestCost::High;
+	ValidItemType = UEnvQueryItemType_VectorBase::StaticClass();
+	TestMode = EEnvTestPathfinding::PathExist;
+	PathToItem.Value = true;
+	FloatFilter.Value = 1000.0f;
+	DiscardUnreachable.Value = true;
+	HierarchicalPathfinding.Value = true;
+}
+
+void UEnvQueryTest_Pathfinding::RunTest(struct FEnvQueryInstance& QueryInstance)
+{
+	bool bWantsPath = false;
+	bool bPathToItem = false;
+	bool bHierarchical = false;
+	bool bDiscardFailed = false;
+// 	float UseThreshold = 0.0f;
+	float MinThresholdValue = 0.0f;
+	float MaxThresholdValue = 0.0f;
+
+	if (!QueryInstance.GetParamValue(BoolFilter, bWantsPath, TEXT("BoolFilter")) ||
+		!QueryInstance.GetParamValue(PathToItem, bPathToItem, TEXT("PathToItem")) ||
+		!QueryInstance.GetParamValue(HierarchicalPathfinding, bHierarchical, TEXT("HierarchicalPathfinding")) ||
+		!QueryInstance.GetParamValue(DiscardUnreachable, bDiscardFailed, TEXT("DiscardUnreachable")) ||
+// 		!QueryInstance.GetParamValue(FloatFilter, UseThreshold, TEXT("FloatFilter")) ||
+		!QueryInstance.GetParamValue(FloatFilterMax, MaxThresholdValue, TEXT("FloatFilterMax")) ||
+		!QueryInstance.GetParamValue(FloatFilterMin, MinThresholdValue, TEXT("FloatFilterMin"))
+	   )
+
+	{
+		return;
+	}
+
+	UNavigationSystem* NavSys = QueryInstance.World->GetNavigationSystem();
+	ANavigationData* NavData = FindNavigationData(NavSys, QueryInstance.Owner.Get());
+	if (!NavData)
+	{
+		return;
+	}
+
+	TArray<FVector> ContextLocations;
+	if (!QueryInstance.PrepareContext(Context, ContextLocations))
+	{
+		return;
+	}
+
+	EPathFindingMode::Type PFMode(bHierarchical ? EPathFindingMode::Hierarchical : EPathFindingMode::Regular);
+	if (GetWorkOnFloatValues())
+	{
+		FFindPathSignature FindPathFunc;
+		FindPathFunc.BindUObject(this, TestMode == EEnvTestPathfinding::PathLength ?
+			(bPathToItem ? &UEnvQueryTest_Pathfinding::FindPathLengthTo : &UEnvQueryTest_Pathfinding::FindPathLengthFrom) :
+			(bPathToItem ? &UEnvQueryTest_Pathfinding::FindPathCostTo : &UEnvQueryTest_Pathfinding::FindPathCostFrom) );
+
+		NavData->BeginBatchQuery();
+		for (FEnvQueryInstance::ItemIterator It(this, QueryInstance); It; ++It)
+		{
+			const FVector ItemLocation = GetItemLocation(QueryInstance, *It);
+			for (int32 iContext = 0; iContext < ContextLocations.Num(); iContext++)
+			{
+				const float PathValue = FindPathFunc.Execute(ItemLocation, ContextLocations[iContext], PFMode, NavData, NavSys, QueryInstance.Owner.Get());
+				It.SetScore(TestPurpose, FilterType, PathValue, MinThresholdValue, MaxThresholdValue);
+
+				if (bDiscardFailed && PathValue >= BIG_NUMBER)
+				{
+					It.DiscardItem();
+				}
+			}
+		}
+		NavData->FinishBatchQuery();
+	}
+	else
+	{
+		FTestPathSignature TestPathFunc;
+		TestPathFunc.BindUObject(this, (bPathToItem ? &UEnvQueryTest_Pathfinding::TestPathTo : &UEnvQueryTest_Pathfinding::TestPathFrom));
+	
+		NavData->BeginBatchQuery();
+		for (FEnvQueryInstance::ItemIterator It(this, QueryInstance); It; ++It)
+		{
+			const FVector ItemLocation = GetItemLocation(QueryInstance, *It);
+			for (int32 iContext = 0; iContext < ContextLocations.Num(); iContext++)
+			{
+				const bool bFoundPath = TestPathFunc.Execute(ItemLocation, ContextLocations[iContext], PFMode, NavData, NavSys, QueryInstance.Owner.Get());
+				It.SetScore(TestPurpose, FilterType, bFoundPath, bWantsPath);
+			}
+		}
+		NavData->FinishBatchQuery();
+	}
+}
+
+FString UEnvQueryTest_Pathfinding::GetDescriptionTitle() const
+{
+	FString ModeDesc[] = { TEXT("PathExist"), TEXT("PathCost"), TEXT("PathLength") };
+
+	FString DirectionDesc = PathToItem.IsNamedParam() ?
+		FString::Printf(TEXT("%s, direction: %s"), *UEnvQueryTypes::DescribeContext(Context).ToString(), *UEnvQueryTypes::DescribeBoolParam(PathToItem)) :
+		FString::Printf(TEXT("%s %s"), PathToItem.Value ? TEXT("from") : TEXT("to"), *UEnvQueryTypes::DescribeContext(Context).ToString());
+
+	return FString::Printf(TEXT("%s: %s"), *ModeDesc[TestMode], *DirectionDesc);
+}
+
+FText UEnvQueryTest_Pathfinding::GetDescriptionDetails() const
+{
+	FText AdditionalDesc;
+	if (HierarchicalPathfinding.IsNamedParam())
+	{
+		FFormatNamedArguments Args;
+		Args.Add(TEXT("PathName"), FText::FromString(HierarchicalPathfinding.ParamName.ToString()));
+		AdditionalDesc = FText::Format(LOCTEXT("HierarchicalPathfindingWithName", "Hierarchical pathfinding: {PathName}"), Args);
+	}
+	else if (HierarchicalPathfinding.Value)
+	{
+		AdditionalDesc = LOCTEXT("UseHierarchicalPathfinding", "Use hierarchical pathfinding");
+	}
+
+	if (DiscardUnreachable.IsNamedParam())
+	{
+		FFormatNamedArguments Args;
+		Args.Add(TEXT("Description"), AdditionalDesc);
+		Args.Add(TEXT("UnreachableParamName"), FText::FromName(DiscardUnreachable.ParamName));
+
+		if(AdditionalDesc.IsEmpty())
+		{
+			AdditionalDesc = FText::Format(LOCTEXT("DiscardUnreachableWithParam", "Discard unreachable: {UnreachableParamName}"), Args);
+		}
+		else
+		{
+			AdditionalDesc = FText::Format(LOCTEXT("DescWithUnreachableAndParam", "{Description}, discard unreachable: {UnreachableParamName}"), Args);
+		}
+	}
+	else if (DiscardUnreachable.Value)
+	{
+		if(AdditionalDesc.IsEmpty())
+		{
+			AdditionalDesc = LOCTEXT("DiscardUnreachable", "Discard unreachable");
+		}
+		else
+		{
+			FFormatNamedArguments Args;
+			Args.Add(TEXT("Description"), AdditionalDesc);
+			AdditionalDesc = FText::Format(LOCTEXT("DescWithUnreachable", "{Description}, discard unreachable"), Args);
+		}
+	}
+
+	FFormatNamedArguments Args;
+	Args.Add(TEXT("Description"), AdditionalDesc);
+
+	if (GetWorkOnFloatValues())
+	{
+		Args.Add(TEXT("FloatParamsDescription"), DescribeFloatTestParams());
+		AdditionalDesc = FText::Format(LOCTEXT("DescriptionWithDescribedFloatParams", "{Description}\n{FloatParamsDescription}"), Args);
+	}
+	else
+	{
+		Args.Add(TEXT("BoolParamsDescription"), DescribeBoolTestParams("existing path"));
+		AdditionalDesc = FText::Format(LOCTEXT("DescriptionWithDescribedBoolParams", "{Description}\n{BoolParamsDescription}"), Args);
+	}
+
+	return AdditionalDesc;
+}
+
+#if WITH_EDITOR
+void UEnvQueryTest_Pathfinding::PostEditChangeProperty(struct FPropertyChangedEvent& PropertyChangedEvent)
+{
+	Super::PostEditChangeProperty(PropertyChangedEvent);
+	if (PropertyChangedEvent.Property && PropertyChangedEvent.Property->GetFName() == GET_MEMBER_NAME_CHECKED(UEnvQueryTest_Pathfinding,TestMode))
+	{
+		SetWorkOnFloatValues(TestMode != EEnvTestPathfinding::PathExist);
+		Condition = EEnvTestCondition::NoCondition;
+	}
+}
+#endif
+
+bool UEnvQueryTest_Pathfinding::TestPathFrom(const FVector& ItemPos, const FVector& ContextPos, EPathFindingMode::Type& Mode, const ANavigationData* NavData, UNavigationSystem* NavSys, const UObject* PathOwner)
+{
+	const bool bPathExists = NavSys->TestPathSync(FPathFindingQuery(PathOwner, NavData, ItemPos, ContextPos), Mode);
+	return bPathExists;
+}
+
+bool UEnvQueryTest_Pathfinding::TestPathTo(const FVector& ItemPos, const FVector& ContextPos, EPathFindingMode::Type& Mode, const ANavigationData* NavData, UNavigationSystem* NavSys, const UObject* PathOwner)
+{
+	const bool bPathExists = NavSys->TestPathSync(FPathFindingQuery(PathOwner, NavData, ContextPos, ItemPos), Mode);
+	return bPathExists;
+}
+
+float UEnvQueryTest_Pathfinding::FindPathCostFrom(const FVector& ItemPos, const FVector& ContextPos, EPathFindingMode::Type& Mode, const ANavigationData* NavData, UNavigationSystem* NavSys, const UObject* PathOwner)
+{
+	FPathFindingResult Result = NavSys->FindPathSync(FPathFindingQuery(PathOwner, NavData, ItemPos, ContextPos), Mode);
+	return (Result.IsSuccessful()) ? Result.Path->GetCost() : BIG_NUMBER;
+}
+
+float UEnvQueryTest_Pathfinding::FindPathCostTo(const FVector& ItemPos, const FVector& ContextPos, EPathFindingMode::Type& Mode, const ANavigationData* NavData, UNavigationSystem* NavSys, const UObject* PathOwner)
+{
+	FPathFindingResult Result = NavSys->FindPathSync(FPathFindingQuery(PathOwner, NavData, ContextPos, ItemPos), Mode);
+	return (Result.IsSuccessful()) ? Result.Path->GetCost() : BIG_NUMBER;
+}
+
+float UEnvQueryTest_Pathfinding::FindPathLengthFrom(const FVector& ItemPos, const FVector& ContextPos, EPathFindingMode::Type& Mode, const ANavigationData* NavData, UNavigationSystem* NavSys, const UObject* PathOwner)
+{
+	FPathFindingResult Result = NavSys->FindPathSync(FPathFindingQuery(PathOwner, NavData, ItemPos, ContextPos), Mode);
+	return (Result.IsSuccessful()) ? Result.Path->GetLength() : BIG_NUMBER;
+}
+
+float UEnvQueryTest_Pathfinding::FindPathLengthTo(const FVector& ItemPos, const FVector& ContextPos, EPathFindingMode::Type& Mode, const ANavigationData* NavData, UNavigationSystem* NavSys, const UObject* PathOwner)
+{
+	FPathFindingResult Result = NavSys->FindPathSync(FPathFindingQuery(PathOwner, NavData, ContextPos, ItemPos), Mode);
+	return (Result.IsSuccessful()) ? Result.Path->GetLength() : BIG_NUMBER;
+}
+
+ANavigationData* UEnvQueryTest_Pathfinding::FindNavigationData(UNavigationSystem* NavSys, UObject* Owner)
+{
+	INavAgentInterface* NavAgent = InterfaceCast<INavAgentInterface>(Owner);
+	if (NavAgent)
+	{
+		return NavSys->GetNavDataForProps(*NavAgent->GetNavAgentProperties());
+	}
+
+	return NavSys->GetMainNavData(FNavigationSystem::DontCreate);
+}
+
+#undef LOCTEXT_NAMESPACE
diff --git a/Engine/Source/Runtime/AIModule/Private/EnvironmentQuery/Tests/EnvQueryTest_Trace.cpp b/Engine/Source/Runtime/AIModule/Private/EnvironmentQuery/Tests/EnvQueryTest_Trace.cpp
new file mode 100644
--- /dev/null
+++ b/Engine/Source/Runtime/AIModule/Private/EnvironmentQuery/Tests/EnvQueryTest_Trace.cpp
@@ -0,0 +1,194 @@
+// Copyright 1998-2014 Epic Games, Inc. All Rights Reserved.
+
+#include "AIModulePrivate.h"
+#include "EnvironmentQuery/Contexts/EnvQueryContext_Querier.h"
+#include "EnvironmentQuery/Items/EnvQueryItemType_VectorBase.h"
+#include "EnvironmentQuery/Tests/EnvQueryTest_Trace.h"
+
+#define LOCTEXT_NAMESPACE "EnvQueryGenerator"
+
+UEnvQueryTest_Trace::UEnvQueryTest_Trace(const class FPostConstructInitializeProperties& PCIP) : Super(PCIP)
+{
+	ExecuteDelegate.BindUObject(this, &UEnvQueryTest_Trace::RunTest);
+
+	Cost = EEnvTestCost::High;
+	ValidItemType = UEnvQueryItemType_VectorBase::StaticClass();
+	SetWorkOnFloatValues(false);
+	
+	Context = UEnvQueryContext_Querier::StaticClass();
+	TraceToItem.Value = false;
+	ItemOffsetZ.Value = 0.0f;
+	ContextOffsetZ.Value = 0.0f;
+
+	TraceData.SetGeometryOnly();
+}
+
+void UEnvQueryTest_Trace::RunTest(struct FEnvQueryInstance& QueryInstance)
+{
+	bool bWantsHit = false;
+	bool bTraceToItem = false;
+	float ItemZ = 0.0f;
+	float ContextZ = 0.0f;
+
+	if (!QueryInstance.GetParamValue(BoolFilter, bWantsHit, TEXT("BoolFilter")) ||
+		!QueryInstance.GetParamValue(TraceToItem, bTraceToItem, TEXT("TraceToItem")) ||
+		!QueryInstance.GetParamValue(ItemOffsetZ, ItemZ, TEXT("ItemOffsetZ")) ||
+		!QueryInstance.GetParamValue(ContextOffsetZ, ContextZ, TEXT("ContextOffsetZ"))
+		)
+	{
+		return;
+	}
+
+	TArray<FVector> ContextLocations;
+	if (!QueryInstance.PrepareContext(Context, ContextLocations))
+	{
+		return;
+	}
+
+	FCollisionQueryParams TraceParams(TEXT("EnvQueryTrace"), TraceData.bTraceComplex);
+	TraceParams.bTraceAsyncScene = true;
+
+	TArray<AActor*> IgnoredActors;
+	if (QueryInstance.PrepareContext(Context, IgnoredActors))
+	{
+		TraceParams.AddIgnoredActors(IgnoredActors);
+	}
+	
+	ECollisionChannel TraceCollisionChannel = UEngineTypes::ConvertToCollisionChannel(TraceData.TraceChannel);	
+	FVector TraceExtent(TraceData.ExtentX, TraceData.ExtentY, TraceData.ExtentZ);
+	FRunTraceSignature TraceFunc;
+	switch (TraceData.TraceShape)
+	{
+	case EEnvTraceShape::Line:
+		TraceFunc.BindUObject(this, bTraceToItem ? &UEnvQueryTest_Trace::RunLineTraceTo : &UEnvQueryTest_Trace::RunLineTraceFrom);
+		break;
+
+	case EEnvTraceShape::Box:
+		TraceFunc.BindUObject(this, bTraceToItem ? &UEnvQueryTest_Trace::RunBoxTraceTo : &UEnvQueryTest_Trace::RunBoxTraceFrom);
+		break;
+
+	case EEnvTraceShape::Sphere:
+		TraceFunc.BindUObject(this, bTraceToItem ? &UEnvQueryTest_Trace::RunSphereTraceTo : &UEnvQueryTest_Trace::RunSphereTraceFrom);
+		break;
+
+	case EEnvTraceShape::Capsule:
+		TraceFunc.BindUObject(this, bTraceToItem ? &UEnvQueryTest_Trace::RunCapsuleTraceTo : &UEnvQueryTest_Trace::RunCapsuleTraceFrom);
+		break;
+
+	default:
+		return;
+	}
+
+	for (int32 iContext = 0; iContext < ContextLocations.Num(); iContext++)
+	{
+		ContextLocations[iContext].Z += ContextZ;
+	}
+
+	for (FEnvQueryInstance::ItemIterator It(this, QueryInstance); It; ++It)
+	{
+		const FVector ItemLocation = GetItemLocation(QueryInstance, *It) + FVector(0,0,ItemZ);
+		AActor* ItemActor = GetItemActor(QueryInstance, *It);
+
+		for (int32 iContext = 0; iContext < ContextLocations.Num(); iContext++)
+		{
+			const bool bHit = TraceFunc.Execute(ItemLocation, ContextLocations[iContext], ItemActor, QueryInstance.World, TraceCollisionChannel, TraceParams, TraceExtent);
+			It.SetScore(TestPurpose, FilterType, bHit, bWantsHit);
+		}
+	}
+}
+
+FString UEnvQueryTest_Trace::GetDescriptionTitle() const
+{
+	UEnum* ChannelEnum = FindObject<UEnum>(ANY_PACKAGE, TEXT("ETraceTypeQuery"), true);
+	FString ChannelDesc = ChannelEnum->GetEnumText(TraceData.TraceChannel).ToString();
+
+	FString DirectionDesc = TraceToItem.IsNamedParam() ?
+		FString::Printf(TEXT("%s, direction: %s"), *UEnvQueryTypes::DescribeContext(Context).ToString(), *UEnvQueryTypes::DescribeBoolParam(TraceToItem)) :
+		FString::Printf(TEXT("%s %s"), TraceToItem.Value ? TEXT("from") : TEXT("to"), *UEnvQueryTypes::DescribeContext(Context).ToString());
+
+	return FString::Printf(TEXT("%s: %s on %s"), 
+		*Super::GetDescriptionTitle(), *DirectionDesc, *ChannelDesc);
+}
+
+FText UEnvQueryTest_Trace::GetDescriptionDetails() const
+{
+	FFormatNamedArguments Args;
+	Args.Add(TEXT("TraceData"),  TraceData.ToText(FEnvTraceData::Detailed));
+	Args.Add(TEXT("TestParams"),  DescribeBoolTestParams("hit"));
+	
+	return LOCTEXT("TraceDescription", "{TraceData}\n{TestParams}");
+}
+
+bool UEnvQueryTest_Trace::RunLineTraceTo(const FVector& ItemPos, const FVector& ContextPos, AActor* ItemActor, UWorld* World, enum ECollisionChannel Channel, const FCollisionQueryParams& Params, const FVector& Extent)
+{
+	FCollisionQueryParams TraceParams(Params);
+	TraceParams.AddIgnoredActor(ItemActor);
+
+	const bool bHit = World->LineTraceTest(ContextPos, ItemPos, Channel, TraceParams);
+	return bHit;
+}
+
+bool UEnvQueryTest_Trace::RunLineTraceFrom(const FVector& ItemPos, const FVector& ContextPos, AActor* ItemActor, UWorld* World, enum ECollisionChannel Channel, const FCollisionQueryParams& Params, const FVector& Extent)
+{
+	FCollisionQueryParams TraceParams(Params);
+	TraceParams.AddIgnoredActor(ItemActor);
+
+	const bool bHit = World->LineTraceTest(ItemPos, ContextPos, Channel, TraceParams);
+	return bHit;
+}
+
+bool UEnvQueryTest_Trace::RunBoxTraceTo(const FVector& ItemPos, const FVector& ContextPos, AActor* ItemActor, UWorld* World, enum ECollisionChannel Channel, const FCollisionQueryParams& Params, const FVector& Extent)
+{
+	FCollisionQueryParams TraceParams(Params);
+	TraceParams.AddIgnoredActor(ItemActor);
+
+	const bool bHit = World->SweepTest(ContextPos, ItemPos, FQuat((ItemPos - ContextPos).Rotation()), Channel, FCollisionShape::MakeBox(Extent), TraceParams);
+	return bHit;
+}
+
+bool UEnvQueryTest_Trace::RunBoxTraceFrom(const FVector& ItemPos, const FVector& ContextPos, AActor* ItemActor, UWorld* World, enum ECollisionChannel Channel, const FCollisionQueryParams& Params, const FVector& Extent)
+{
+	FCollisionQueryParams TraceParams(Params);
+	TraceParams.AddIgnoredActor(ItemActor);
+
+	const bool bHit = World->SweepTest(ItemPos, ContextPos, FQuat((ContextPos - ItemPos).Rotation()), Channel, FCollisionShape::MakeBox(Extent), TraceParams);
+	return bHit;
+}
+
+bool UEnvQueryTest_Trace::RunSphereTraceTo(const FVector& ItemPos, const FVector& ContextPos, AActor* ItemActor, UWorld* World, enum ECollisionChannel Channel, const FCollisionQueryParams& Params, const FVector& Extent)
+{
+	FCollisionQueryParams TraceParams(Params);
+	TraceParams.AddIgnoredActor(ItemActor);
+
+	const bool bHit = World->SweepTest(ContextPos, ItemPos, FQuat::Identity, Channel, FCollisionShape::MakeSphere(Extent.X), TraceParams);
+	return bHit;
+}
+
+bool UEnvQueryTest_Trace::RunSphereTraceFrom(const FVector& ItemPos, const FVector& ContextPos, AActor* ItemActor, UWorld* World, enum ECollisionChannel Channel, const FCollisionQueryParams& Params, const FVector& Extent)
+{
+	FCollisionQueryParams TraceParams(Params);
+	TraceParams.AddIgnoredActor(ItemActor);
+
+	const bool bHit = World->SweepTest(ItemPos, ContextPos, FQuat::Identity, Channel, FCollisionShape::MakeSphere(Extent.X), TraceParams);
+	return bHit;
+}
+
+bool UEnvQueryTest_Trace::RunCapsuleTraceTo(const FVector& ItemPos, const FVector& ContextPos, AActor* ItemActor, UWorld* World, enum ECollisionChannel Channel, const FCollisionQueryParams& Params, const FVector& Extent)
+{
+	FCollisionQueryParams TraceParams(Params);
+	TraceParams.AddIgnoredActor(ItemActor);
+
+	const bool bHit = World->SweepTest(ContextPos, ItemPos, FQuat::Identity, Channel, FCollisionShape::MakeCapsule(Extent.X, Extent.Z), TraceParams);
+	return bHit;
+}
+
+bool UEnvQueryTest_Trace::RunCapsuleTraceFrom(const FVector& ItemPos, const FVector& ContextPos, AActor* ItemActor, UWorld* World, enum ECollisionChannel Channel, const FCollisionQueryParams& Params, const FVector& Extent)
+{
+	FCollisionQueryParams TraceParams(Params);
+	TraceParams.AddIgnoredActor(ItemActor);
+
+	const bool bHit = World->SweepTest(ItemPos, ContextPos, FQuat::Identity, Channel, FCollisionShape::MakeCapsule(Extent.X, Extent.Z), TraceParams);
+	return bHit;
+}
+
+#undef LOCTEXT_NAMESPACE
diff --git a/Engine/Source/Runtime/AIModule/Private/Navigation/CrowdFollowingComponent.cpp b/Engine/Source/Runtime/AIModule/Private/Navigation/CrowdFollowingComponent.cpp
new file mode 100644
--- /dev/null
+++ b/Engine/Source/Runtime/AIModule/Private/Navigation/CrowdFollowingComponent.cpp
@@ -0,0 +1,796 @@
+// Copyright 1998-2014 Epic Games, Inc. All Rights Reserved.
+
+#include "AIModulePrivate.h"
+#include "Navigation/CrowdFollowingComponent.h"
+#include "Navigation/CrowdManager.h"
+#include "Navigation/NavigationComponent.h"
+#include "AI/Navigation/SmartNavLinkComponent.h"
+
+DEFINE_LOG_CATEGORY(LogCrowdFollowing);
+
+UCrowdFollowingComponent::UCrowdFollowingComponent(const class FPostConstructInitializeProperties& PCIP) : Super(PCIP)
+{
+	bAffectFallingVelocity = false;
+	bRotateToVelocity = true;
+	bEnableCrowdSimulation = true;
+	bSuspendCrowdSimulation = false;
+
+	bEnableAnticipateTurns = false;
+	bEnableObstacleAvoidance = true;
+	bEnableSeparation = false;
+	bEnableOptimizeVisibility = true;
+	bEnableOptimizeTopology = true;
+
+	SeparationWeight = 2.0f;
+	CollisionQueryRange = 400.0f;		// approx: radius * 12.0f
+	PathOptimizationRange = 1000.0f;	// approx: radius * 30.0f
+	AvoidanceQuality = ECrowdAvoidanceQuality::Low;
+}
+
+FVector UCrowdFollowingComponent::GetCrowdAgentLocation() const
+{
+	return MovementComp ? MovementComp->GetActorFeetLocation() : FVector::ZeroVector;
+}
+
+FVector UCrowdFollowingComponent::GetCrowdAgentVelocity() const
+{
+	FVector Velocity(MovementComp ? MovementComp->Velocity : FVector::ZeroVector);
+	Velocity *= (Status == EPathFollowingStatus::Moving) ? 1.0f : 0.25f;
+	return Velocity;
+}
+
+void UCrowdFollowingComponent::GetCrowdAgentCollisions(float& CylinderRadius, float& CylinderHalfHeight) const
+{
+	if (MovementComp && MovementComp->UpdatedComponent)
+	{
+		MovementComp->UpdatedComponent->CalcBoundingCylinder(CylinderRadius, CylinderHalfHeight);
+	}
+}
+
+float UCrowdFollowingComponent::GetCrowdAgentMaxSpeed() const
+{
+	return MovementComp ? MovementComp->GetMaxSpeed() : 0.0f;
+}
+
+void UCrowdFollowingComponent::SetCrowdAnticipateTurns(bool bEnable, bool bUpdateAgent)
+{
+	if (bEnableAnticipateTurns != bEnable)
+	{
+		bEnableAnticipateTurns = bEnable;
+		
+		if (bUpdateAgent)
+		{
+			UpdateCrowdAgentParams();
+		}
+	}
+}
+
+void UCrowdFollowingComponent::SetCrowdObstacleAvoidance(bool bEnable, bool bUpdateAgent)
+{
+	if (bEnableObstacleAvoidance != bEnable)
+	{
+		bEnableObstacleAvoidance = bEnable;
+
+		if (bUpdateAgent)
+		{
+			UpdateCrowdAgentParams();
+		}
+	}
+}
+
+void UCrowdFollowingComponent::SetCrowdSeparation(bool bEnable, bool bUpdateAgent)
+{
+	if (bEnableSeparation != bEnable)
+	{
+		bEnableSeparation = bEnable;
+
+		if (bUpdateAgent)
+		{
+			UpdateCrowdAgentParams();
+		}
+	}
+}
+
+void UCrowdFollowingComponent::SetCrowdOptimizeVisibility(bool bEnable, bool bUpdateAgent)
+{
+	if (bEnableOptimizeVisibility != bEnable)
+	{
+		bEnableOptimizeVisibility = bEnable;
+
+		if (bUpdateAgent)
+		{
+			UpdateCrowdAgentParams();
+		}
+	}
+}
+
+void UCrowdFollowingComponent::SetCrowdOptimizeTopology(bool bEnable, bool bUpdateAgent)
+{
+	if (bEnableOptimizeTopology != bEnable)
+	{
+		bEnableOptimizeTopology = bEnable;
+
+		if (bUpdateAgent)
+		{
+			UpdateCrowdAgentParams();
+		}
+	}
+}
+
+void UCrowdFollowingComponent::SetCrowdSeparationWeight(float Weight, bool bUpdateAgent)
+{
+	if (SeparationWeight != Weight)
+	{
+		SeparationWeight = Weight;
+
+		if (bUpdateAgent)
+		{
+			UpdateCrowdAgentParams();
+		}
+	}
+}
+
+void UCrowdFollowingComponent::SetCrowdCollisionQueryRange(float Range, bool bUpdateAgent)
+{
+	if (CollisionQueryRange != Range)
+	{
+		CollisionQueryRange = Range;
+
+		if (bUpdateAgent)
+		{
+			UpdateCrowdAgentParams();
+		}
+	}
+}
+
+void UCrowdFollowingComponent::SetCrowdPathOptimizationRange(float Range, bool bUpdateAgent)
+{
+	if (PathOptimizationRange != Range)
+	{
+		PathOptimizationRange = Range;
+
+		if (bUpdateAgent)
+		{
+			UpdateCrowdAgentParams();
+		}
+	}
+}
+
+void UCrowdFollowingComponent::SetCrowdAvoidanceQuality(ECrowdAvoidanceQuality::Type Quality, bool bUpdateAgent)
+{
+	if (AvoidanceQuality != Quality)
+	{
+		AvoidanceQuality = Quality;
+
+		if (bUpdateAgent)
+		{
+			UpdateCrowdAgentParams();
+		}
+	}
+}
+
+void UCrowdFollowingComponent::SuspendCrowdSteering(bool bSuspend)
+{
+	if (bSuspendCrowdSimulation != bSuspend)
+	{
+		bSuspendCrowdSimulation = bSuspend;
+		UpdateCrowdAgentParams();
+	}
+}
+
+void UCrowdFollowingComponent::UpdateCrowdAgentParams() const
+{
+	UCrowdManager* CrowdManager = UCrowdManager::GetCurrent(GetWorld());
+	if (CrowdManager)
+	{
+		const ICrowdAgentInterface* IAgent = InterfaceCast<ICrowdAgentInterface>(this);
+		CrowdManager->UpdateAgentParams(IAgent);
+	}
+}
+
+void UCrowdFollowingComponent::UpdateCachedDirections(const FVector& NewVelocity, const FVector& NextPathCorner, bool bTraversingLink)
+{
+	// MoveSegmentDirection = direction on string pulled path
+	const FVector AgentLoc = GetCrowdAgentLocation();
+	const FVector ToCorner = NextPathCorner - AgentLoc;
+	if (ToCorner.SizeSquared() > FMath::Square(10.0f))
+	{
+		MoveSegmentDirection = ToCorner.SafeNormal();
+	}
+
+	// CrowdAgentMoveDirection either direction on path or aligned with current velocity
+	if (!bTraversingLink)
+	{
+		CrowdAgentMoveDirection = bRotateToVelocity && (NewVelocity.SizeSquared() > KINDA_SMALL_NUMBER) ? NewVelocity.SafeNormal() : MoveSegmentDirection;
+	}
+}
+
+void UCrowdFollowingComponent::ApplyCrowdAgentVelocity(const FVector& NewVelocity, const FVector& DestPathCorner, bool bTraversingLink)
+{
+	if (bEnableCrowdSimulation)
+	{
+		if (Status != EPathFollowingStatus::Paused &&
+			(bAffectFallingVelocity || CharacterMovement == NULL || CharacterMovement->MovementMode != MOVE_Falling))
+		{
+			MovementComp->RequestDirectMove(NewVelocity, false);
+
+			UpdateCachedDirections(NewVelocity, DestPathCorner, bTraversingLink);
+		}
+	}
+}
+
+void UCrowdFollowingComponent::SetCrowdSimulation(bool bEnable)
+{
+	if (bEnableCrowdSimulation == bEnable)
+	{
+		return;
+	}
+
+	if (GetStatus() != EPathFollowingStatus::Idle)
+	{
+		UE_VLOG(GetOwner(), LogCrowdFollowing, Warning, TEXT("SetCrowdSimulation failed, agent is not in Idle state!"));
+		return;
+	}
+
+	UE_VLOG(GetOwner(), LogCrowdFollowing, Log, TEXT("SetCrowdSimulation: %s"), bEnable ? TEXT("enabled") : TEXT("disabled"));
+	bEnableCrowdSimulation = bEnable;
+
+	if (NavComp)
+	{
+		if (bEnableCrowdSimulation)
+		{
+			// disable path post processing (string pulling), crowd simulation needs to handle 
+			// large paths by splitting into smaller parts and optimization gets in the way
+			NavComp->SetNavDataFlag(ERecastPathFlags::SkipStringPulling);
+		}
+		else
+		{
+			NavComp->ClearNavDataFlag(ERecastPathFlags::SkipStringPulling);
+		}
+	}
+}
+
+void UCrowdFollowingComponent::Initialize()
+{
+	Super::Initialize();
+
+	UCrowdManager* CrowdManager = UCrowdManager::GetCurrent(GetWorld());
+	if (CrowdManager)
+	{
+		const ICrowdAgentInterface* IAgent = InterfaceCast<ICrowdAgentInterface>(this);
+		// if CrowdManager->RegisterAgent fails we disable the crowd simulation for this guy
+		// @todo might consider re-trying in some time
+		bEnableCrowdSimulation = CrowdManager->RegisterAgent(IAgent);
+	}
+
+	if (bEnableCrowdSimulation && NavComp)
+	{
+		// disable path post processing (string pulling), crowd simulation needs to handle 
+		// large paths by splitting into smaller parts and optimization gets in the way
+		NavComp->SetNavDataFlag(ERecastPathFlags::SkipStringPulling);
+	}
+}
+
+void UCrowdFollowingComponent::Cleanup()
+{
+	Super::Cleanup();
+
+	UCrowdManager* CrowdManager = UCrowdManager::GetCurrent(GetWorld());
+	if (CrowdManager)
+	{
+		const ICrowdAgentInterface* IAgent = InterfaceCast<ICrowdAgentInterface>(this);
+		CrowdManager->UnregisterAgent(IAgent);
+	}
+}
+
+void UCrowdFollowingComponent::AbortMove(const FString& Reason, FAIRequestID RequestID, bool bResetVelocity, bool bSilent)
+{
+	if (bEnableCrowdSimulation && (Status != EPathFollowingStatus::Idle) && RequestID.IsEquivalent(GetCurrentRequestId()))
+	{
+		UCrowdManager* CrowdManager = UCrowdManager::GetCurrent(GetWorld());
+		if (CrowdManager)
+		{
+			CrowdManager->ClearAgentMoveTarget(this);
+		}
+	}
+
+	Super::AbortMove(Reason, RequestID, bResetVelocity, bSilent);
+}
+
+void UCrowdFollowingComponent::PauseMove(FAIRequestID RequestID, bool bResetVelocity)
+{
+	if (bEnableCrowdSimulation && (Status != EPathFollowingStatus::Paused) && RequestID.IsEquivalent(GetCurrentRequestId()))
+	{
+		UCrowdManager* CrowdManager = UCrowdManager::GetCurrent(GetWorld());
+		if (CrowdManager)
+		{
+			CrowdManager->PauseAgent(this);
+		}
+	}
+
+	Super::PauseMove(RequestID, bResetVelocity);
+}
+
+void UCrowdFollowingComponent::ResumeMove(FAIRequestID RequestID)
+{
+	if (bEnableCrowdSimulation && (Status == EPathFollowingStatus::Paused) && RequestID.IsEquivalent(GetCurrentRequestId()))
+	{
+		UCrowdManager* CrowdManager = UCrowdManager::GetCurrent(GetWorld());
+		if (CrowdManager)
+		{
+			const bool bHasMoved = HasMovedDuringPause();
+			CrowdManager->ResumeAgent(this, bHasMoved);
+		}
+	}
+
+	Super::ResumeMove(RequestID);
+}
+
+void UCrowdFollowingComponent::Reset()
+{
+	Super::Reset();
+	PathStartIndex = 0;
+
+	bFinalPathPart = false;
+	bCheckMovementAngle = false;
+	bUpdateDirectMoveVelocity = false;
+}
+
+bool UCrowdFollowingComponent::UpdateMovementComponent(bool bForce)
+{
+	bool bRet = Super::UpdateMovementComponent(bForce);
+	CharacterMovement = Cast<UCharacterMovementComponent>(MovementComp);
+
+	return bRet;
+}
+
+void UCrowdFollowingComponent::OnLanded()
+{
+	UCrowdManager* CrowdManager = UCrowdManager::GetCurrent(GetWorld());
+	if (bEnableCrowdSimulation && CrowdManager)
+	{
+		const ICrowdAgentInterface* IAgent = InterfaceCast<ICrowdAgentInterface>(this);
+		CrowdManager->UpdateAgentState(IAgent);
+	}
+}
+
+void UCrowdFollowingComponent::OnPathFinished(EPathFollowingResult::Type Result)
+{
+	UCrowdManager* CrowdManager = UCrowdManager::GetCurrent(GetWorld());
+	if (bEnableCrowdSimulation && CrowdManager)
+	{
+		CrowdManager->ClearAgentMoveTarget(this);
+	}
+
+	Super::OnPathFinished(Result);
+}
+
+void UCrowdFollowingComponent::OnPathUpdated()
+{
+	PathStartIndex = 0;
+}
+
+bool UCrowdFollowingComponent::ShouldCheckPathOnResume() const
+{
+	if (bEnableCrowdSimulation)
+	{
+		// never call SetMoveSegment on resuming
+		return false;
+	}
+
+	return HasMovedDuringPause();
+}
+
+bool UCrowdFollowingComponent::HasMovedDuringPause() const
+{
+	return Super::ShouldCheckPathOnResume();
+}
+
+bool UCrowdFollowingComponent::IsOnPath() const
+{
+	if (bEnableCrowdSimulation)
+	{
+		// agent can move off path for steering/avoidance purposes
+		// just pretend it's always on path to avoid problems when movement is being resumed
+		return true;
+	}
+
+	return Super::IsOnPath();
+}
+
+int32 UCrowdFollowingComponent::DetermineStartingPathPoint(const FNavigationPath* ConsideredPath) const
+{
+	int32 StartIdx = 0;
+
+	if (bEnableCrowdSimulation)
+	{
+		StartIdx = PathStartIndex;
+
+		// no path = called from SwitchToNextPathPart, check if agent is already further on path
+		if (ConsideredPath == NULL && Path.IsValid())
+		{
+			UCrowdManager* Manager = UCrowdManager::GetCurrent(GetWorld());
+			FNavMeshPath* NavPath = Path->CastPath<FNavMeshPath>();
+			if (Manager && NavPath)
+			{
+				int32 AdjustedStartIdx = PathStartIndex;
+				Manager->AdjustAgentPathStart(this, NavPath, AdjustedStartIdx);
+
+				StartIdx = AdjustedStartIdx;
+			}
+		}
+	}
+	else
+	{
+		StartIdx = Super::DetermineStartingPathPoint(ConsideredPath);
+	}
+
+	return StartIdx;
+}
+
+void LogPathPartHelper(AActor* LogOwner, FNavMeshPath* NavMeshPath, int32 StartIdx, int32 EndIdx)
+{
+#if ENABLE_VISUAL_LOG && WITH_RECAST
+	ARecastNavMesh* NavMesh = Cast<ARecastNavMesh>(NavMeshPath->GetOwner());
+	NavMesh->BeginBatchQuery();
+	
+	FVector SegmentStart = NavMeshPath->PathPoints[0].Location;
+	FVector SegmentEnd(FVector::ZeroVector);
+
+	if (StartIdx > 0)
+	{
+		NavMesh->GetPolyCenter(NavMeshPath->PathCorridor[StartIdx], SegmentStart);
+	}
+
+	for (int32 Idx = StartIdx + 1; Idx < EndIdx; Idx++)
+	{
+		NavMesh->GetPolyCenter(NavMeshPath->PathCorridor[Idx], SegmentEnd);
+		UE_VLOG_SEGMENT_THICK(LogOwner, LogCrowdFollowing, Log, SegmentStart, SegmentEnd, FColor::Yellow, 2, TEXT_EMPTY);
+
+		SegmentStart = SegmentEnd;
+	}
+
+	if (EndIdx == (NavMeshPath->PathCorridor.Num() - 1))
+	{
+		SegmentEnd = NavMeshPath->PathPoints[1].Location;
+	}
+	else
+	{
+		NavMesh->GetPolyCenter(NavMeshPath->PathCorridor[EndIdx], SegmentEnd);
+	}
+
+	UE_VLOG_SEGMENT_THICK(LogOwner, LogCrowdFollowing, Log, SegmentStart, SegmentEnd, FColor::Yellow, 2, TEXT_EMPTY);
+	NavMesh->FinishBatchQuery();
+#endif // ENABLE_VISUAL_LOG && WITH_RECAST
+}
+
+void UCrowdFollowingComponent::SetMoveSegment(int32 SegmentStartIndex)
+{
+	if (!bEnableCrowdSimulation)
+	{
+		Super::SetMoveSegment(SegmentStartIndex);
+		return;
+	}
+
+	PathStartIndex = SegmentStartIndex;
+	if (!Path.IsValid())
+	{
+		return;
+	}
+
+	FVector CurrentTargetPt = Path->PathPoints[1].Location;
+
+	FNavMeshPath* NavMeshPath = Path->CastPath<FNavMeshPath>();
+	if (NavMeshPath)
+	{
+#if WITH_RECAST
+		// cut paths into parts to avoid problems with crowds getting into local minimum
+		// due to using only first 10 steps of A*
+
+		// do NOT use PathPoints here, crowd simulation disables path post processing
+		// which means, that PathPoints contains only start and end position 
+		// full path is available through PathCorridor array (poly refs)
+
+		ARecastNavMesh* RecastNavData = Cast<ARecastNavMesh>(MyNavData);
+
+		const int32 PathPartSize = 15;
+		const int32 LastPolyIdx = NavMeshPath->PathCorridor.Num() - 1;
+		const int32 PathPartEndIdx = FMath::Min(PathStartIndex + PathPartSize, LastPolyIdx);
+
+		bFinalPathPart = (PathPartEndIdx == LastPolyIdx);
+		if (!bFinalPathPart)
+		{
+			RecastNavData->GetPolyCenter(NavMeshPath->PathCorridor[PathPartEndIdx], CurrentTargetPt);
+		}
+
+		// not safe to read those directions yet, you have to wait until crowd manager gives you next corner of string pulled path
+		CrowdAgentMoveDirection = FVector::ZeroVector;
+		MoveSegmentDirection = FVector::ZeroVector;
+
+		CurrentDestination.Set(Path->Base.Get(), CurrentTargetPt);
+		SuspendCrowdSteering(false);
+
+		UE_VLOG(GetOwner(), LogCrowdFollowing, Log, TEXT("SetMoveSegment, from:%d segments:%d%s"), PathStartIndex, PathPartSize, bFinalPathPart ? TEXT(" (final)") : TEXT(""));
+		UE_VLOG_SEGMENT(GetOwner(), LogCrowdFollowing, Log, MovementComp->GetActorFeetLocation(), CurrentTargetPt, FColor::Red, TEXT("path part"));
+		LogPathPartHelper(GetOwner(), NavMeshPath, PathStartIndex, PathPartEndIdx);
+
+		UCrowdManager* CrowdManager = UCrowdManager::GetCurrent(GetWorld());
+		if (CrowdManager)
+		{
+			CrowdManager->SetAgentMovePath(this, NavMeshPath, PathStartIndex, PathPartEndIdx);
+		}
+#endif
+	}
+	else
+	{
+		// direct paths are not using any steering or avoidance
+		// pathfinding is replaced with simple velocity request 
+
+		const FVector AgentLoc = MovementComp->GetActorFeetLocation();
+
+		bFinalPathPart = true;
+		bCheckMovementAngle = true;
+		bUpdateDirectMoveVelocity = DestinationActor.IsValid();
+		CurrentDestination.Set(Path->Base.Get(), CurrentTargetPt);
+		CrowdAgentMoveDirection = (CurrentTargetPt - AgentLoc).SafeNormal();
+		MoveSegmentDirection = CrowdAgentMoveDirection;
+		SuspendCrowdSteering(true);
+
+		UE_VLOG(GetOwner(), LogCrowdFollowing, Log, TEXT("SetMoveSegment, direct move"));
+		UE_VLOG_SEGMENT(GetOwner(), LogCrowdFollowing, Log, AgentLoc, CurrentTargetPt, FColor::Red, TEXT("path"));
+
+		UCrowdManager* CrowdManager = UCrowdManager::GetCurrent(GetWorld());
+		if (CrowdManager)
+		{
+			CrowdManager->SetAgentMoveDirection(this, CrowdAgentMoveDirection);
+		}
+	}
+}
+
+void UCrowdFollowingComponent::UpdatePathSegment()
+{
+	if (!bEnableCrowdSimulation)
+	{
+		Super::UpdatePathSegment();
+		return;
+	}
+
+	if (!Path.IsValid() || MovementComp == NULL)
+	{
+		AbortMove(TEXT("no path"));
+		return;
+	}
+
+	if (!Path->IsValid())
+	{
+		if (NavComp == NULL || !NavComp->IsWaitingForRepath())
+		{
+			AbortMove(TEXT("no path"));
+		}
+		return;
+	}
+
+	// if agent has control over its movement, check finish conditions
+	const bool bCanReachTarget = MovementComp->CanStopPathFollowing();
+	if (bCanReachTarget && Status == EPathFollowingStatus::Moving)
+	{
+		const FVector CurrentLocation = MovementComp->GetActorFeetLocation();
+		const FVector GoalLocation = GetCurrentTargetLocation();
+
+		if (bFinalPathPart)
+		{
+			const FVector ToTarget = (GoalLocation - MovementComp->GetActorFeetLocation());
+			const float SegmentDot = FVector::DotProduct(ToTarget, Path->IsDirect() ? MovementComp->Velocity : CrowdAgentMoveDirection);
+			const bool bMovedTooFar = bCheckMovementAngle && (SegmentDot < 0.0);
+
+			if (bMovedTooFar || HasReachedDestination(CurrentLocation))
+			{
+				UE_VLOG(GetOwner(), LogCrowdFollowing, Log, TEXT("Last path segment finished due to \'%s\'"), bMovedTooFar ? TEXT("Missing Last Point") : TEXT("Reaching Destination"));
+				OnPathFinished(EPathFollowingResult::Success);
+			}
+		}
+		else
+		{
+			// override radius multiplier and switch to next path part when closer than 4x agent radius
+			const float SavedAgentRadiusPct = MinAgentRadiusPct;
+			MinAgentRadiusPct = 4.0f;
+			
+			const bool bHasReached = HasReachedInternal(GoalLocation, 0.0f, 0.0f, CurrentLocation, 0.0f, false);
+			
+			MinAgentRadiusPct = SavedAgentRadiusPct;
+
+			if (bHasReached)
+			{
+				SwitchToNextPathPart();
+			}
+		}
+	}
+
+	// gather location samples to detect if moving agent is blocked
+	if (bCanReachTarget && Status == EPathFollowingStatus::Moving)
+	{
+		const bool bHasNewSample = UpdateBlockDetection();
+		if (bHasNewSample && IsBlocked())
+		{
+			OnPathFinished(EPathFollowingResult::Blocked);
+		}
+	}
+}
+
+void UCrowdFollowingComponent::FollowPathSegment(float DeltaTime)
+{
+	if (!bEnableCrowdSimulation)
+	{
+		Super::FollowPathSegment(DeltaTime);
+		return;
+	}
+
+	if (bUpdateDirectMoveVelocity && DestinationActor.IsValid())
+	{
+		const FVector CurrentTargetPt = DestinationActor->GetActorLocation();
+		const float DistSq = (CurrentTargetPt - GetCurrentTargetLocation()).SizeSquared();
+		if (DistSq > FMath::Square(10.0f))
+		{
+			UCrowdManager* Manager = UCrowdManager::GetCurrent(GetWorld());
+			const FVector AgentLoc = GetCrowdAgentLocation();
+
+			CurrentDestination.Set(Path->Base.Get(), CurrentTargetPt);
+			CrowdAgentMoveDirection = (CurrentTargetPt - AgentLoc).SafeNormal();
+			MoveSegmentDirection = CrowdAgentMoveDirection;
+
+			Manager->SetAgentMoveDirection(this, MoveSegmentDirection);
+			UE_VLOG(GetOwner(), LogCrowdFollowing, Log, TEXT("Updated direct move direction for crowd agent."));
+		}
+	}
+
+	UpdateMoveFocus();
+}
+
+FVector UCrowdFollowingComponent::GetMoveFocus(bool bAllowStrafe) const
+{
+	// can't really use CurrentDestination here, as it's pointing at end of path part
+	// fallback to looking at point in front of agent
+
+	if (!bAllowStrafe && MovementComp && bEnableCrowdSimulation)
+	{
+		const FVector AgentLoc = MovementComp->GetActorLocation();
+		const FVector ForwardDir = CrowdAgentMoveDirection.IsNearlyZero() && MovementComp && MovementComp->GetOwner() ?
+			MovementComp->GetOwner()->GetActorRotation().Vector() :
+			CrowdAgentMoveDirection;
+
+		return AgentLoc + ForwardDir * 100.0f;
+	}
+
+	return Super::GetMoveFocus(bAllowStrafe);
+}
+
+void UCrowdFollowingComponent::OnNavNodeChanged(NavNodeRef NewPolyRef, NavNodeRef PrevPolyRef, int32 CorridorSize)
+{
+	if (bEnableCrowdSimulation)
+	{
+		UE_VLOG(GetOwner(), LogCrowdFollowing, Verbose, TEXT("OnNavNodeChanged, CorridorSize:%d"), CorridorSize);
+
+		const bool bSwitchPart = ShouldSwitchPathPart(CorridorSize);
+		if (bSwitchPart && !bFinalPathPart)
+		{
+			SwitchToNextPathPart();
+		}
+	}
+}
+
+bool UCrowdFollowingComponent::ShouldSwitchPathPart(int32 CorridorSize) const
+{
+	return CorridorSize <= 2;
+}
+
+void UCrowdFollowingComponent::SwitchToNextPathPart()
+{
+	const int32 NewPartStart = DetermineStartingPathPoint(NULL);
+	SetMoveSegment(NewPartStart);
+}
+
+void UCrowdFollowingComponent::GetDebugStringTokens(TArray<FString>& Tokens, TArray<EPathFollowingDebugTokens::Type>& Flags) const
+{
+	if (!bEnableCrowdSimulation)
+	{
+		Super::GetDebugStringTokens(Tokens, Flags);
+		return;
+	}
+
+	Tokens.Add(GetStatusDesc());
+	Flags.Add(EPathFollowingDebugTokens::Description);
+
+	if (Status != EPathFollowingStatus::Moving)
+	{
+		return;
+	}
+
+	FString& StatusDesc = Tokens[0];
+	if (Path.IsValid())
+	{
+		FNavMeshPath* NavMeshPath = Path->CastPath<FNavMeshPath>();
+		if (NavMeshPath)
+		{
+			StatusDesc += FString::Printf(TEXT(" (path:%d)"), PathStartIndex);
+		}
+		else
+		{
+			StatusDesc += TEXT(" (direct)");
+		}
+	}
+
+	// get cylinder of moving agent
+	float AgentRadius = 0.0f;
+	float AgentHalfHeight = 0.0f;
+	AActor* MovingAgent = MovementComp->GetOwner();
+	MovingAgent->GetSimpleCollisionCylinder(AgentRadius, AgentHalfHeight);
+
+	if (bFinalPathPart)
+	{
+		float CurrentDot = 0.0f, CurrentDistance = 0.0f, CurrentHeight = 0.0f;
+		uint8 bFailedDot = 0, bFailedDistance = 0, bFailedHeight = 0;
+		DebugReachTest(CurrentDot, CurrentDistance, CurrentHeight, bFailedHeight, bFailedDistance, bFailedHeight);
+
+		Tokens.Add(TEXT("dist2D"));
+		Flags.Add(EPathFollowingDebugTokens::ParamName);
+		Tokens.Add(FString::Printf(TEXT("%.0f"), CurrentDistance));
+		Flags.Add(bFailedDistance ? EPathFollowingDebugTokens::FailedValue : EPathFollowingDebugTokens::PassedValue);
+
+		Tokens.Add(TEXT("distZ"));
+		Flags.Add(EPathFollowingDebugTokens::ParamName);
+		Tokens.Add(FString::Printf(TEXT("%.0f"), CurrentHeight));
+		Flags.Add(bFailedHeight ? EPathFollowingDebugTokens::FailedValue : EPathFollowingDebugTokens::PassedValue);
+	}
+	else
+	{
+		const FVector CurrentLocation = MovementComp->GetActorFeetLocation();
+
+		// make sure we're not too close to end of path part (poly count can always fail when AI goes off path)
+		const float DistSq = (GetCurrentTargetLocation() - CurrentLocation).SizeSquared();
+		const float PathSwitchThresSq = FMath::Square(AgentRadius * 5.0f);
+
+		Tokens.Add(TEXT("dist2D"));
+		Flags.Add(EPathFollowingDebugTokens::ParamName);
+		Tokens.Add(FString::Printf(TEXT("%.0f"), FMath::Sqrt(DistSq)));
+		Flags.Add((DistSq < PathSwitchThresSq) ? EPathFollowingDebugTokens::PassedValue : EPathFollowingDebugTokens::FailedValue);
+	}
+}
+
+#if ENABLE_VISUAL_LOG
+
+void UCrowdFollowingComponent::DescribeSelfToVisLog(struct FVisLogEntry* Snapshot) const
+{
+	if (!bEnableCrowdSimulation)
+	{
+		Super::DescribeSelfToVisLog(Snapshot);
+		return;
+	}
+
+	FVisLogEntry::FStatusCategory Category;
+	Category.Category = TEXT("Path following");
+
+	if (DestinationActor.IsValid())
+	{
+		Category.Add(TEXT("Goal"), GetNameSafe(DestinationActor.Get()));
+	}
+
+	FString StatusDesc = GetStatusDesc();
+
+	FNavMeshPath* NavMeshPath = Path.IsValid() ? Path->CastPath<FNavMeshPath>() : NULL;
+	if (Status == EPathFollowingStatus::Moving)
+	{
+		StatusDesc += FString::Printf(TEXT(" [path:%d]"), PathStartIndex);
+	}
+
+	Category.Add(TEXT("Status"), StatusDesc);
+	Category.Add(TEXT("Path"), !Path.IsValid() ? TEXT("none") : NavMeshPath ? TEXT("navmesh") : TEXT("direct"));
+
+	if (CurrentSmartLink)
+	{
+		Category.Add(TEXT("SmartLink"), CurrentSmartLink->GetName());
+	}
+
+	Snapshot->Status.Add(Category);
+}
+
+#endif // ENABLE_VISUAL_LOG
diff --git a/Engine/Source/Runtime/AIModule/Private/Navigation/CrowdManager.cpp b/Engine/Source/Runtime/AIModule/Private/Navigation/CrowdManager.cpp
new file mode 100644
--- /dev/null
+++ b/Engine/Source/Runtime/AIModule/Private/Navigation/CrowdManager.cpp
@@ -0,0 +1,1047 @@
+// Copyright 1998-2014 Epic Games, Inc. All Rights Reserved.
+
+#include "AIModulePrivate.h"
+
+#if WITH_RECAST
+#include "AI/Navigation/PImplRecastNavMesh.h"
+#include "AI/Navigation/RecastHelpers.h"
+#include "DetourCrowd/DetourCrowd.h"
+#endif
+
+#include "Navigation/CrowdManager.h"
+#include "Navigation/CrowdFollowingComponent.h"
+#include "Navigation/CrowdAgentInterface.h"
+
+DECLARE_STATS_GROUP(TEXT("Crowd"), STATGROUP_AICrowd, STATCAT_Advanced);
+
+DECLARE_CYCLE_STAT(TEXT("Simulation Time"), STAT_AI_Crowd_Tick, STATGROUP_AICrowd);
+DECLARE_CYCLE_STAT(TEXT("Step: corridor update"), STAT_AI_Crowd_StepCorridorTime, STATGROUP_AICrowd);
+DECLARE_CYCLE_STAT(TEXT("Step: paths"), STAT_AI_Crowd_StepPathsTime, STATGROUP_AICrowd);
+DECLARE_CYCLE_STAT(TEXT("Step: proximity"), STAT_AI_Crowd_StepProximityTime, STATGROUP_AICrowd);
+DECLARE_CYCLE_STAT(TEXT("Step: next point"), STAT_AI_Crowd_StepNextPointTime, STATGROUP_AICrowd);
+DECLARE_CYCLE_STAT(TEXT("Step: steering"), STAT_AI_Crowd_StepSteeringTime, STATGROUP_AICrowd);
+DECLARE_CYCLE_STAT(TEXT("Step: avoidance"), STAT_AI_Crowd_StepAvoidanceTime, STATGROUP_AICrowd);
+DECLARE_CYCLE_STAT(TEXT("Step: components"), STAT_AI_Crowd_StepComponentsTime, STATGROUP_AICrowd);
+DECLARE_CYCLE_STAT(TEXT("Step: navlinks"), STAT_AI_Crowd_StepNavLinkTime, STATGROUP_AICrowd);
+DECLARE_CYCLE_STAT(TEXT("Step: movement"), STAT_AI_Crowd_StepMovementTime, STATGROUP_AICrowd);
+DECLARE_CYCLE_STAT(TEXT("Agent Update Time"), STAT_AI_Crowd_AgentUpdateTime, STATGROUP_AICrowd);
+DECLARE_DWORD_COUNTER_STAT(TEXT("Num Agents"), STAT_AI_Crowd_NumAgents, STATGROUP_AICrowd);
+
+DEFINE_LOG_CATEGORY_STATIC(LogEngineCrowdFollowing, Warning, All)
+
+namespace CrowdDebugDrawing
+{
+	/** if set, debug information will be displayed for agent selected in editor */
+	bool bDebugSelectedActors = true;
+	/** if set, basic debug information will be recorded in VisLog for all agents */
+	bool bDebugVisLog = false;
+
+	/** debug flags, works only for selected actor */
+	bool bDrawDebugCorners = true;
+	bool bDrawDebugCollisionSegments = true;
+	bool bDrawDebugPath = true;
+	bool bDrawDebugVelocityObstacles = true;
+	bool bDrawDebugPathOptimization = true;
+	bool bDrawDebugNeighbors = true;
+
+	const FVector Offset(0, 0, 20);
+
+	const FColor Corner(128, 0, 0);
+	const FColor CornerLink(192, 0, 0);
+	const FColor CollisionRange(192, 0, 128);
+	const FColor CollisionSeg0(192, 0, 128);
+	const FColor CollisionSeg1(96, 0, 64);
+	const FColor Path(255, 255, 255);
+	const FColor PathSpecial(255, 192, 203);
+	const FColor PathOpt(0, 128, 0);
+	const FColor AvoidanceRange(255, 255, 255);
+	const FColor Neighbor(0, 192, 128);
+
+	const float LineThickness = 3.f;
+}
+
+void FCrowdTickHelper::Tick(float DeltaTime)
+{
+#if WITH_EDITOR
+	if (Owner.IsValid())
+	{
+		Owner->DebugTick();
+	}
+#endif // WITH_EDITOR
+}
+
+TStatId FCrowdTickHelper::GetStatId() const
+{
+	RETURN_QUICK_DECLARE_CYCLE_STAT(FCrowdTickHelper, STATGROUP_Tickables);
+}
+
+UCrowdManager::UCrowdManager(const class FPostConstructInitializeProperties& PCIP) : Super(PCIP)
+{
+	MyNavData = NULL;
+#if WITH_RECAST
+	DetourCrowd = NULL;
+
+	if (!HasAnyFlags(RF_ClassDefaultObject))
+	{
+		DetourAvoidanceDebug = dtAllocObstacleAvoidanceDebugData();
+		DetourAvoidanceDebug->init(2048);
+
+		DetourAgentDebug = new dtCrowdAgentDebugInfo();
+		FMemory::Memzero(DetourAgentDebug, sizeof(dtCrowdAgentDebugInfo));
+		DetourAgentDebug->idx = -1;
+		DetourAgentDebug->vod = DetourAvoidanceDebug;
+	}
+	else
+	{
+		DetourAgentDebug = NULL;
+		DetourAvoidanceDebug = NULL;
+	}
+#endif
+#if WITH_EDITOR
+	TickHelper = NULL;
+	if (!HasAnyFlags(RF_ClassDefaultObject) && GIsEditor)
+	{
+		TickHelper = new FCrowdTickHelper();
+		TickHelper->Owner = this;
+	}
+#endif
+
+	MaxAgents = 50;
+	MaxAgentRadius = 100.0f;
+	NavmeshCheckInterval = 1.0f;
+	bSingleAreaVisibilityOptimization = true;
+	bPruneStartedOffmeshConnections = false;
+	
+	FCrowdAvoidanceConfig AvoidanceConfig11;		// 11 samples, ECrowdAvoidanceQuality::Low
+	AvoidanceConfig11.VelocityBias = 0.5f;
+	AvoidanceConfig11.AdaptiveDivisions = 5;
+	AvoidanceConfig11.AdaptiveRings = 2;
+	AvoidanceConfig11.AdaptiveDepth = 1;
+	AvoidanceConfig.Add(AvoidanceConfig11);
+	
+	FCrowdAvoidanceConfig AvoidanceConfig22;		// 22 samples, ECrowdAvoidanceQuality::Medium
+	AvoidanceConfig22.VelocityBias = 0.5f;
+	AvoidanceConfig22.AdaptiveDivisions = 5;
+	AvoidanceConfig22.AdaptiveRings = 2;
+	AvoidanceConfig22.AdaptiveDepth = 2;
+	AvoidanceConfig.Add(AvoidanceConfig22);
+
+	FCrowdAvoidanceConfig AvoidanceConfig45;		// 45 samples, ECrowdAvoidanceQuality::Good
+	AvoidanceConfig45.VelocityBias = 0.5f;
+	AvoidanceConfig45.AdaptiveDivisions = 7;
+	AvoidanceConfig45.AdaptiveRings = 2;
+	AvoidanceConfig45.AdaptiveDepth = 3;
+	AvoidanceConfig.Add(AvoidanceConfig45);
+
+	FCrowdAvoidanceConfig AvoidanceConfig66;		// 66 samples, ECrowdAvoidanceQuality::High
+	AvoidanceConfig66.VelocityBias = 0.5f;
+	AvoidanceConfig66.AdaptiveDivisions = 7;
+	AvoidanceConfig66.AdaptiveRings = 3;
+	AvoidanceConfig66.AdaptiveDepth = 3;
+	AvoidanceConfig.Add(AvoidanceConfig66);
+}
+
+void UCrowdManager::BeginDestroy()
+{
+#if WITH_RECAST
+	// cleanup allocated link filters
+	for (auto It = ActiveAgents.CreateConstIterator(); It; ++It)
+	{
+		const FCrowdAgentData& AgentData = It.Value();
+		delete AgentData.LinkFilter;
+	}
+
+	dtFreeObstacleAvoidanceDebugData(DetourAvoidanceDebug);
+	delete DetourAgentDebug;
+#endif
+
+#if WITH_EDITOR
+	delete TickHelper;
+#endif
+
+	Super::BeginDestroy();
+}
+
+void UCrowdManager::Tick(float DeltaTime)
+{
+	SCOPE_CYCLE_COUNTER(STAT_AI_Crowd_Tick);
+	INC_DWORD_STAT_BY(STAT_AI_Crowd_NumAgents, ActiveAgents.Num());
+
+#if WITH_RECAST
+	if (DetourCrowd)
+	{
+		int32 NumActive = DetourCrowd->cacheActiveAgents();
+		if (NumActive)
+		{
+			MyNavData->BeginBatchQuery();
+
+			for (auto It = ActiveAgents.CreateConstIterator(); It; ++It)
+			{
+				// collect position and velocity
+				const FCrowdAgentData& AgentData = It.Value();
+				UpdateSpatialData(It.Key(), AgentData);
+			}
+
+			// corridor update from previous step
+			{
+				SCOPE_CYCLE_COUNTER(STAT_AI_Crowd_StepCorridorTime);
+				DetourCrowd->updateStepCorridor(DeltaTime, DetourAgentDebug);
+			}
+
+			// regular steps
+			{
+				SCOPE_CYCLE_COUNTER(STAT_AI_Crowd_StepPathsTime);
+				DetourCrowd->updateStepPaths(DeltaTime, DetourAgentDebug);
+			}
+			{
+				SCOPE_CYCLE_COUNTER(STAT_AI_Crowd_StepProximityTime);
+				DetourCrowd->updateStepProximityData(DeltaTime, DetourAgentDebug);
+				PostProximityUpdate();
+			}
+			{
+				SCOPE_CYCLE_COUNTER(STAT_AI_Crowd_StepNextPointTime);
+				DetourCrowd->updateStepNextMovePoint(DeltaTime, DetourAgentDebug);
+			}
+			{
+				SCOPE_CYCLE_COUNTER(STAT_AI_Crowd_StepSteeringTime);
+				DetourCrowd->updateStepSteering(DeltaTime, DetourAgentDebug);
+			}
+			{
+				SCOPE_CYCLE_COUNTER(STAT_AI_Crowd_StepAvoidanceTime);
+				DetourCrowd->updateStepAvoidance(DeltaTime, DetourAgentDebug);
+			}
+			{
+				SCOPE_CYCLE_COUNTER(STAT_AI_Crowd_StepComponentsTime);
+				UpdateAgentPaths();
+			}
+			{
+				SCOPE_CYCLE_COUNTER(STAT_AI_Crowd_StepNavLinkTime);
+				DetourCrowd->updateStepOffMeshVelocity(DeltaTime, DetourAgentDebug);
+			}
+
+			// velocity updates
+			{
+				SCOPE_CYCLE_COUNTER(STAT_AI_Crowd_StepMovementTime);
+				for (auto It = ActiveAgents.CreateConstIterator(); It; ++It)
+				{
+					const FCrowdAgentData& AgentData = It.Value();
+					if (AgentData.bIsSimulated)
+					{
+						const UCrowdFollowingComponent* CrowdComponent = Cast<const UCrowdFollowingComponent>(It.Key()->GetUObjectInterfaceCrowdAgentInterface());
+						ApplyVelocity(CrowdComponent, AgentData.AgentIndex);
+					}
+				}
+			}
+
+			MyNavData->FinishBatchQuery();
+
+#if WITH_EDITOR
+			// normalize samples only for debug drawing purposes
+			DetourAvoidanceDebug->normalizeSamples();
+#endif
+		}
+	}
+#endif // WITH_RECAST
+}
+
+#if WITH_EDITOR
+void UCrowdManager::PostEditChangeProperty(struct FPropertyChangedEvent& PropertyChangedEvent)
+{
+	Super::PostEditChangeProperty(PropertyChangedEvent);
+	
+#if WITH_RECAST
+	// recreate crowd manger
+	DestroyCrowdManager();
+	CreateCrowdManager();
+#endif // WITH_RECAST
+}
+#endif // WITH_EDITOR
+
+bool UCrowdManager::RegisterAgent(const ICrowdAgentInterface* Agent)
+{
+	UpdateNavData();
+
+	FCrowdAgentData AgentData;
+
+#if WITH_RECAST
+	if (DetourCrowd)
+	{
+		AddAgent(Agent, AgentData);
+	}
+#endif
+
+	const bool bSuccessfullyAdded = (AgentData.AgentIndex >= 0);
+	if (bSuccessfullyAdded)
+	{
+		ActiveAgents.Add(Agent, AgentData);
+	}
+	
+	return bSuccessfullyAdded;
+}
+
+void UCrowdManager::UnregisterAgent(const ICrowdAgentInterface* Agent)
+{
+#if WITH_RECAST
+	const FCrowdAgentData* AgentData = ActiveAgents.Find(Agent);
+	if (DetourCrowd && AgentData)
+	{
+		RemoveAgent(Agent, AgentData);
+	}
+#endif
+
+	ActiveAgents.Remove(Agent);
+}
+
+void UCrowdManager::UpdateAgentParams(const ICrowdAgentInterface* Agent) const
+{
+	SCOPE_CYCLE_COUNTER(STAT_AI_Crowd_AgentUpdateTime);
+
+#if WITH_RECAST
+	const FCrowdAgentData* AgentData = ActiveAgents.Find(Agent);
+	if (DetourCrowd && AgentData)
+	{
+		dtCrowdAgentParams Params;
+		GetAgentParams(Agent, &Params);
+
+		DetourCrowd->updateAgentParameters(AgentData->AgentIndex, &Params);
+	}
+#endif
+}
+
+void UCrowdManager::UpdateAgentState(const class ICrowdAgentInterface* Agent) const
+{
+	SCOPE_CYCLE_COUNTER(STAT_AI_Crowd_AgentUpdateTime);
+
+#if WITH_RECAST
+	const FCrowdAgentData* AgentData = ActiveAgents.Find(Agent);
+	if (DetourCrowd && AgentData)
+	{
+		DetourCrowd->updateAgentState(AgentData->AgentIndex, false);
+	}
+#endif
+}
+
+bool UCrowdManager::SetAgentMoveTarget(const UCrowdFollowingComponent* AgentComponent, const FVector& MoveTarget, TSharedPtr<const FNavigationQueryFilter> Filter) const
+{
+	SCOPE_CYCLE_COUNTER(STAT_AI_Crowd_AgentUpdateTime);
+
+	bool bSuccess = false;
+
+#if WITH_RECAST
+	const FCrowdAgentData* AgentData = ActiveAgents.Find(AgentComponent);
+	if (AgentData && AgentData->bIsSimulated && DetourCrowd)
+	{
+		FNavLocation ProjectedLoc;
+		MyNavData->ProjectPoint(MoveTarget, ProjectedLoc, MyNavData->GetDefaultQueryExtent(), Filter);
+
+		const INavigationQueryFilterInterface* NavFilter = Filter.IsValid() ? Filter->GetImplementation() : MyNavData->GetDefaultQueryFilterImpl();
+		const dtQueryFilter* DetourFilter = ((const FRecastQueryFilter*)NavFilter)->GetAsDetourQueryFilter();
+		DetourCrowd->updateAgentFilter(AgentData->AgentIndex, DetourFilter);
+		DetourCrowd->updateAgentState(AgentData->AgentIndex, false);
+
+		const FVector RcTargetPos = Unreal2RecastPoint(MoveTarget);
+		bSuccess = DetourCrowd->requestMoveTarget(AgentData->AgentIndex, ProjectedLoc.NodeRef, &RcTargetPos.X);
+	}
+#endif
+
+	return bSuccess;
+}
+
+bool UCrowdManager::SetAgentMoveDirection(const class UCrowdFollowingComponent* AgentComponent, const FVector& MoveDirection) const
+{
+	SCOPE_CYCLE_COUNTER(STAT_AI_Crowd_AgentUpdateTime);
+
+	bool bSuccess = false;
+
+#if WITH_RECAST
+	const FCrowdAgentData* AgentData = ActiveAgents.Find(AgentComponent);
+	if (AgentData && AgentData->bIsSimulated && DetourCrowd)
+	{
+		DetourCrowd->updateAgentState(AgentData->AgentIndex, false);
+
+		const FVector RcTargetVelocity = Unreal2RecastPoint(MoveDirection * AgentComponent->GetCrowdAgentMaxSpeed());
+		bSuccess = DetourCrowd->requestMoveVelocity(AgentData->AgentIndex, &RcTargetVelocity.X);
+	}
+#endif
+
+	return bSuccess;
+}
+
+bool UCrowdManager::SetAgentMovePath(const class UCrowdFollowingComponent* AgentComponent, const FNavMeshPath* Path, int32 PathSectionStart, int32 PathSectionEnd) const
+{
+	SCOPE_CYCLE_COUNTER(STAT_AI_Crowd_AgentUpdateTime);
+
+	bool bSuccess = false;
+
+#if WITH_RECAST
+	const FCrowdAgentData* AgentData = ActiveAgents.Find(AgentComponent);
+	ARecastNavMesh* RecastNavData = Cast<ARecastNavMesh>(MyNavData);
+	if (AgentData && AgentData->bIsSimulated && DetourCrowd &&
+		RecastNavData && Path && (Path->PathPoints.Num() > 1) &&
+		Path->PathCorridor.IsValidIndex(PathSectionStart) && Path->PathCorridor.IsValidIndex(PathSectionEnd))
+	{
+		FVector TargetPos = Path->PathPoints.Last().Location;
+		if (PathSectionEnd < (Path->PathCorridor.Num() - 1))
+		{
+			RecastNavData->GetPolyCenter(Path->PathCorridor[PathSectionEnd], TargetPos);
+		}
+
+		TArray<dtPolyRef> PathRefs;
+		for (int32 Idx = PathSectionStart; Idx <= PathSectionEnd; Idx++)
+		{
+			PathRefs.Add(Path->PathCorridor[Idx]);
+		}
+
+		const INavigationQueryFilterInterface* NavFilter = Path->Filter.IsValid() ? Path->Filter->GetImplementation() : MyNavData->GetDefaultQueryFilterImpl();
+		const dtQueryFilter* DetourFilter = ((const FRecastQueryFilter*)NavFilter)->GetAsDetourQueryFilter();
+		DetourCrowd->updateAgentFilter(AgentData->AgentIndex, DetourFilter);
+		DetourCrowd->updateAgentState(AgentData->AgentIndex, false);
+
+		const FVector RcTargetPos = Unreal2RecastPoint(TargetPos);		
+		bSuccess = DetourCrowd->requestMoveTarget(AgentData->AgentIndex, PathRefs.Last(), &RcTargetPos.X);
+		if (bSuccess)
+		{
+			bSuccess = DetourCrowd->setAgentCorridor(AgentData->AgentIndex, PathRefs.GetTypedData(), PathRefs.Num());
+		}
+	}
+#endif
+
+	return bSuccess;
+}
+
+void UCrowdManager::ClearAgentMoveTarget(const UCrowdFollowingComponent* AgentComponent) const
+{
+	SCOPE_CYCLE_COUNTER(STAT_AI_Crowd_AgentUpdateTime);
+
+#if WITH_RECAST
+	const FCrowdAgentData* AgentData = ActiveAgents.Find(AgentComponent);
+	if (AgentData && AgentData->bIsSimulated && DetourCrowd)
+	{
+		DetourCrowd->resetMoveTarget(AgentData->AgentIndex);
+		DetourCrowd->resetAgentVelocity(AgentData->AgentIndex);
+	}
+#endif
+}
+
+void UCrowdManager::PauseAgent(const UCrowdFollowingComponent* AgentComponent) const
+{
+	SCOPE_CYCLE_COUNTER(STAT_AI_Crowd_AgentUpdateTime);
+
+#if WITH_RECAST
+	const FCrowdAgentData* AgentData = ActiveAgents.Find(AgentComponent);
+	if (AgentData && AgentData->bIsSimulated && DetourCrowd)
+	{
+		DetourCrowd->setAgentWaiting(AgentData->AgentIndex);
+		DetourCrowd->resetAgentVelocity(AgentData->AgentIndex);
+	}
+#endif
+
+}
+
+void UCrowdManager::ResumeAgent(const UCrowdFollowingComponent* AgentComponent, bool bForceReplanPath) const
+{
+	SCOPE_CYCLE_COUNTER(STAT_AI_Crowd_AgentUpdateTime);
+
+#if WITH_RECAST
+	const FCrowdAgentData* AgentData = ActiveAgents.Find(AgentComponent);
+	if (AgentData && AgentData->bIsSimulated && DetourCrowd)
+	{
+		DetourCrowd->updateAgentState(AgentData->AgentIndex, bForceReplanPath);
+	}
+#endif
+}
+
+int32 UCrowdManager::GetNumNearbyAgents(const ICrowdAgentInterface* Agent) const
+{
+	int32 NumNearby = 0;
+
+#if WITH_RECAST
+	const FCrowdAgentData* AgentData = ActiveAgents.Find(Agent);
+	if (AgentData && AgentData->bIsSimulated && DetourCrowd)
+	{
+		const dtCrowdAgent* ag = DetourCrowd->getAgent(AgentData->AgentIndex);
+		NumNearby = ag ? ag->nneis : 0;
+	}
+#endif
+
+	return NumNearby;
+}
+
+bool UCrowdManager::GetAvoidanceConfig(int32 Idx, FCrowdAvoidanceConfig& Data) const
+{
+	if (AvoidanceConfig.IsValidIndex(Idx))
+	{
+		Data = AvoidanceConfig[Idx];
+		return true;
+	}
+
+	return false;
+}
+
+bool UCrowdManager::SetAvoidanceConfig(int32 Idx, const FCrowdAvoidanceConfig& Data)
+{
+	if (AvoidanceConfig.IsValidIndex(Idx))
+	{
+		AvoidanceConfig[Idx] = Data;
+	}
+#if WITH_RECAST
+	else if (Idx < DT_CROWD_MAX_OBSTAVOIDANCE_PARAMS)
+	{
+		AvoidanceConfig.SetNum(Idx + 1);
+		AvoidanceConfig[Idx] = Data;
+	}
+#endif
+	else
+	{
+		return false;
+	}
+
+	UpdateAvoidanceConfig();
+	return true;
+}
+
+void UCrowdManager::AdjustAgentPathStart(const class UCrowdFollowingComponent* AgentComponent, const FNavMeshPath* Path, int32& PathStartIdx) const
+{
+#if WITH_RECAST
+	const FCrowdAgentData* AgentData = ActiveAgents.Find(AgentComponent);
+	if (AgentData && AgentData->bIsSimulated && DetourCrowd && 
+		Path && Path->PathCorridor.Num() > PathStartIdx)
+	{
+		const dtCrowdAgent* ag = DetourCrowd->getAgent(AgentData->AgentIndex);
+		const dtPolyRef* agPath = ag->corridor.getPath();
+
+		for (int32 Idx = 0; Idx < ag->corridor.getPathCount(); Idx++)
+		{
+			const dtPolyRef TestRef = ag->corridor.getFirstPoly();
+
+			for (int32 TestIdx = PathStartIdx; TestIdx < Path->PathCorridor.Num(); TestIdx++)
+			{
+				if (Path->PathCorridor[TestIdx] == TestRef)
+				{
+					PathStartIdx = TestIdx;
+					return;
+				}
+			}
+		}
+	}
+#endif
+}
+
+void UCrowdManager::SetOffmeshConnectionPruning(bool bRemoveFromCorridor)
+{
+	bPruneStartedOffmeshConnections = bRemoveFromCorridor;
+#if WITH_RECAST
+	if (DetourCrowd)
+	{
+		DetourCrowd->setPruneStartedOffmeshConnections(bRemoveFromCorridor);
+	}
+#endif
+}
+
+void UCrowdManager::SetSingleAreaVisibilityOptimization(bool bEnable)
+{
+	bSingleAreaVisibilityOptimization = bEnable;
+#if WITH_RECAST
+	if (DetourCrowd)
+	{
+		DetourCrowd->setSingleAreaVisibilityOptimization(bEnable);
+	}
+#endif
+}
+
+#if WITH_RECAST
+
+void UCrowdManager::AddAgent(const ICrowdAgentInterface* Agent, FCrowdAgentData& AgentData) const
+{
+	SCOPE_CYCLE_COUNTER(STAT_AI_Crowd_AgentUpdateTime);
+
+	dtCrowdAgentParams Params;
+	GetAgentParams(Agent, &Params);
+
+	// create link filter for fully simulated agents
+	// (used to determine if agent can traverse smart links)
+	FRecastSpeciaLinkFilter* MyLinkFilter = NULL;
+	const UCrowdFollowingComponent* CrowdComponent = Cast<const UCrowdFollowingComponent>(Agent->GetUObjectInterfaceCrowdAgentInterface());
+	if (CrowdComponent)
+	{
+		UNavigationSystem* NavSys = Cast<UNavigationSystem>(GetOuter());
+		MyLinkFilter = new FRecastSpeciaLinkFilter(NavSys, CrowdComponent->GetOuter());
+	}
+
+	Params.linkFilter = MyLinkFilter;
+
+	const FVector RcAgentPos = Unreal2RecastPoint(Agent->GetCrowdAgentLocation());
+	const dtQueryFilter* DefaultFilter = ((const FRecastQueryFilter*)MyNavData->GetDefaultQueryFilterImpl())->GetAsDetourQueryFilter();
+
+	AgentData.AgentIndex = DetourCrowd->addAgent(&RcAgentPos.X, &Params, DefaultFilter);
+	AgentData.bIsSimulated = (Params.collisionQueryRange > 0.0f) && (CrowdComponent == NULL || CrowdComponent->IsCrowdSimulationEnabled());
+	AgentData.LinkFilter = MyLinkFilter;
+}
+
+void UCrowdManager::RemoveAgent(const ICrowdAgentInterface* Agent, const FCrowdAgentData* AgentData) const
+{
+	SCOPE_CYCLE_COUNTER(STAT_AI_Crowd_AgentUpdateTime);
+
+	DetourCrowd->removeAgent(AgentData->AgentIndex);
+	delete AgentData->LinkFilter;
+}
+
+void UCrowdManager::GetAgentParams(const ICrowdAgentInterface* Agent, struct dtCrowdAgentParams* AgentParams) const
+{
+	float CylRadius = 0.0f, CylHalfHeight = 0.0f;
+	Agent->GetCrowdAgentCollisions(CylRadius, CylHalfHeight);
+
+	FMemory::Memzero(AgentParams, sizeof(dtCrowdAgentParams));
+	AgentParams->radius = CylRadius;
+	AgentParams->height = CylHalfHeight * 2.0f;
+	// skip maxSpeed, it will be constantly updated in every tick
+	// skip maxAcceleration, we don't use Detour's movement code
+
+	const UCrowdFollowingComponent* CrowdComponent = Cast<const UCrowdFollowingComponent>(Agent->GetUObjectInterfaceCrowdAgentInterface());
+	if (CrowdComponent)
+	{
+		AgentParams->collisionQueryRange = CrowdComponent->GetCrowdCollisionQueryRange();
+		AgentParams->pathOptimizationRange = CrowdComponent->GetCrowdPathOptimizationRange();
+		AgentParams->separationWeight = CrowdComponent->GetCrowdSeparationWeight();
+		AgentParams->obstacleAvoidanceType = CrowdComponent->GetCrowdAvoidanceQuality();
+	
+		AgentParams->updateFlags =
+			(CrowdComponent->IsCrowdAnticipateTurnsEnabled() ? DT_CROWD_ANTICIPATE_TURNS : 0) |
+			(CrowdComponent->IsCrowdObstacleAvoidanceEnabled() ? DT_CROWD_OBSTACLE_AVOIDANCE : 0) |
+			(CrowdComponent->IsCrowdSeparationEnabled() ? DT_CROWD_SEPARATION : 0) |
+			(CrowdComponent->IsCrowdOptimizeVisibilityEnabled() ? (DT_CROWD_OPTIMIZE_VIS | DT_CROWD_OPTIMIZE_VIS_MULTI) : 0) |
+			(CrowdComponent->IsCrowdOptimizeTopologyEnabled() ? DT_CROWD_OPTIMIZE_TOPO : 0);
+	}
+}
+
+void UCrowdManager::UpdateSpatialData(const ICrowdAgentInterface* Agent, const FCrowdAgentData& AgentData) const
+{
+	dtCrowdAgent* ag = (dtCrowdAgent*)DetourCrowd->getAgent(AgentData.AgentIndex);
+	ag->params.maxSpeed = Agent->GetCrowdAgentMaxSpeed();
+
+	FVector RcLocation = Unreal2RecastPoint(Agent->GetCrowdAgentLocation());
+	FVector RcVelocity = Unreal2RecastPoint(Agent->GetCrowdAgentVelocity());
+
+	dtVcopy(ag->npos, &RcLocation.X);
+	dtVcopy(ag->vel, &RcVelocity.X);
+}
+
+void UCrowdManager::ApplyVelocity(const UCrowdFollowingComponent* AgentComponent, int32 AgentIndex) const
+{
+	const dtCrowdAgent* ag = DetourCrowd->getAgent(AgentIndex);
+	const dtCrowdAgentAnimation* anims = DetourCrowd->getAgentAnims();
+
+	FVector NewVelocity = Recast2UnrealPoint(ag->nvel);
+	const float* RcDestCorner = anims[AgentIndex].active ? anims[AgentIndex].endPos : 
+		ag->ncorners ? &ag->cornerVerts[0] : &ag->npos[0];
+
+	FVector DestPathCorner = Recast2UnrealPoint(RcDestCorner);
+	((UCrowdFollowingComponent*)AgentComponent)->ApplyCrowdAgentVelocity(NewVelocity, DestPathCorner, anims->active != 0);
+}
+
+void UCrowdManager::UpdateAgentPaths()
+{
+	UNavigationSystem* NavSys = Cast<UNavigationSystem>(GetOuter());
+	ARecastNavMesh* RecastNavData = Cast<ARecastNavMesh>(MyNavData);
+	FPImplRecastNavMesh* PImplNavMesh = RecastNavData ? RecastNavData->RecastNavMeshImpl : NULL;
+	if (PImplNavMesh == NULL)
+	{
+		return;
+	}
+
+	const dtCrowdAgentAnimation* AgentAnims = DetourCrowd->getAgentAnims();
+	for (auto It = ActiveAgents.CreateIterator(); It; ++It)
+	{
+		FCrowdAgentData& AgentData = It.Value();
+		if (AgentData.bIsSimulated)
+		{
+			UCrowdFollowingComponent* CrowdComponent = (UCrowdFollowingComponent*)Cast<const UCrowdFollowingComponent>(It.Key()->GetUObjectInterfaceCrowdAgentInterface());
+
+			const dtCrowdAgent* Agent = DetourCrowd->getAgent(AgentData.AgentIndex);
+			dtPolyRef AgentPolyRef = Agent->corridor.getFirstPoly();
+
+			// look for newly triggered smart links
+			const dtCrowdAgentAnimation& AnimInfo = AgentAnims[AgentData.AgentIndex];
+			if (AnimInfo.active)
+			{
+				AgentPolyRef = AnimInfo.polyRef;
+
+				if (AnimInfo.t == 0)
+				{
+					const uint32 NavLinkId = PImplNavMesh->GetLinkUserId(AnimInfo.polyRef);
+					USmartNavLinkComponent* SmartLink = NavSys->GetSmartLink(NavLinkId);
+
+					if (SmartLink)
+					{
+						FVector EndPt = Recast2UnrealPoint(AnimInfo.endPos);
+
+						// switch to waiting state
+						DetourCrowd->setAgentWaiting(AgentData.AgentIndex);
+						DetourCrowd->resetAgentVelocity(AgentData.AgentIndex);
+
+						// start using smart link
+						if (CrowdComponent)
+						{
+							CrowdComponent->SetCurrentSmartLink(SmartLink, EndPt);
+						}
+					}
+				}
+			}
+
+			// look for poly updates
+			if (AgentPolyRef != AgentData.PrevPoly && CrowdComponent)
+			{
+				CrowdComponent->OnNavNodeChanged(AgentPolyRef, AgentData.PrevPoly, Agent->corridor.getPathCount());
+				AgentData.PrevPoly = AgentPolyRef;
+			}
+		}
+	}
+}
+
+void UCrowdManager::UpdateSelectedDebug(const ICrowdAgentInterface* Agent, int32 AgentIndex) const
+{
+#if WITH_EDITOR
+	const UObject* Obj = Agent->GetUObjectInterfaceCrowdAgentInterface();
+	if (GIsEditor && Obj)
+	{
+		const AController* TestController = Cast<const AController>(Obj->GetOuter());
+		if (TestController && TestController->GetPawn() && TestController->GetPawn()->IsSelected())
+		{
+			DetourAgentDebug->idx = AgentIndex;
+		}
+	}
+#endif
+}
+
+void UCrowdManager::CreateCrowdManager()
+{
+	ARecastNavMesh* RecastNavData = Cast<ARecastNavMesh>(MyNavData);
+	FPImplRecastNavMesh* PImplNavMesh = RecastNavData ? RecastNavData->RecastNavMeshImpl : NULL;
+	dtNavMesh* NavMeshPtr = PImplNavMesh ? PImplNavMesh->GetRecastMesh() : NULL;
+
+	if (NavMeshPtr)
+	{
+		DetourCrowd = dtAllocCrowd();
+	}
+		
+	if (DetourCrowd)
+	{
+		DetourCrowd->init(MaxAgents, MaxAgentRadius, NavMeshPtr);
+		DetourCrowd->setAgentCheckInterval(NavmeshCheckInterval);
+		DetourCrowd->setOffmeshLinkSegmentRemovalRadius(MaxAgentRadius);
+		DetourCrowd->setSingleAreaVisibilityOptimization(bSingleAreaVisibilityOptimization);
+		DetourCrowd->setPruneStartedOffmeshConnections(bPruneStartedOffmeshConnections);
+		UpdateAvoidanceConfig();
+
+		for (auto It = ActiveAgents.CreateIterator(); It; ++It)
+		{
+			AddAgent(It.Key(), It.Value());
+		}
+	}
+}
+
+void UCrowdManager::DestroyCrowdManager()
+{
+	dtFreeCrowd(DetourCrowd);
+	DetourCrowd = NULL;
+}
+
+void UCrowdManager::DrawDebugCorners(const struct dtCrowdAgent* CrowdAgent) const
+{
+	FVector P0 = Recast2UnrealPoint(CrowdAgent->npos);
+	for (int32 Idx = 0; Idx < CrowdAgent->ncorners; Idx++)
+	{
+		FVector P1 = Recast2UnrealPoint(&CrowdAgent->cornerVerts[Idx * 3]);
+		DrawDebugLine(GetWorld(), P0 + CrowdDebugDrawing::Offset, P1 + CrowdDebugDrawing::Offset, CrowdDebugDrawing::Corner, false, -1.0f, SDPG_World, 2.0f);
+		P0 = P1;
+	}
+
+	if (CrowdAgent->ncorners && (CrowdAgent->cornerFlags[CrowdAgent->ncorners - 1] & DT_STRAIGHTPATH_OFFMESH_CONNECTION))
+	{
+		FVector P0 = Recast2UnrealPoint(&CrowdAgent->cornerVerts[(CrowdAgent->ncorners - 1) * 3]);
+		DrawDebugLine(GetWorld(), P0, P0 + CrowdDebugDrawing::Offset * 2.0f, CrowdDebugDrawing::CornerLink, false, -1.0f, SDPG_World, 2.0f);
+	}
+}
+
+void UCrowdManager::DrawDebugCollisionSegments(const struct dtCrowdAgent* CrowdAgent) const
+{
+	FVector Center = Recast2UnrealPoint(CrowdAgent->boundary.getCenter()) + CrowdDebugDrawing::Offset;
+	DrawDebugCylinder(GetWorld(), Center - CrowdDebugDrawing::Offset, Center, CrowdAgent->params.collisionQueryRange, 32, CrowdDebugDrawing::CollisionRange);
+
+	for (int32 Idx = 0; Idx < CrowdAgent->boundary.getSegmentCount(); Idx++)
+	{
+		const float* s = CrowdAgent->boundary.getSegment(Idx);
+		FColor Color = (dtTriArea2D(CrowdAgent->npos, s, s + 3) < 0.0f) ? CrowdDebugDrawing::CollisionSeg1 : CrowdDebugDrawing::CollisionSeg0;
+		FVector Pt0 = Recast2UnrealPoint(s);
+		FVector Pt1 = Recast2UnrealPoint(s + 3);
+
+		DrawDebugLine(GetWorld(), Pt0 + CrowdDebugDrawing::Offset, Pt1 + CrowdDebugDrawing::Offset, Color, false, -1.0f, SDPG_World, 3.5f);
+	}
+}
+
+void UCrowdManager::DrawDebugPath(const struct dtCrowdAgent* CrowdAgent) const
+{
+	ARecastNavMesh* NavMesh = Cast<ARecastNavMesh>(MyNavData);
+	if (NavMesh == NULL)
+	{
+		return;
+	}
+
+	NavMesh->BeginBatchQuery();
+	
+	const dtPolyRef* Path = CrowdAgent->corridor.getPath();
+	TArray<FVector> Verts;
+
+	for (int32 Idx = 0; Idx < CrowdAgent->corridor.getPathCount(); Idx++)
+	{
+		Verts.Reset();
+		NavMesh->GetPolyVerts(Path[Idx], Verts);
+
+		uint16 PolyFlags = 0;
+		uint16 AreaFlags = 0;
+		NavMesh->GetPolyFlags(Path[Idx], PolyFlags, AreaFlags);
+		const FColor PolyColor = AreaFlags != 1 ? CrowdDebugDrawing::Path : CrowdDebugDrawing::PathSpecial;
+
+		for (int32 VertIdx = 0; VertIdx < Verts.Num(); VertIdx++)
+		{
+			const FVector Pt0 = Verts[VertIdx];
+			const FVector Pt1 = Verts[(VertIdx + 1) % Verts.Num()];
+
+			DrawDebugLine(GetWorld(), Pt0 + CrowdDebugDrawing::Offset * 0.5f, Pt1 + CrowdDebugDrawing::Offset * 0.5f, PolyColor, false
+				, /*LifeTime*/-1.f, /*DepthPriority*/0
+				, /*Thickness*/CrowdDebugDrawing::LineThickness);
+		}
+	}
+
+	NavMesh->FinishBatchQuery();
+}
+
+void UCrowdManager::DrawDebugVelocityObstacles(const struct dtCrowdAgent* CrowdAgent) const
+{
+	FVector Center = Recast2UnrealPoint(CrowdAgent->npos) + CrowdDebugDrawing::Offset;
+	DrawDebugCylinder(GetWorld(), Center - CrowdDebugDrawing::Offset, Center, CrowdAgent->params.maxSpeed, 32, CrowdDebugDrawing::AvoidanceRange);
+
+	float BestSampleScore = -1.0f;
+	FVector BestSampleLocation = FVector::ZeroVector;
+
+	for (int32 Idx = 0; Idx < DetourAvoidanceDebug->getSampleCount(); Idx++)
+	{
+		const float* p = DetourAvoidanceDebug->getSampleVelocity(Idx);
+		const float sr = DetourAvoidanceDebug->getSampleSize(Idx);
+		const float pen = DetourAvoidanceDebug->getSamplePenalty(Idx);
+		const float pen2 = DetourAvoidanceDebug->getSamplePreferredSidePenalty(Idx);
+
+		FVector SamplePos = Center + Recast2UnrealPoint(p);
+
+		if (BestSampleScore <= -1.0f || pen < BestSampleScore)
+		{
+			BestSampleScore = pen;
+			BestSampleLocation = SamplePos;
+		}
+
+		float SamplePenalty = pen * 0.75f + pen2 * 0.25f;
+		FColor SampleColor = FColor::MakeRedToGreenColorFromScalar(1.0f - SamplePenalty);
+			
+		FPlane Plane(0, 0, 1, SamplePos.Z);
+		DrawDebugSolidPlane(GetWorld(), Plane, SamplePos, sr, SampleColor);
+	}
+
+	if (BestSampleScore >= 0.0f)
+	{
+		DrawDebugLine(GetWorld(), BestSampleLocation + FVector(0, 0, 100), BestSampleLocation + FVector(0, 0, -100), FColor::Green);
+	}
+}
+
+void UCrowdManager::DrawDebugPathOptimization(const struct dtCrowdAgent* CrowdAgent) const
+{
+	FVector Pt0 = Recast2UnrealPoint(DetourAgentDebug->optStart) + CrowdDebugDrawing::Offset * 1.25f;
+	FVector Pt1 = Recast2UnrealPoint(DetourAgentDebug->optEnd) + CrowdDebugDrawing::Offset * 1.25f;
+
+	DrawDebugLine(GetWorld(), Pt0, Pt1, CrowdDebugDrawing::PathOpt, false, -1.0f, SDPG_World, 2.5f);
+}
+
+void UCrowdManager::DrawDebugNeighbors(const struct dtCrowdAgent* CrowdAgent) const
+{
+	FVector Center = Recast2UnrealPoint(CrowdAgent->npos) + CrowdDebugDrawing::Offset;
+	DrawDebugCylinder(GetWorld(), Center - CrowdDebugDrawing::Offset, Center, CrowdAgent->params.collisionQueryRange, 32, CrowdDebugDrawing::CollisionRange);
+
+	for (int32 Idx = 0; Idx < CrowdAgent->nneis; Idx++)
+	{
+		const dtCrowdAgent* nei = DetourCrowd->getAgent(CrowdAgent->neis[Idx].idx);
+		if (nei)
+		{
+			FVector Pt0 = Recast2UnrealPoint(nei->npos) + CrowdDebugDrawing::Offset;
+			DrawDebugLine(GetWorld(), Center, Pt0, CrowdDebugDrawing::Neighbor);
+		}
+	}
+}
+
+#endif // WITH_RECAST
+
+#if WITH_EDITOR
+
+void UCrowdManager::DebugTick() const
+{
+#if WITH_RECAST
+	if (DetourCrowd == NULL || DetourAgentDebug == NULL)
+	{
+		return;
+	}
+
+	for (auto It = ActiveAgents.CreateConstIterator(); It; ++It)
+	{
+		const FCrowdAgentData& AgentData = It.Value();
+		UpdateSelectedDebug(It.Key(), AgentData.AgentIndex);
+	}
+
+	// on screen debugging
+	const struct dtCrowdAgent* SelectedAgent = DetourAgentDebug->idx >= 0 ? DetourCrowd->getAgent(DetourAgentDebug->idx) : NULL;
+	if (SelectedAgent && CrowdDebugDrawing::bDebugSelectedActors)
+	{
+		if (CrowdDebugDrawing::bDrawDebugCorners)
+		{
+			DrawDebugCorners(SelectedAgent);
+		}
+
+		if (CrowdDebugDrawing::bDrawDebugCollisionSegments)
+		{
+			DrawDebugCollisionSegments(SelectedAgent);
+		}
+
+		if (CrowdDebugDrawing::bDrawDebugPath)
+		{
+			DrawDebugPath(SelectedAgent);
+		}
+
+		if (CrowdDebugDrawing::bDrawDebugVelocityObstacles)
+		{
+			DrawDebugVelocityObstacles(SelectedAgent);
+		}
+
+		if (CrowdDebugDrawing::bDrawDebugPathOptimization)
+		{
+			DrawDebugPathOptimization(SelectedAgent);
+		}
+
+		if (CrowdDebugDrawing::bDrawDebugNeighbors)
+		{
+			DrawDebugNeighbors(SelectedAgent);
+		}
+	}
+
+	// vislog debugging
+	if (CrowdDebugDrawing::bDebugVisLog)
+	{
+		for (auto It = ActiveAgents.CreateConstIterator(); It; ++It)
+		{
+			const ICrowdAgentInterface* IAgent = It.Key();
+			const UObject* AgentOb = IAgent ? IAgent->GetUObjectInterfaceCrowdAgentInterface() : NULL;
+			const AActor* LogOwner = AgentOb ? Cast<const AActor>(AgentOb->GetOuter()) : NULL;
+
+			const FCrowdAgentData& AgentData = It.Value();
+			const struct dtCrowdAgent* CrowdAgent = DetourCrowd->getAgent(AgentData.AgentIndex);
+
+			if (CrowdAgent && LogOwner)
+			{
+				FVector P0 = Recast2UnrealPoint(CrowdAgent->npos);
+				for (int32 Idx = 0; Idx < CrowdAgent->ncorners; Idx++)
+				{
+					FVector P1 = Recast2UnrealPoint(&CrowdAgent->cornerVerts[Idx * 3]);
+					UE_VLOG_SEGMENT(LogOwner, LogEngineCrowdFollowing, Log, P0 + CrowdDebugDrawing::Offset, P1 + CrowdDebugDrawing::Offset, CrowdDebugDrawing::Corner, TEXT(""));
+					P0 = P1;
+				}
+
+				if (CrowdAgent->ncorners && (CrowdAgent->cornerFlags[CrowdAgent->ncorners - 1] & DT_STRAIGHTPATH_OFFMESH_CONNECTION))
+				{
+					FVector P0 = Recast2UnrealPoint(&CrowdAgent->cornerVerts[(CrowdAgent->ncorners - 1) * 3]);
+					UE_VLOG_SEGMENT(LogOwner, LogEngineCrowdFollowing, Log, P0, P0 + CrowdDebugDrawing::Offset * 2.0f, CrowdDebugDrawing::CornerLink, TEXT(""));
+				}
+			}
+		}
+	}
+#endif	// WITH_RECAST
+}
+
+#endif // WITH_EDITOR
+
+void UCrowdManager::UpdateNavData()
+{
+	if (MyNavData == NULL)
+	{
+		UNavigationSystem* NavSys = Cast<UNavigationSystem>(GetOuter());
+		if (NavSys)
+		{
+			for (int32 Idx = 0; Idx < NavSys->NavDataSet.Num(); Idx++)
+			{
+				ARecastNavMesh* RecastNavData = Cast<ARecastNavMesh>(NavSys->NavDataSet[Idx]);
+				if (RecastNavData && RecastNavData->IsSupportingDefaultAgent())
+				{
+					MyNavData = RecastNavData;
+					RecastNavData->OnNavMeshUpdate.AddUObject(this, &UCrowdManager::OnNavMeshUpdate);
+					OnNavMeshUpdate();
+
+					break;
+				}
+			}
+		}
+	}
+}
+
+void UCrowdManager::OnNavMeshUpdate()
+{
+#if WITH_RECAST
+	DestroyCrowdManager();
+	CreateCrowdManager();
+#endif // WITH_RECAST
+}
+
+void UCrowdManager::UpdateAvoidanceConfig()
+{
+#if WITH_RECAST
+	if (DetourCrowd == NULL)
+	{
+		return;
+	}
+
+	for (int32 Idx = 0; Idx < AvoidanceConfig.Num(); Idx++)
+	{
+		const FCrowdAvoidanceConfig& ConfigInfo = AvoidanceConfig[Idx];
+		
+		dtObstacleAvoidanceParams params;
+		params.velBias = ConfigInfo.VelocityBias;
+		params.weightDesVel = ConfigInfo.DesiredVelocityWeight;
+		params.weightCurVel = ConfigInfo.CurrentVelocityWeight;
+		params.weightSide = ConfigInfo.SideBiasWeight;
+		params.weightToi = ConfigInfo.ImpactTimeWeight;
+		params.horizTime = ConfigInfo.ImpactTimeRange;
+		params.adaptiveDivs = ConfigInfo.AdaptiveDivisions;
+		params.adaptiveRings = ConfigInfo.AdaptiveRings;
+		params.adaptiveDepth = ConfigInfo.AdaptiveDepth;
+
+		DetourCrowd->setObstacleAvoidanceParams(Idx, &params);
+	}
+#endif // WITH_RECAST
+}
+
+void UCrowdManager::PostProximityUpdate()
+{
+	// empty in base class
+}
+
+class UWorld* UCrowdManager::GetWorld() const
+{
+	UNavigationSystem* NavSys = Cast<UNavigationSystem>(GetOuter());
+	return NavSys ? NavSys->GetWorld() : NULL;
+}
+
+UCrowdManager* UCrowdManager::GetCurrent(class UObject* WorldContextObject)
+{
+	UNavigationSystem* NavSys = UNavigationSystem::GetCurrent(WorldContextObject);
+	return NavSys ? NavSys->GetCrowdManager() : NULL;
+}
+
+UCrowdManager* UCrowdManager::GetCurrent(class UWorld* World)
+{
+	UNavigationSystem* NavSys = UNavigationSystem::GetCurrent(World);
+	return NavSys ? NavSys->GetCrowdManager() : NULL;
+}
