Extracting AI code from the engine into a separate module, pass 1 #UE4

diff --git a/Engine/Source/Runtime/AIModule/Private/Navigation/NavigationComponent.cpp b/Engine/Source/Runtime/AIModule/Private/Navigation/NavigationComponent.cpp
new file mode 100644
--- /dev/null
+++ b/Engine/Source/Runtime/AIModule/Private/Navigation/NavigationComponent.cpp
@@ -0,0 +1,915 @@
+// Copyright 1998-2014 Epic Games, Inc. All Rights Reserved.
+
+#include "AIModulePrivate.h"
+#include "VisualLog.h"
+#include "AI/Navigation/SmartNavLinkComponent.h"
+#include "Navigation/NavigationComponent.h"
+
+//----------------------------------------------------------------------//
+// Life cycle                                                        
+//----------------------------------------------------------------------//
+
+UNavigationComponent::UNavigationComponent(const class FPostConstructInitializeProperties& PCIP)
+	: Super(PCIP)
+	, bIsPathPartial(false)
+	, AsynPathQueryID(INVALID_NAVQUERYID)
+{
+	PrimaryComponentTick.bCanEverTick = true;
+	PrimaryComponentTick.bStartWithTickEnabled = false;
+	
+	bNotifyPathFollowing = true;
+	bRepathWhenInvalid = true;
+	bUpdateForSmartLinks = true;
+	bWantsInitializeComponent = true;
+	bAutoActivate = true;
+
+	NavDataFlags = 0;
+}
+
+void UNavigationComponent::InitializeComponent()
+{
+	Super::InitializeComponent();
+
+	// setting up delegate here in case function used has virtual overrides
+	PathObserverDelegate = FNavigationPath::FPathObserverDelegate::CreateUObject(this, &UNavigationComponent::OnPathInvalid);
+	
+	UpdateCachedComponents();
+}
+
+void UNavigationComponent::UpdateCachedComponents()
+{
+	// try to find path following component in owning actor
+	AActor* MyOwner = GetOwner();
+	if (MyOwner)
+	{
+		PathFollowComp = MyOwner->FindComponentByClass<UPathFollowingComponent>();
+	}
+}
+
+void UNavigationComponent::TickComponent(float DeltaTime, enum ELevelTick TickType, FActorComponentTickFunction *ThisTickFunction)
+{
+	// already requested not to tick. 
+	if (bShoudTick == false)
+	{
+		return;
+	}
+
+	Super::TickComponent(DeltaTime, TickType, ThisTickFunction);
+
+	check( GetOuter() != NULL );
+
+	bool bSuccess = false;
+
+	if (GoalActor != NULL)
+	{
+		const FVector GoalActorLocation = GetCurrentMoveGoal(GoalActor, GetOwner());
+		const float DistanceSq = FVector::DistSquared(GoalActorLocation, OriginalGoalActorLocation);
+
+		if (DistanceSq > RepathDistanceSq)
+		{
+			if (RePathTo(GoalActorLocation, StoredQueryFilter))
+			{
+				bSuccess = true;
+				OriginalGoalActorLocation = GoalActorLocation;
+			}
+			else
+			{
+				FAIMessage::Send(Cast<AController>(GetOwner()), FAIMessage(UBrainComponent::AIMessage_RepathFailed, this));
+			}
+		}
+		else
+		{
+			bSuccess = true;
+		}
+	}
+
+	if (bSuccess == false)
+	{
+		ResetTransientData();
+	}
+}
+
+void UNavigationComponent::OnPathInvalid(FNavigationPath* InvalidatedPath)
+{
+	if (InvalidatedPath == NULL || InvalidatedPath != Path.Get())
+    {
+		return;
+    }
+
+	// try to repath automatically
+	if (!bIsPaused && bRepathWhenInvalid)
+	{
+		RepathData.GoalActor = GoalActor;
+		RepathData.GoalLocation = CurrentGoal;
+		RepathData.bSimplePath = bUseSimplePath;
+		bIsWaitingForRepath = true;
+
+		FSimpleDelegateGraphTask::CreateAndDispatchWhenReady(
+			FSimpleDelegateGraphTask::FDelegate::CreateUObject(this, &UNavigationComponent::DeferredRepathToGoal)
+			, TEXT("Deferred restore path")
+			, NULL
+			, ENamedThreads::GameThread
+			);
+	}
+
+	// reset state
+	const bool bWasPaused = bIsPaused;
+	ResetTransientData();
+	bIsPaused = bWasPaused;
+
+	// and notify observer
+	if (MyPathObserver && !bRepathWhenInvalid)
+	{
+		MyPathObserver->OnPathInvalid(this);
+	}
+}
+
+bool UNavigationComponent::RePathTo(const FVector& GoalLocation, TSharedPtr<const FNavigationQueryFilter> InStoredQueryFilter)
+{
+	return bUseSimplePath ? GenerateSimplePath(GoalLocation) :
+		bDoAsyncPathfinding ? AsyncGeneratePathTo(GoalLocation, StoredQueryFilter) : GeneratePathTo(GoalLocation, InStoredQueryFilter);
+}
+
+FNavLocation UNavigationComponent::GetRandomPointOnNavMesh() const
+{
+#if WITH_RECAST
+	ARecastNavMesh* const Nav = (ARecastNavMesh*)GetWorld()->GetNavigationSystem()->GetMainNavData(FNavigationSystem::Create);
+	if (Nav)
+	{
+		return Nav->GetRandomPoint();
+	}
+#endif // WITH_RECAST
+	return FNavLocation();
+}
+
+FNavLocation UNavigationComponent::ProjectPointToNavigation(const FVector& Location) const
+{
+	FNavLocation OutLocation(Location);
+	
+	if (GetNavData() != NULL)
+	{
+		if (!MyNavData->ProjectPoint(Location, OutLocation, GetQueryExtent()))
+		{
+			OutLocation = FNavLocation();
+		}
+	}
+
+	return OutLocation;
+}
+
+bool UNavigationComponent::ProjectPointToNavigation(const FVector& Location, FNavLocation& ProjectedLocation) const
+{
+	return GetNavData() != NULL && MyNavData->ProjectPoint(Location, ProjectedLocation, GetQueryExtent());
+}
+
+bool UNavigationComponent::ProjectPointToNavigation(const FVector& Location, FVector& ProjectedLocation) const
+{
+	FNavLocation NavLoc;
+	const bool bSuccess = ProjectPointToNavigation(Location, NavLoc);
+	ProjectedLocation = NavLoc.Location;
+	return bSuccess;
+}
+
+bool UNavigationComponent::AsyncGeneratePathTo(const FVector& GoalLocation, TSharedPtr<const FNavigationQueryFilter> QueryFilter)
+{
+	return GetOuter() && MyNavAgent && AsyncGeneratePath(MyNavAgent->GetNavAgentLocation(), GoalLocation, QueryFilter);
+}
+
+bool UNavigationComponent::AsyncGeneratePath(const FVector& FromLocation, const FVector& GoalLocation, TSharedPtr<const FNavigationQueryFilter> QueryFilter)
+{
+	UNavigationSystem* NavSys = UNavigationSystem::GetCurrent(GetWorld());
+	if (NavSys == NULL || GetOuter() == NULL || MyNavAgent == NULL || GetNavData() == NULL)
+	{
+		return false;
+	}
+
+	// Make sure we're trying to path to the closest valid location TO that location rather than the absolute location.
+	// After talking with Steve P., if we need to ever allow pathing WITHOUT projecting to nav-mesh, it will probably be
+	// best to do so using a flag while keeping this as the default behavior.  NOTE: If any changes to this behavior
+	// are made, you should search for other places in UNavigationComponent using NavMeshGoalLocation and make sure the
+	// behavior remains consistent.
+	const FNavAgentProperties* NavAgentProps = MyNavAgent->GetNavAgentProperties();
+	if (NavAgentProps != NULL)
+	{
+		FVector NavMeshGoalLocation;
+
+		if (ProjectPointToNavigation(GoalLocation, NavMeshGoalLocation) == QuerySuccess)
+		{
+			const EPathFindingMode::Type Mode = bDoHierarchicalPathfinding ? EPathFindingMode::Hierarchical : EPathFindingMode::Regular;
+			FPathFindingQuery Query(GetOwner(), GetNavData(), FromLocation, NavMeshGoalLocation, QueryFilter);
+			Query.NavDataFlags = NavDataFlags;
+
+			AsynPathQueryID = NavSys->FindPathAsync(*MyNavAgent->GetNavAgentProperties(), Query,
+				FNavPathQueryDelegate::CreateUObject(this, &UNavigationComponent::AsyncGeneratePath_OnCompleteCallback), Mode);
+		}
+		else
+		{
+			// projecting destination point failed
+			AsynPathQueryID = INVALID_NAVQUERYID;
+
+			UE_VLOG(GetOwner(), LogAINavigation, Warning, TEXT("AI trying to generate path to point off of navmesh"));
+			UE_VLOG_LOCATION(GetOwner(), LogAINavigation, Warning, GoalLocation, 30, FColor::Red, TEXT("Invalid pathing GoalLocation"));
+			if (MyNavData != NULL)
+			{
+				UE_VLOG_BOX(GetOwner(), LogAINavigation, Warning, FBox(GoalLocation - GetQueryExtent(), GoalLocation + GetQueryExtent()), FColor::Red, TEXT_EMPTY);
+			}
+		}
+	}
+	else
+	{
+		AsynPathQueryID = INVALID_NAVQUERYID;
+	}
+
+	return AsynPathQueryID != INVALID_NAVQUERYID;
+}
+
+void UNavigationComponent::Pause()
+{
+	bIsPaused = true;
+	SetComponentTickEnabledAsync(false);
+}
+
+void UNavigationComponent::AbortAsyncFindPathRequest()
+{
+	if (AsynPathQueryID == INVALID_NAVQUERYID || GetWorld()->GetNavigationSystem() == NULL)
+	{
+		return;
+	}
+
+	GetWorld()->GetNavigationSystem()->AbortAsyncFindPathRequest((uint32)AsynPathQueryID);
+
+	AsynPathQueryID = INVALID_NAVQUERYID;
+}
+
+void UNavigationComponent::DeferredResumePath()
+{
+	ResumePath();
+}
+
+bool UNavigationComponent::ResumePath()
+{
+	const bool bWasPaused = bIsPaused;
+
+	bIsPaused = false;
+	
+	if (bWasPaused && Path.IsValid() && Path->IsValid())
+	{
+		if (IsFollowing() == true)
+		{
+			// make sure ticking is enabled (and it shouldn't be enabled before _first_ async path finding)
+			SetComponentTickEnabledAsync(true);
+		}
+
+		// don't notify about path update, it's the same path...
+		CurrentGoal = Path->PathPoints[ Path->PathPoints.Num() - 1 ].Location;
+		return true;
+	}
+
+	return false;
+}
+
+void UNavigationComponent::AsyncGeneratePath_OnCompleteCallback(uint32 QueryID, ENavigationQueryResult::Type Result, FNavPathSharedPtr ResultPath)
+{
+	if (QueryID != uint32(AsynPathQueryID))
+	{
+		UE_VLOG(GetOwner(), LogAINavigation, Display, TEXT("Pathfinding query %u finished, but it's different from AsynPathQueryID (%d)"), QueryID, AsynPathQueryID);
+		return;
+	}
+
+	check(GetOuter() != NULL);
+	
+	if (Result == ENavigationQueryResult::Success && ResultPath.IsValid() && ResultPath->IsValid())
+	{
+		const bool bIsSamePath = Path.IsValid() == true && ((const FNavMeshPath*)Path.Get())->ContainsWithSameEnd((const FNavMeshPath*)ResultPath.Get())
+			&& FVector::DistSquared(Path->GetDestinationLocation(), ResultPath->GetDestinationLocation()) < KINDA_SMALL_NUMBER;
+
+		if (bIsSamePath == false)
+		{
+			if (IsFollowing() == true)
+			{
+				// make sure ticking is enabled (and it shouldn't be enabled before _first_ async path finding)
+				SetComponentTickEnabledAsync(true);
+			}
+
+			SetPath(ResultPath);
+			CurrentGoal = ResultPath->PathPoints[ ResultPath->PathPoints.Num() - 1 ].Location;
+			NotifyPathUpdate();
+		}
+		else
+		{
+			UE_VLOG(GetOwner(), LogAINavigation, Display, TEXT("Skipping newly generated path due to it being the same as current one"));
+		}
+	}
+	else
+	{
+		ResetTransientData();
+		if (MyPathObserver != NULL)
+		{
+			MyPathObserver->OnPathFailed(this);  
+		}
+	}
+}
+
+bool UNavigationComponent::GeneratePathTo(const FVector& GoalLocation, TSharedPtr<const FNavigationQueryFilter> QueryFilter)
+{
+	if (GetWorld() == NULL || GetWorld()->GetNavigationSystem() == NULL || GetOuter() == NULL)
+	{
+		return false;
+	}
+
+	// Make sure we're trying to path to the closest valid location TO that location rather than the absolute location.
+	// After talking with Steve P., if we need to ever allow pathing WITHOUT projecting to nav-mesh, it will probably be
+	// best to do so using a flag while keeping this as the default behavior.  NOTE:` If any changes to this behavior
+	// are made, you should search for other places in UNavigationComponent using NavMeshGoalLocation and make sure the
+	// behavior remains consistent.
+	
+	// make sure that nav data is updated
+	GetNavData();
+
+	const FNavAgentProperties* NavAgentProps = MyNavAgent ? MyNavAgent->GetNavAgentProperties() : NULL;
+	if (NavAgentProps != NULL)
+	{
+		FVector NavMeshGoalLocation;
+
+#if ENABLE_VISUAL_LOG
+		UE_VLOG_LOCATION(GetOwner(), LogAINavigation, Log, GoalLocation, 34, FColor(0, 127, 14), TEXT("GoalLocation"));
+		UE_VLOG_BOX(GetOwner(), LogAINavigation, Log, FBox(GoalLocation - GetQueryExtent(), GoalLocation + GetQueryExtent()), FColor::Green, TEXT_EMPTY);
+#endif
+
+		if (ProjectPointToNavigation(GoalLocation, NavMeshGoalLocation) == QuerySuccess)
+		{
+			UE_VLOG_SEGMENT(GetOwner(), LogAINavigation, Log, GoalLocation, NavMeshGoalLocation, FColor::Blue, TEXT_EMPTY);
+			UE_VLOG_LOCATION(GetOwner(), LogAINavigation, Log, NavMeshGoalLocation, 34, FColor::Blue, TEXT_EMPTY);
+			
+			UNavigationSystem* NavSys = GetWorld()->GetNavigationSystem();
+			FPathFindingQuery Query(GetOwner(), MyNavData, MyNavAgent->GetNavAgentLocation(), NavMeshGoalLocation, QueryFilter);
+			Query.NavDataFlags = NavDataFlags;
+
+			const EPathFindingMode::Type Mode = bDoHierarchicalPathfinding ? EPathFindingMode::Hierarchical : EPathFindingMode::Regular;
+			FPathFindingResult Result = NavSys->FindPathSync(*MyNavAgent->GetNavAgentProperties(), Query, Mode);
+
+			// try reversing direction
+			if (bSearchFromGoalWhenOutOfNodes && !bDoHierarchicalPathfinding &&
+				Result.Path.IsValid() && Result.Path->DidSearchReachedLimit())
+			{
+				// quick check if path exists
+				bool bPathExists = false;
+				{
+					DECLARE_SCOPE_CYCLE_COUNTER(TEXT("Pathfinding: HPA* test time"), STAT_Navigation_PathVerifyTime, STATGROUP_Navigation);
+					bPathExists = NavSys->TestPathSync(Query, EPathFindingMode::Hierarchical);
+				}
+
+				UE_VLOG(GetOwner(), LogAINavigation, Log, TEXT("GeneratePathTo: out of nodes, HPA* path: %s"),
+					bPathExists ? TEXT("exists, trying reversed search") : TEXT("doesn't exist"));
+
+				// reverse search
+				if (bPathExists)
+				{
+					DECLARE_SCOPE_CYCLE_COUNTER(TEXT("Pathfinding: reversed test time"), STAT_Navigation_PathReverseTime, STATGROUP_Navigation);
+
+					TSharedPtr<FNavigationQueryFilter> Filter = QueryFilter.IsValid() ? QueryFilter->GetCopy() : MyNavData->GetDefaultQueryFilter()->GetCopy();					
+					Filter->SetBacktrackingEnabled(true);
+
+					FPathFindingQuery ReversedQuery(GetOwner(), MyNavData, Query.EndLocation, Query.StartLocation, Filter);
+					ReversedQuery.NavDataFlags = NavDataFlags;
+
+					Result = NavSys->FindPathSync(*MyNavAgent->GetNavAgentProperties(), Query, Mode);
+				}
+			}
+
+			if (Result.IsSuccessful() || Result.IsPartial())
+			{
+				CurrentGoal = NavMeshGoalLocation;
+				SetPath(Result.Path);
+
+				if (IsFollowing() == true)
+				{
+					// make sure ticking is enabled (and it shouldn't be enabled before _first_ async path finding)
+					SetComponentTickEnabledAsync(true);
+				}
+
+				NotifyPathUpdate();
+				return true;
+			}
+			else
+			{
+				UE_VLOG(GetOwner(), LogAINavigation, Display, TEXT("Failed to generate path to destination"));
+			}
+		}
+		else
+		{
+			UE_VLOG(GetOwner(), LogAINavigation, Display, TEXT("Destination not on navmesh (and nowhere near!)"));
+		}
+	}
+	else
+	{
+		UE_VLOG(GetOwner(), LogAINavigation, Display, TEXT("UNavigationComponent::GeneratePathTo: NavAgentProps == NULL (Probably Pawn died)"));
+	}
+
+	return false;
+}
+
+bool UNavigationComponent::GeneratePath(const FVector& FromLocation, const FVector& GoalLocation, TSharedPtr<const FNavigationQueryFilter> QueryFilter)
+{
+	UNavigationSystem* NavSys = UNavigationSystem::GetCurrent(GetWorld());
+	if (NavSys == NULL || GetOuter() == NULL || MyNavAgent == NULL || GetNavData() == NULL)
+	{
+		return false;
+	}
+
+	const EPathFindingMode::Type Mode = bDoHierarchicalPathfinding ? EPathFindingMode::Hierarchical : EPathFindingMode::Regular;
+	const FNavLocation NavMeshGoalLocation = ProjectPointToNavigation(GoalLocation);
+	FPathFindingQuery Query(GetOwner(), GetNavData(), FromLocation, NavMeshGoalLocation.Location, QueryFilter);
+	Query.NavDataFlags = NavDataFlags;
+
+	FPathFindingResult Result = NavSys->FindPathSync(*MyNavAgent->GetNavAgentProperties(), Query, Mode);
+	if (Result.IsSuccessful() || Result.IsPartial())
+	{
+		CurrentGoal = NavMeshGoalLocation.Location;
+		SetPath(Result.Path);
+		return true;
+	}
+
+	return false;
+}
+
+bool UNavigationComponent::GenerateSimplePath(const FVector& Destination)
+{
+	if (GetWorld() == NULL || GetOwner() == NULL || MyNavAgent == NULL)
+	{
+		return false;
+	}
+	
+	TArray<FVector> PathPoints;
+	PathPoints.Add(MyNavAgent->GetNavAgentLocation());
+	PathPoints.Add(Destination);
+
+	FNavPathSharedPtr SimplePath = MakeShareable(new FNavigationPath(PathPoints, NULL));
+	
+	CurrentGoal = Destination;
+	SetPath(SimplePath);
+	return true;
+}
+
+bool UNavigationComponent::GetPathPoint(uint32 PathVertIdx, FNavPathPoint& PathPoint) const
+{
+	if (Path.IsValid() && Path->IsValid() && (int32)PathVertIdx < Path->PathPoints.Num())
+	{
+		PathPoint = Path->PathPoints[PathVertIdx];
+		return true;
+	}
+
+	return false;
+}
+
+bool UNavigationComponent::HasValidPath() const
+{ 
+	return Path.IsValid() && Path->IsValid(); 
+}
+
+bool UNavigationComponent::FindSimplePathToActor(const AActor* NewGoalActor, float RepathDistance)
+{
+	bool bPathGenerationSucceeded = false;
+
+	if (bIsPaused && NewGoalActor != NULL && GoalActor == NewGoalActor
+		&& Path.IsValid() && Path->IsValid() && Path->GetOwner() == NULL)
+	{
+		RepathDistanceSq = FMath::Square(RepathDistance);
+		OriginalGoalActorLocation = GetCurrentMoveGoal(NewGoalActor, GetOwner());
+
+		bPathGenerationSucceeded = ResumePath();
+	}
+	else 
+	{
+		if (NewGoalActor != NULL)
+		{
+			const FVector NewGoalActorLocation = GetCurrentMoveGoal(NewGoalActor, GetOwner());
+			bPathGenerationSucceeded = GenerateSimplePath(NewGoalActorLocation);
+
+			if (bPathGenerationSucceeded)
+			{
+				GoalActor = NewGoalActor;
+				RepathDistanceSq = FMath::Square(RepathDistance);
+				OriginalGoalActorLocation = NewGoalActorLocation;
+				bUseSimplePath = true;
+
+				// if doing sync pathfinding enabling component's ticking needs to be done now
+				SetComponentTickEnabledAsync(true);
+			}
+		}
+
+		if (!bPathGenerationSucceeded)
+		{
+			UE_VLOG(GetOwner(), LogAINavigation, Warning, TEXT("Failed to generate simple path to %s"), *GetNameSafe(NewGoalActor));
+		}
+	}
+
+	if (bPathGenerationSucceeded == false)
+	{
+		ResetTransientData();
+	}
+
+	return bPathGenerationSucceeded;
+}
+
+bool UNavigationComponent::FindPathToActor(const AActor* NewGoalActor, TSharedPtr<const FNavigationQueryFilter> QueryFilter, float RepathDistance)
+{
+	bool bPathGenerationSucceeded = false;
+
+	if (bIsPaused && NewGoalActor != NULL && GoalActor == NewGoalActor 
+		&& Path.IsValid() && Path->IsValid())
+	{
+		RepathDistanceSq = FMath::Square(RepathDistance);
+		OriginalGoalActorLocation = GetCurrentMoveGoal(NewGoalActor, GetOwner());
+
+		if (bDoAsyncPathfinding == true)
+		{
+			FSimpleDelegateGraphTask::CreateAndDispatchWhenReady(
+				FSimpleDelegateGraphTask::FDelegate::CreateUObject(this, &UNavigationComponent::DeferredResumePath)
+				, TEXT("Deferred resume path")
+				, NULL
+				, ENamedThreads::GameThread
+				);
+			bPathGenerationSucceeded = true;
+		}
+		else
+		{
+			bPathGenerationSucceeded = ResumePath();
+		}
+	}
+	else
+	{
+		if (NewGoalActor != NULL)
+		{
+			if (QueryFilter.IsValid() == false)
+			{
+				QueryFilter = GetStoredQueryFilter();
+			}
+
+			const FVector NewGoalActorLocation = GetCurrentMoveGoal(NewGoalActor, GetOwner());
+			if (bDoAsyncPathfinding == true)
+			{
+				bPathGenerationSucceeded = AsyncGeneratePathTo(NewGoalActorLocation, QueryFilter);
+			}
+			else
+			{
+				bPathGenerationSucceeded = GeneratePathTo(NewGoalActorLocation, QueryFilter);
+			}
+
+			if (bPathGenerationSucceeded)
+			{
+				// if there's a specific query filter store it
+				StoredQueryFilter = QueryFilter;
+
+				GoalActor = NewGoalActor;
+				RepathDistanceSq = FMath::Square(RepathDistance);
+				OriginalGoalActorLocation = NewGoalActorLocation;
+				bUseSimplePath = false;
+
+				if (bDoAsyncPathfinding == false)
+				{
+					// if doing sync pathfinding enabling component's ticking needs to be done now
+					SetComponentTickEnabledAsync(true);
+				}
+			}
+		}
+
+		if (!bPathGenerationSucceeded)
+		{
+			UE_VLOG(GetOwner(), LogAINavigation, Log, TEXT("%s failed to generate path to %s"), *GetNameSafe(GetOwner()), *GetNameSafe(NewGoalActor));
+		}
+	}
+
+	if (bPathGenerationSucceeded == false)
+	{
+		ResetTransientData();
+	}
+
+	return bPathGenerationSucceeded;
+}
+
+bool UNavigationComponent::FindSimplePathToLocation(const FVector& DestLocation)
+{
+	bool bPathGenerationSucceeded = false;
+
+	if (bIsPaused && Path.IsValid() && Path->IsValid() && Path->GetOwner() == NULL
+		&& FVector::DistSquared(CurrentGoal, DestLocation) < SMALL_NUMBER)
+	{
+		bPathGenerationSucceeded = ResumePath();
+	}
+
+	if (bPathGenerationSucceeded == false)
+	{
+		ResetTransientData();
+
+		bPathGenerationSucceeded = GenerateSimplePath(DestLocation);
+	}
+
+	return bPathGenerationSucceeded;
+}
+
+bool UNavigationComponent::FindPathToLocation(const FVector& DestLocation, TSharedPtr<const FNavigationQueryFilter> QueryFilter)
+{
+	bool bPathGenerationSucceeded = false;
+
+	if (bIsPaused == true && Path.IsValid() && Path->IsValid() 
+		&& FVector::DistSquared(CurrentGoal, DestLocation) < SMALL_NUMBER)
+	{
+		// path to destination already generated. Use it.
+
+		if (bDoAsyncPathfinding == true)
+		{
+			FSimpleDelegateGraphTask::CreateAndDispatchWhenReady(
+				FSimpleDelegateGraphTask::FDelegate::CreateUObject(this, &UNavigationComponent::DeferredResumePath)
+				, TEXT("Deferred resume path")
+				, NULL
+				, ENamedThreads::GameThread
+				);
+			bPathGenerationSucceeded = true;
+		}
+		else
+		{
+			bPathGenerationSucceeded = ResumePath();
+		}
+	}
+	
+	if (bPathGenerationSucceeded == false)
+	{
+		ResetTransientData();
+
+		if (QueryFilter.IsValid() == false)
+		{
+			QueryFilter = GetStoredQueryFilter();
+		}
+
+		if (bDoAsyncPathfinding == true)
+		{
+			bPathGenerationSucceeded = AsyncGeneratePathTo(DestLocation, QueryFilter);
+		}
+		else
+		{
+			bPathGenerationSucceeded = GeneratePathTo(DestLocation, QueryFilter);
+		}
+	}
+
+	return bPathGenerationSucceeded;
+}
+
+void UNavigationComponent::SetPath(const FNavPathSharedPtr& ResultPath)
+{
+	check(ResultPath.Get() != NULL);
+
+	Path = ResultPath;
+	Path->SetObserver(PathObserverDelegate);
+	bIsPathPartial = Path->IsPartial();
+}
+
+void UNavigationComponent::ResetTransientData()
+{
+	GoalActor = NULL;
+	RepathDistanceSq = 0.f;
+	bIsPathPartial = false;
+	bIsPaused = false;
+	bUseSimplePath = false;
+	Path = NULL;
+	StoredQueryFilter = NULL;
+
+	// stop ticking - this component should tick only when following an actor
+	SetComponentTickEnabledAsync(false);
+}
+
+void UNavigationComponent::NotifyPathUpdate()
+{
+#if ENABLE_VISUAL_LOG
+	if (Path.IsValid())
+	{
+		UE_VLOG(GetOwner(), LogAINavigation, Log, TEXT("NotifyPathUpdate %s"), *Path->GetDescription());
+	}
+	else
+	{
+		UE_VLOG(GetOwner(), LogAINavigation, Log, TEXT("NotifyPathUpdate: path MISSING!"));
+	}
+#endif
+
+	if (MyPathObserver != NULL)
+	{
+		MyPathObserver->OnPathUpdated(this);
+	}
+
+	if (PathFollowComp && bNotifyPathFollowing && PathFollowComp->GetMoveGoal() == GoalActor)
+	{
+		PathFollowComp->UpdateMove(Path);
+	}
+
+	if (!Path.IsValid())
+	{
+		ResetTransientData();
+	}
+	else if (!Path->IsValid())
+	{
+		UE_VLOG(GetOwner(), LogAINavigation, Warning
+			, TEXT("NotifyPathUpdate fetched invalid NavPath!  NavPath has %d points, is%sready, and is%sup-to-date")
+			, Path->PathPoints.Num(), Path->IsReady() ? TEXT(" ") : TEXT(" NOT "), Path->IsUpToDate() ? TEXT(" ") : TEXT(" NOT ")	);
+
+		ResetTransientData();
+	}
+}
+
+void UNavigationComponent::DeferredRepathToGoal()
+{
+	// check if repath is allowed right now
+	const bool bShouldPostponeRepath = MyNavAgent && MyNavAgent->ShouldPostponePathUpdates();
+	if (bShouldPostponeRepath)
+	{
+		GetWorld()->GetTimerManager().SetTimer(this, &UNavigationComponent::DeferredRepathToGoal, 0.1f, false);
+		return;
+	}
+
+	GoalActor = RepathData.GoalActor.Get();
+	CurrentGoal = RepathData.GoalLocation;
+	bUseSimplePath = RepathData.bSimplePath;
+	FMemory::MemZero<FDeferredRepath>(RepathData);
+
+	const FVector GoalLocation = GoalActor ? GetCurrentMoveGoal(GoalActor, GetOwner()) : CurrentGoal;
+	if (RePathTo(GoalLocation))
+	{
+		OriginalGoalActorLocation = GoalLocation;
+	}
+	else
+	{
+		FAIMessage::Send(Cast<AController>(GetOwner()), FAIMessage(UBrainComponent::AIMessage_RepathFailed, this));
+	}
+
+	bIsWaitingForRepath = false;
+}
+
+void UNavigationComponent::SetPathFollowingUpdates(bool bEnabled)
+{
+	bNotifyPathFollowing = bEnabled;
+}
+
+void UNavigationComponent::SetRepathWhenInvalid(bool bEnabled)
+{
+	bRepathWhenInvalid = bEnabled;
+}
+
+void UNavigationComponent::SetComponentTickEnabledAsync(bool bEnabled)
+{
+	bShoudTick = bEnabled;
+	Super::SetComponentTickEnabledAsync(bEnabled);
+}
+
+void UNavigationComponent::OnNavAgentChanged()
+{
+	MyNavAgent = NULL;
+	PickNavData();
+}
+
+ANavigationData* UNavigationComponent::PickNavData() const
+{
+	// reset to default
+	MyNavData = NULL;
+
+	if (MyNavAgent == NULL)
+	{
+		MyNavAgent = InterfaceCast<INavAgentInterface>( GetOuter() );
+	}
+
+	if (MyPathObserver == NULL)
+	{
+		MyPathObserver = InterfaceCast<INavPathObserverInterface>( GetOuter() );
+	}
+
+	if (GetWorld() != NULL && GetWorld()->GetNavigationSystem() != NULL && MyNavAgent != NULL)
+	{
+		const FNavAgentProperties* MoveProps = MyNavAgent->GetNavAgentProperties();
+
+		// Make sure any previous registration for the event is cleared.
+		GetWorld()->GetNavigationSystem()->OnNavDataRegisteredEvent.RemoveDynamic(this, &UNavigationComponent::OnNavDataRegistered);
+		
+		if (MoveProps != NULL)
+		{
+			// cache it
+			MyNavData = GetWorld()->GetNavigationSystem()->GetNavDataForProps(*MoveProps);
+		}
+
+		if (MyNavData == NULL)
+		{	// Failed to pick it, so register to try again.
+			// note that if this happens a lot there's a bigger problem like AIs being spawned before navmesh is ready
+			GetWorld()->GetNavigationSystem()->OnNavDataRegisteredEvent.AddDynamic(this, &UNavigationComponent::OnNavDataRegistered);
+		}
+	}
+
+	CacheNavQueryExtent();
+
+	return MyNavData;
+}
+
+void UNavigationComponent::CacheNavQueryExtent() const
+{
+	NavigationQueryExtent = MyNavData ? MyNavData->GetDefaultQueryExtent() : FVector(DEFAULT_NAV_QUERY_EXTENT_HORIZONTAL, DEFAULT_NAV_QUERY_EXTENT_HORIZONTAL, DEFAULT_NAV_QUERY_EXTENT_VERTICAL);
+	if (MyNavAgent)
+	{
+		const FNavAgentProperties* AgentProperties = MyNavAgent->GetNavAgentProperties();
+		check(AgentProperties);
+		NavigationQueryExtent = FVector(FMath::Max(NavigationQueryExtent.X, AgentProperties->AgentRadius)
+			, FMath::Max(NavigationQueryExtent.Y, AgentProperties->AgentRadius)
+			, FMath::Max(NavigationQueryExtent.Z, AgentProperties->AgentHeight / 2));
+	}
+}
+
+void UNavigationComponent::OnNavDataRegistered(class ANavigationData* NavData)
+{
+	// re-pick regardless of whether we have MyNavData set already - the new one can be better!
+	PickNavData();
+}
+
+void UNavigationComponent::DrawDebugCurrentPath(UCanvas* Canvas, FColor PathColor, const uint32 NextPathPointIndex)
+{
+	if (MyNavData == NULL)
+	{
+		return;
+	}
+
+	const bool bPersistent = Canvas == NULL;
+
+	if (bPersistent)
+	{
+		FlushPersistentDebugLines( GetWorld() );
+	}
+
+	DrawDebugBox(GetWorld(), CurrentGoal, FVector(10.f), PathColor, bPersistent);
+	if (Path.IsValid())
+	{
+		MyNavData->DrawDebugPath(Path.Get(), PathColor, Canvas, bPersistent, NextPathPointIndex);
+	}
+}
+
+void UNavigationComponent::SetReceiveSmartLinkUpdates(bool bEnabled)
+{
+	bUpdateForSmartLinks = bEnabled;
+}
+
+void UNavigationComponent::OnSmartLinkBroadcast(class USmartNavLinkComponent* NearbyLink)
+{
+	// update only when agent is actually moving
+	if (NearbyLink == NULL || !Path.IsValid() ||
+		PathFollowComp == NULL || PathFollowComp->GetStatus() == EPathFollowingStatus::Idle)
+	{
+		return;
+	}
+
+	UNavigationSystem* NavSys = UNavigationSystem::GetCurrent(GetWorld());
+	if (NavSys == NULL || MyNavAgent == NULL || GetNavData() == NULL)
+	{
+		return;
+	}
+
+	const bool bHasLink = Path->ContainsSmartLink(NearbyLink);
+	const bool bIsEnabled = NearbyLink->IsEnabled();
+	if (bIsEnabled == bHasLink)
+	{
+		// enabled link, path already use it - or - disabled link and path it's not using it
+		return;
+	}
+
+	const EPathFindingMode::Type Mode = bDoHierarchicalPathfinding ? EPathFindingMode::Hierarchical : EPathFindingMode::Regular;
+	const FVector GoalLocation = Path->PathPoints.Last().Location;
+	FPathFindingQuery Query(GetOwner(), GetNavData(), MyNavAgent->GetNavAgentLocation(), GoalLocation, StoredQueryFilter);
+	Query.NavDataFlags = NavDataFlags;
+
+	FPathFindingResult Result = NavSys->FindPathSync(*MyNavAgent->GetNavAgentProperties(), Query, Mode);
+	if (Result.IsSuccessful() || Result.IsPartial())
+	{
+		const bool bNewHasLink = Result.Path->ContainsSmartLink(NearbyLink);
+
+		if ((bIsEnabled && !bHasLink && bNewHasLink) ||
+			(!bIsEnabled && bHasLink && !bNewHasLink))
+		{			
+			SetPath(Result.Path);
+			NotifyPathUpdate();
+		}
+	}
+}
+
+void UNavigationComponent::SwapCurrentMoveGoal(const AActor* NewGoalActor)
+{
+	if (NewGoalActor != NULL && GoalActor != NULL && GoalActor != NewGoalActor && Path.IsValid() && Path->IsValid() &&
+		PathFollowComp != NULL && PathFollowComp->GetStatus() != EPathFollowingStatus::Idle )
+	{
+		UE_VLOG(GetOwner(), LogAINavigation, Log, TEXT("GoalActor is going to be swapped for '%s'"), *NewGoalActor->GetHumanReadableName());
+		GoalActor = NewGoalActor;
+		PathFollowComp->SetDestinationActor( NewGoalActor );
+	}
+}
+
+#if ENABLE_VISUAL_LOG
+//----------------------------------------------------------------------//
+// debug
+//----------------------------------------------------------------------//
+void UNavigationComponent::DescribeSelfToVisLog(FVisLogEntry* Snapshot) const
+{
+	// path will be drawn by path following
+}
+
+#endif // ENABLE_VISUAL_LOG
diff --git a/Engine/Source/Runtime/AIModule/Private/Navigation/NavigationInterfaces.cpp b/Engine/Source/Runtime/AIModule/Private/Navigation/NavigationInterfaces.cpp
new file mode 100644
--- /dev/null
+++ b/Engine/Source/Runtime/AIModule/Private/Navigation/NavigationInterfaces.cpp
@@ -0,0 +1,9 @@
+// Copyright 1998-2014 Epic Games, Inc. All Rights Reserved.
+
+#include "AIModulePrivate.h"
+#include "Navigation/CrowdAgentInterface.h"
+
+UCrowdAgentInterface::UCrowdAgentInterface(const class FPostConstructInitializeProperties& PCIP)
+	: Super(PCIP)
+{
+}
diff --git a/Engine/Source/Runtime/AIModule/Private/Navigation/PathFollowingComponent.cpp b/Engine/Source/Runtime/AIModule/Private/Navigation/PathFollowingComponent.cpp
new file mode 100644
--- /dev/null
+++ b/Engine/Source/Runtime/AIModule/Private/Navigation/PathFollowingComponent.cpp
@@ -0,0 +1,1438 @@
+// Copyright 1998-2014 Epic Games, Inc. All Rights Reserved.
+
+#include "AIModulePrivate.h"
+#if WITH_RECAST
+#	include "Detour/DetourNavMeshQuery.h"
+#endif
+#include "Navigation/NavigationComponent.h"
+
+#include "AI/Navigation/SmartNavLinkComponent.h"
+
+#define USE_PHYSIC_FOR_VISIBILITY_TESTS 1 // Physic will be used for visibility tests if set or only raycasts on navmesh if not
+
+DEFINE_LOG_CATEGORY(LogPathFollowing);
+
+//----------------------------------------------------------------------//
+// Life cycle                                                        
+//----------------------------------------------------------------------//
+uint32 UPathFollowingComponent::NextRequestId = 0;
+const float UPathFollowingComponent::DefaultAcceptanceRadius = -1.f;
+UPathFollowingComponent::FRequestCompletedSignature UPathFollowingComponent::UnboundRequestDelegate;
+
+UPathFollowingComponent::UPathFollowingComponent(const class FPostConstructInitializeProperties& PCIP) : Super(PCIP)
+{
+	PrimaryComponentTick.bCanEverTick = true;
+	bAutoActivate = true;
+
+	MinAgentRadiusPct = 0.1f;
+	MinAgentHeightPct = 0.05f;
+	BlockDetectionDistance = 10.0f;
+	BlockDetectionInterval = 0.5f;
+	BlockDetectionSampleCount = 10;
+	LastSampleTime = 0.0f;
+	NextSampleIdx = 0;
+	bUseBlockDetection = true;
+	bLastMoveReachedGoal = false;
+	bUseVisibilityTestsSimplification = false;
+	bPendingPathStartUpdate = false;
+
+	MoveSegmentStartIndex = 0;
+	MoveSegmentEndIndex = 1;
+	MoveSegmentStartRef = INVALID_NAVNODEREF;
+	MoveSegmentEndRef = INVALID_NAVNODEREF;
+
+	bStopOnOverlap = true;
+	Status = EPathFollowingStatus::Idle;
+}
+
+void LogPathHelper(AActor* LogOwner, FNavPathSharedPtr Path, const AActor* GoalActor, const int32 CurrentPathPointGoal)
+{
+#if ENABLE_VISUAL_LOG
+	FVisualLog* Vlog = FVisualLog::Get();
+	if (Vlog && Vlog->IsRecording() && 
+		Path.IsValid() && Path->IsValid() && Path->PathPoints.Num())
+	{
+		FVisLogEntry* Entry = Vlog->GetEntryToWrite(LogOwner);
+		Path->DescribeSelfToVisLog(Entry);
+
+		const FVector PathEnd = Path->PathPoints.Last().Location;
+		if (GoalActor)
+		{
+			const FVector GoalLoc = GoalActor->GetActorLocation();
+			if (FVector::DistSquared(GoalLoc, PathEnd) > 1.0f)
+			{
+				UE_VLOG_LOCATION(LogOwner, LogPathFollowing, Verbose, GoalLoc, 30, FColor::Green, TEXT("GoalActor"));
+				UE_VLOG_SEGMENT(LogOwner, LogPathFollowing, Verbose, GoalLoc, PathEnd, FColor::Green, TEXT_EMPTY);
+			}
+		}
+
+		UE_VLOG_BOX(LogOwner, LogPathFollowing, Verbose, FBox(PathEnd - FVector(30.0f), PathEnd + FVector(30.0f)), FColor::Green, TEXT("PathEnd"));
+	}
+#endif // ENABLE_VISUAL_LOG
+}
+
+void LogBlockHelper(AActor* LogOwner, class UNavMovementComponent* MoveComp, float RadiusPct, float HeightPct, const FVector& SegmentStart, const FVector& SegmentEnd)
+{
+#if ENABLE_VISUAL_LOG
+	if (MoveComp && LogOwner)
+	{
+		const FVector AgentLocation = MoveComp->GetActorFeetLocation();
+		const FVector ToTarget = (SegmentEnd - AgentLocation);
+		const float SegmentDot = FVector::DotProduct(ToTarget.SafeNormal(), (SegmentEnd - SegmentStart).SafeNormal());
+		UE_VLOG(LogOwner, LogPathFollowing, Verbose, TEXT("[agent to segment end] dot [segment dir]: %f"), SegmentDot);
+		
+		float AgentRadius = 0.0f;
+		float AgentHalfHeight = 0.0f;
+		AActor* MovingAgent = MoveComp->GetOwner();
+		MovingAgent->GetSimpleCollisionCylinder(AgentRadius, AgentHalfHeight);
+
+		const float Dist2D = ToTarget.Size2D();
+		UE_VLOG(LogOwner, LogPathFollowing, Verbose, TEXT("dist 2d: %f (agent radius: %f [%f])"), Dist2D, AgentRadius, AgentRadius * (1 + RadiusPct));
+
+		const float ZDiff = FMath::Abs(ToTarget.Z);
+		UE_VLOG(LogOwner, LogPathFollowing, Verbose, TEXT("Z diff: %f (agent halfZ: %f [%f])"), ZDiff, AgentHalfHeight, AgentHalfHeight * (1 + HeightPct));
+	}
+#endif // ENABLE_VISUAL_LOG
+}
+
+FString GetPathDescHelper(FNavPathSharedPtr Path)
+{
+	return !Path.IsValid() ? TEXT("missing") :
+		!Path->IsValid() ? TEXT("invalid") :
+		FString::Printf(TEXT("%s:%d"), Path->IsPartial() ? TEXT("partial") : TEXT("complete"), Path->PathPoints.Num());
+}
+
+FAIRequestID UPathFollowingComponent::RequestMove(FNavPathSharedPtr InPath, FRequestCompletedSignature OnComplete,
+	const AActor* InDestinationActor, float InAcceptanceRadius, bool bInStopOnOverlap, FCustomMoveSharedPtr InGameData)
+{
+	UE_VLOG(GetOwner(), LogPathFollowing, Log, TEXT("RequestMove: Path(%s), AcceptRadius(%.1f%s), DestinationActor(%s), GameData(%s)"),
+		*GetPathDescHelper(InPath),
+		InAcceptanceRadius, bInStopOnOverlap ? TEXT(" + agent") : TEXT(""),
+		*GetNameSafe(InDestinationActor),
+		!InGameData.IsValid() ? TEXT("missing") : TEXT("valid"));
+
+	LogPathHelper(GetOwner(), InPath, InDestinationActor, GetNextPathIndex());
+
+	if (InAcceptanceRadius == UPathFollowingComponent::DefaultAcceptanceRadius)
+	{
+		InAcceptanceRadius = MyDefaultAcceptanceRadius;
+	}
+
+	if (ResourceLock.IsLocked())
+	{
+		UE_VLOG(GetOwner(), LogPathFollowing, Log, TEXT("Rejecting move request due to resource lock by %s"), *ResourceLock.GetLockSourceName());
+		return FAIRequestID::InvalidRequest;
+	}
+
+	if (!InPath.IsValid() || InAcceptanceRadius < 0.0f)
+	{
+		UE_VLOG(GetOwner(), LogPathFollowing, Log, TEXT("RequestMove: invalid request"));
+		return FAIRequestID::InvalidRequest;
+	}
+
+	// try to grab movement component
+	if (!UpdateMovementComponent())
+	{
+		UE_VLOG(GetOwner(), LogPathFollowing, Warning, TEXT("RequestMove: missing movement component"));
+		return FAIRequestID::InvalidRequest;
+	}
+
+	// update ID first, so any observer notified by AbortMove() could detect new request
+	const uint32 PrevMoveId = CurrentRequestId;
+	
+	// abort previous movement
+	if (Status == EPathFollowingStatus::Paused && Path.IsValid() && InPath.Get() == Path.Get() && DestinationActor == InDestinationActor)
+	{
+		ResumeMove();
+	}
+	else
+	{
+		if (Status == EPathFollowingStatus::Moving)
+		{
+			const bool bResetVelocity = false;
+			const bool bFinishAsSkipped = true;
+			AbortMove(TEXT("new request"), PrevMoveId, bResetVelocity, bFinishAsSkipped);
+		}
+		
+		Reset();
+
+		StoreRequestId();
+
+		// store new data
+		Path = InPath;
+		OnPathUpdated();
+
+		AcceptanceRadius = InAcceptanceRadius;
+		GameData = InGameData;
+		OnRequestFinished = OnComplete;	
+		bStopOnOverlap = bInStopOnOverlap;
+		SetDestinationActor(InDestinationActor);
+
+	#if ENABLE_VISUAL_LOG
+		const FVector CurrentLocation = MovementComp ? MovementComp->GetActorFeetLocation() : FVector::ZeroVector;
+		const FVector DestLocation = InPath->GetDestinationLocation();
+		const FVector ToDest = DestLocation - CurrentLocation;
+		UE_VLOG(GetOwner(), LogPathFollowing, Log, TEXT("RequestMove: accepted, ID(%u) dist2D(%.0f) distZ(%.0f)"),
+			CurrentRequestId.GetID(), ToDest.Size2D(), FMath::Abs(ToDest.Z));
+	#endif // ENABLE_VISUAL_LOG
+
+		// with async pathfinding paths can be incomplete, movement will start after receiving UpdateMove 
+		if (Path->IsValid())
+		{
+			Status = EPathFollowingStatus::Moving;
+
+			// determine with path segment should be followed
+			const uint32 CurrentSegment = DetermineStartingPathPoint(InPath.Get());
+			bPendingPathStartUpdate = false;
+			SetMoveSegment(CurrentSegment);
+		}
+		else
+		{
+			Status = EPathFollowingStatus::Waiting;
+		}
+	}
+
+	return CurrentRequestId;
+}
+
+bool UPathFollowingComponent::UpdateMove(FNavPathSharedPtr InPath, FAIRequestID RequestID)
+{
+	UE_VLOG(GetOwner(), LogPathFollowing, Log, TEXT("UpdateMove: Path(%s) Status(%s) RequestID(%u)"),
+		*GetPathDescHelper(InPath),
+		*GetStatusDesc(), RequestID);
+
+	LogPathHelper(GetOwner(), InPath, DestinationActor.Get(), GetNextPathIndex());
+
+	if (!InPath.IsValid() || !InPath->IsValid() || Status == EPathFollowingStatus::Idle 
+		|| RequestID.IsEquivalent(GetCurrentRequestId()) == false)
+	{
+		UE_VLOG(GetOwner(), LogPathFollowing, Log, TEXT("UpdateMove: invalid request"));
+		return false;
+	}
+
+	Path = InPath;
+	OnPathUpdated();
+
+	if (Status == EPathFollowingStatus::Waiting || Status == EPathFollowingStatus::Moving)
+	{
+		Status = EPathFollowingStatus::Moving;
+
+		const int32 CurrentSegment = DetermineStartingPathPoint(InPath.Get());
+		bPendingPathStartUpdate = false;
+		SetMoveSegment(CurrentSegment);
+	}
+	else if (Status == EPathFollowingStatus::Paused)
+	{
+		bPendingPathStartUpdate = true;
+	}
+
+	return true;
+}
+
+void UPathFollowingComponent::AbortMove(const FString& Reason, FAIRequestID RequestID, bool bResetVelocity, bool bSilent)
+{
+	UE_VLOG(GetOwner(), LogPathFollowing, Log, TEXT("AbortMove: Reason(%s) RequestID(%u)"), *Reason, RequestID);
+
+	if ((Status != EPathFollowingStatus::Idle) && RequestID.IsEquivalent(GetCurrentRequestId()))
+	{
+		const EPathFollowingResult::Type FinishResult = bSilent ? EPathFollowingResult::Skipped : EPathFollowingResult::Aborted;
+
+		if (CurrentSmartLink)
+		{
+			CurrentSmartLink->NotifyMoveAborted(this);
+		}
+
+		// save data required for observers before reseting temporary variables
+		const uint32 AbortedMoveId = RequestID ? RequestID : CurrentRequestId;
+		FRequestCompletedSignature SavedReqFinished = OnRequestFinished;
+
+		Reset();
+		bLastMoveReachedGoal = false;
+
+		if (bResetVelocity && MovementComp && MovementComp->CanStopPathFollowing())
+		{
+			MovementComp->StopMovementImmediately();
+		}
+
+		if (NavComp && FinishResult != EPathFollowingResult::Skipped)
+		{
+			NavComp->ResetTransientData();
+		}
+
+		// notify observers after state was reset (they can request another move)
+		SavedReqFinished.ExecuteIfBound(FinishResult);
+		OnMoveFinished.Broadcast(AbortedMoveId, FinishResult);
+		FAIMessage::Send(Cast<AController>(GetOwner()), FAIMessage(UBrainComponent::AIMessage_MoveFinished, this, AbortedMoveId, FAIMessage::Failure));
+	}
+	else
+	{
+		UE_VLOG(GetOwner(), LogPathFollowing, Log, TEXT("AbortMove FAILED due to RequestID(%u) != CurrentRequestId()%d")
+			, RequestID, CurrentRequestId);
+	}
+}
+
+void UPathFollowingComponent::PauseMove(FAIRequestID RequestID, bool bResetVelocity)
+{
+	UE_VLOG(GetOwner(), LogPathFollowing, Log, TEXT("PauseMove: RequestID(%u)"), RequestID);
+	if (Status == EPathFollowingStatus::Paused)
+	{
+		return;
+	}
+
+	if (RequestID.IsEquivalent(GetCurrentRequestId()))
+	{
+		if (bResetVelocity && MovementComp && MovementComp->CanStopPathFollowing())
+		{
+			MovementComp->StopMovementImmediately();
+		}
+
+		LocationWhenPaused = MovementComp ? MovementComp->GetActorFeetLocation() : FVector::ZeroVector;
+		Status = EPathFollowingStatus::Paused;
+
+		UpdateMoveFocus();
+	}
+}
+
+void UPathFollowingComponent::ResumeMove(FAIRequestID RequestID)
+{
+	if (RequestID.IsEquivalent(CurrentRequestId) && RequestID.IsValid())
+	{
+		UE_VLOG(GetOwner(), LogPathFollowing, Log, TEXT("ResumeMove: RequestID(%u)"), RequestID);
+
+		const bool bMovedDuringPause = ShouldCheckPathOnResume();
+		const bool bIsOnPath = IsOnPath();
+		if (bIsOnPath)
+		{
+			Status = EPathFollowingStatus::Moving;
+			if (bMovedDuringPause || bPendingPathStartUpdate)
+			{
+				const int32 CurrentSegment = DetermineStartingPathPoint(Path.Get());
+				bPendingPathStartUpdate = false;
+				SetMoveSegment(CurrentSegment);
+			}
+			else
+			{
+				UpdateMoveFocus();
+			}
+		}
+		else if (Path.IsValid() && Path->IsValid() && Path->GetOwner() == NULL)
+		{
+			// this means it's a scripted path, just resume
+			Status = EPathFollowingStatus::Moving;
+			UpdateMoveFocus();
+		}
+		else
+		{
+			OnPathFinished(EPathFollowingResult::OffPath);
+		}
+	}
+	else
+	{
+		UE_VLOG(GetOwner(), LogPathFollowing, Log, TEXT("ResumeMove: RequestID(%u) is neither \'AnyRequest\' not CurrentRequestId(%u). Ignoring.")
+			, RequestID, CurrentRequestId);
+	}
+}
+
+bool UPathFollowingComponent::ShouldCheckPathOnResume() const
+{
+	bool bCheckPath = true;
+	if (MovementComp != NULL)
+	{
+		float AgentRadius = 0.0f, AgentHalfHeight = 0.0f;
+		MovementComp->GetOwner()->GetSimpleCollisionCylinder(AgentRadius, AgentHalfHeight);
+
+		const FVector CurrentLocation = MovementComp->GetActorFeetLocation();
+		const float DeltaMove2DSq = (CurrentLocation - LocationWhenPaused).SizeSquared2D();
+		const float DeltaZ = FMath::Abs(CurrentLocation.Z - LocationWhenPaused.Z);
+		if (DeltaMove2DSq < FMath::Square(AgentRadius) && DeltaZ < (AgentHalfHeight * 0.5f))
+		{
+			bCheckPath = false;
+		}
+	}
+
+	return bCheckPath;
+}
+
+void UPathFollowingComponent::OnPathFinished(EPathFollowingResult::Type Result)
+{
+	UE_VLOG(GetOwner(), LogPathFollowing, Log, TEXT("OnPathFinished: %s"), *GetResultDesc(Result));
+	
+	// save move status
+	bLastMoveReachedGoal = (Result == EPathFollowingResult::Success) && !HasPartialPath();
+
+	// save data required for observers before reseting temporary variables
+	const FAIRequestID FinishedMoveId = CurrentRequestId;
+	FRequestCompletedSignature SavedReqFinished = OnRequestFinished;
+
+	Reset();
+	UpdateMoveFocus();
+
+	if (MovementComp && MovementComp->CanStopPathFollowing())	
+	{
+		MovementComp->StopMovementImmediately();
+	}
+
+	if (NavComp && Result != EPathFollowingResult::Skipped)
+	{
+		NavComp->ResetTransientData();
+	}
+
+	// notify observers after state was reset (they can request another move)
+	SavedReqFinished.ExecuteIfBound(Result);
+	OnMoveFinished.Broadcast(FinishedMoveId, Result);
+
+	FAIMessage::Send(Cast<AController>(GetOwner()), FAIMessage(UBrainComponent::AIMessage_MoveFinished, this, FinishedMoveId, (Result == EPathFollowingResult::Success)));
+}
+
+void UPathFollowingComponent::OnSegmentFinished()
+{
+	UE_VLOG(GetOwner(), LogPathFollowing, Verbose, TEXT("OnSegmentFinished"));
+}
+
+void UPathFollowingComponent::OnPathUpdated()
+{
+}
+
+void UPathFollowingComponent::Initialize()
+{
+	UpdateCachedComponents();
+}
+
+void UPathFollowingComponent::Cleanup()
+{
+	// empty in base class
+}
+
+void UPathFollowingComponent::UpdateCachedComponents()
+{
+	AActor* MyOwner = GetOwner();
+	if (MyOwner)
+	{
+		NavComp = MyOwner->FindComponentByClass<UNavigationComponent>();
+	}
+
+	UpdateMovementComponent(/*bForce=*/true);
+}
+
+void UPathFollowingComponent::TickComponent(float DeltaTime, enum ELevelTick TickType, FActorComponentTickFunction *ThisTickFunction)
+{
+	Super::TickComponent(DeltaTime, TickType, ThisTickFunction);
+
+	if (Status == EPathFollowingStatus::Moving)
+	{
+		// check finish conditions, update current segment if needed
+		// may result in changing current status!
+		UpdatePathSegment();
+	}
+
+	if (Status == EPathFollowingStatus::Moving)
+	{
+		// follow current path segment
+		FollowPathSegment(DeltaTime);
+	}
+};
+
+void UPathFollowingComponent::SetMovementComponent(UNavMovementComponent* MoveComp)
+{
+	MovementComp = MoveComp;
+	MyNavData = NULL;
+
+	if (MoveComp != NULL)
+	{
+		const FNavAgentProperties* NavAgentProps = MoveComp->GetNavAgentProperties();
+		check(NavAgentProps);
+		MyDefaultAcceptanceRadius = NavAgentProps->AgentRadius;
+		MoveComp->PathFollowingComp = this;
+
+		if (GetWorld() && GetWorld()->GetNavigationSystem())
+		{	
+			MyNavData = GetWorld()->GetNavigationSystem()->GetNavDataForProps(*NavAgentProps);
+		}
+	}
+}
+
+FVector UPathFollowingComponent::GetMoveFocus(bool bAllowStrafe) const
+{
+	FVector MoveFocus = FVector::ZeroVector;
+	if (bAllowStrafe && DestinationActor.IsValid())
+	{
+		MoveFocus = DestinationActor->GetActorLocation();
+	}
+	else
+	{
+		MoveFocus = *CurrentDestination + MoveSegmentDirection * 20.0f;
+	}
+
+	return MoveFocus;
+}
+
+void UPathFollowingComponent::Reset()
+{
+	MoveSegmentStartIndex = 0;
+	MoveSegmentStartRef = INVALID_NAVNODEREF;
+	MoveSegmentEndRef = INVALID_NAVNODEREF;
+
+	LocationSamples.Reset();
+	LastSampleTime = 0.0f;
+	NextSampleIdx = 0;
+
+	Path.Reset();
+	GameData.Reset();
+	DestinationActor.Reset();
+	OnRequestFinished.Unbind();
+	AcceptanceRadius = MyDefaultAcceptanceRadius;
+	bStopOnOverlap = true;
+	bCollidedWithGoal = false;
+
+	CurrentRequestId = FAIRequestID::InvalidRequest;
+
+	CurrentDestination.Clear();
+	CurrentSmartLink = NULL;
+
+	Status = EPathFollowingStatus::Idle;
+}
+
+int32 UPathFollowingComponent::DetermineStartingPathPoint(const FNavigationPath* ConsideredPath) const
+{
+	int32 PickedPathPoint = INDEX_NONE;
+
+	if (ConsideredPath && ConsideredPath->IsValid())
+	{
+		// if we already have some info on where we were on previous path
+		// we can find out if there's a segment on new path we're currently at
+		if (MoveSegmentStartRef != INVALID_NAVNODEREF &&
+			MoveSegmentEndRef != INVALID_NAVNODEREF &&
+			ConsideredPath->GetOwner() != NULL)
+		{
+			// iterate every new path node and see if segment match
+			for (int32 PathPoint = 0; PathPoint < ConsideredPath->PathPoints.Num() - 1; ++PathPoint)
+			{
+				if (ConsideredPath->PathPoints[PathPoint].NodeRef == MoveSegmentStartRef && 
+					ConsideredPath->PathPoints[PathPoint+1].NodeRef == MoveSegmentEndRef)
+				{
+					PickedPathPoint = PathPoint;
+					break;
+				}
+			}
+		}
+
+		if (MovementComp && PickedPathPoint == INDEX_NONE)
+		{
+			if (ConsideredPath->PathPoints.Num() > 2)
+			{
+				// check if is closer to first or second path point (don't assume AI's standing)
+				const FVector CurrentLocation = MovementComp->GetActorFeetLocation();
+				const float SqDistToFirstPoint = FVector::DistSquared(CurrentLocation, ConsideredPath->PathPoints[0].Location);
+				const float SqDistToSecondPoint = FVector::DistSquared(CurrentLocation, ConsideredPath->PathPoints[1].Location);
+				PickedPathPoint = (SqDistToFirstPoint < SqDistToSecondPoint) ? 0 : 1;
+			}
+			else
+			{
+				// If there are only two point we probably should start from the beginning
+				PickedPathPoint = 0;
+			}
+		}
+	}
+
+	return PickedPathPoint;
+}
+
+void UPathFollowingComponent::SetDestinationActor(const AActor* InDestinationActor)
+{
+	DestinationActor = InDestinationActor;
+	DestinationAgent = InterfaceCast<const INavAgentInterface>(InDestinationActor);
+	MoveOffset = DestinationAgent ? DestinationAgent->GetMoveGoalOffset(GetOwner()) : FVector::ZeroVector;
+}
+
+void UPathFollowingComponent::SetMoveSegment(int32 SegmentStartIndex)
+{
+	const bool bResumedFromSmartLink = (CurrentSmartLink != NULL) && (MoveSegmentStartIndex == SegmentStartIndex);
+	SetCurrentSmartLink(NULL, FVector::ZeroVector);
+
+	int32 EndSegmentIndex = SegmentStartIndex + 1;
+	if (Path.IsValid() && Path->PathPoints.IsValidIndex(SegmentStartIndex) && Path->PathPoints.IsValidIndex(EndSegmentIndex))
+	{
+		EndSegmentIndex = DetermineCurrentTargetPathPoint(SegmentStartIndex, bResumedFromSmartLink);
+
+		MoveSegmentStartIndex = SegmentStartIndex;
+		MoveSegmentEndIndex = EndSegmentIndex;
+		const FNavPathPoint& PathPt0 = Path->PathPoints[MoveSegmentStartIndex];
+		const FNavPathPoint& PathPt1 = Path->PathPoints[MoveSegmentEndIndex];
+
+		MoveSegmentStartRef = PathPt0.NodeRef;
+		MoveSegmentEndRef = PathPt1.NodeRef;
+		
+		const FVector SegmentStart = PathPt0.Location;
+		const FVector SegmentEnd = PathPt1.Location;
+
+		CurrentDestination.Set(Path->Base.Get(), SegmentEnd);
+		CurrentAcceptanceRadius = (Path->PathPoints.Num() == (EndSegmentIndex + 1)) ? AcceptanceRadius : 0.0f;
+
+		MoveSegmentDirection = (SegmentEnd - SegmentStart).SafeNormal();
+		
+		// make sure we have a non-zero direction if still following a valid path
+		if (MoveSegmentDirection.IsZero() && int32(MoveSegmentEndIndex + 1) < Path->PathPoints.Num())
+		{
+			MoveSegmentDirection = (Path->PathPoints[MoveSegmentEndIndex + 1].Location - SegmentStart).SafeNormal();
+		}
+
+		// handle moving through smart links
+		if (!bResumedFromSmartLink && PathPt0.SmartLink.IsValid())
+		{
+			SetCurrentSmartLink(PathPt0.SmartLink.Get(), SegmentEnd);
+		}
+
+		// update move focus in owning AI
+		UpdateMoveFocus();
+	}
+}
+
+int32 UPathFollowingComponent::DetermineCurrentTargetPathPoint(int32 StartIndex, bool bResumedFromSmartLink)
+{
+	if (!bUseVisibilityTestsSimplification ||
+		bResumedFromSmartLink ||
+		Path->PathPoints[StartIndex].SmartLink.IsValid() ||
+		Path->PathPoints[StartIndex + 1].SmartLink.IsValid())
+	{
+		return StartIndex + 1;
+	}
+
+
+#if WITH_RECAST
+	FNavMeshNodeFlags Flags0(Path->PathPoints[StartIndex].Flags);
+	FNavMeshNodeFlags Flags1(Path->PathPoints[StartIndex + 1].Flags);
+
+	const bool bOffMesh0 = (Flags0.PathFlags & RECAST_STRAIGHTPATH_OFFMESH_CONNECTION) != 0;
+	const bool bOffMesh1 = (Flags1.PathFlags & RECAST_STRAIGHTPATH_OFFMESH_CONNECTION) != 0;
+
+	if ((Flags0.Area != Flags1.Area) || (bOffMesh0 && bOffMesh1))
+	{
+		return StartIndex + 1;
+	}
+#endif
+
+	return OptimizeSegmentVisibility(StartIndex);
+}
+
+int32 UPathFollowingComponent::OptimizeSegmentVisibility(int32 StartIndex)
+{
+	//SCOPE_CYCLE_COUNTER(STAT_Navigation_PathVisibilityOptimisation);
+
+	const AAIController* MyAI = Cast<AAIController>(GetOwner());
+	const ANavigationData* NavData = MyAI && MyAI->NavComponent ? MyAI->NavComponent->GetNavData() : NULL;
+	if (Path.IsValid())
+	{
+		Path->ShortcutNodeRefs.Reset();
+	}
+
+	if (NavData == NULL)
+	{
+		return StartIndex + 1;
+	}
+
+	const APawn* MyPawn = MyAI->GetPawn();
+	const float PawnHalfHeight = (MyAI->GetCharacter() && MyAI->GetCharacter()->CapsuleComponent) ? MyAI->GetCharacter()->CapsuleComponent->GetScaledCapsuleHalfHeight() : 0.0f;
+	const FVector PawnLocation = MyPawn->GetActorLocation();
+
+#if USE_PHYSIC_FOR_VISIBILITY_TESTS
+	static const FName NAME_NavAreaTrace = TEXT("NavmeshVisibilityTrace");
+	float Radius, HalfHeight;
+	MyPawn->GetSimpleCollisionCylinder(Radius, HalfHeight);
+	const FCollisionShape CollisionShape = FCollisionShape::MakeCapsule(Radius, HalfHeight);
+	FCollisionQueryParams CapsuleParams(NAME_NavAreaTrace, false, MyAI->GetCharacter());
+	if (MyPawn->GetRootPrimitiveComponent())
+	{
+		CapsuleParams.AddIgnoredActors(MyPawn->GetRootPrimitiveComponent()->MoveIgnoreActors);
+	}
+
+#endif
+	TSharedPtr<FNavigationQueryFilter> QueryFilter = (NavComp && NavComp->GetStoredQueryFilter().IsValid()) ? NavComp->GetStoredQueryFilter()->GetCopy() : NavData->GetDefaultQueryFilter()->GetCopy();
+#if WITH_RECAST
+	const uint8 StartArea = FNavMeshNodeFlags(Path->PathPoints[StartIndex].Flags).Area;
+	TArray<float> CostArray;
+	CostArray.Init(DT_UNWALKABLE_POLY_COST, DT_MAX_AREAS);
+	CostArray[StartArea] = 0;
+	QueryFilter->SetAllAreaCosts(CostArray);
+#endif
+
+	const int32 MaxPoints = Path->PathPoints.Num();
+	int32 Index = StartIndex + 2;
+
+	for (; Index <= MaxPoints; ++Index)
+	{
+		if (!Path->PathPoints.IsValidIndex(Index))
+		{
+			break;
+		}
+
+		const FNavPathPoint& PathPt1 = Path->PathPoints[Index];
+		const FVector AdjustedDestination = PathPt1.Location + FVector(0.0f, 0.0f, PawnHalfHeight);
+
+		FVector HitLocation;
+#if WITH_RECAST
+		ARecastNavMesh::FRaycastResult RaycastResult;
+		const bool RaycastHitResult = ARecastNavMesh::NavMeshRaycast(NavData, PawnLocation, AdjustedDestination, HitLocation, QueryFilter, MyAI, RaycastResult);
+		if (Path.IsValid())
+		{
+			Path->ShortcutNodeRefs.Reserve(RaycastResult.CorridorPolysCount);
+			Path->ShortcutNodeRefs.SetNumUninitialized(RaycastResult.CorridorPolysCount);
+		}
+		FPlatformMemory::Memcmp(Path->ShortcutNodeRefs.GetTypedData(), RaycastResult.CorridorPolys, RaycastResult.CorridorPolysCount * sizeof(NavNodeRef));
+#else
+		const bool RaycastHitResult = NavData->Raycast(PawnLocation, AdjustedDestination, HitLocation, QueryFilter);
+#endif
+		if (RaycastHitResult)
+		{
+			break;
+		}
+#if USE_PHYSIC_FOR_VISIBILITY_TESTS
+		else
+		{
+			FHitResult Result;
+			if (GetWorld()->SweepSingle(Result, PawnLocation, AdjustedDestination, FQuat::Identity, ECC_WorldStatic, CollisionShape, CapsuleParams))
+			{
+				break;
+			}
+		}
+#endif
+
+#if WITH_RECAST
+		/** don't move to next point if we are changing area, we are at off mesh connection or it's a smart link*/
+		if (FNavMeshNodeFlags(Path->PathPoints[StartIndex].Flags).Area != FNavMeshNodeFlags(PathPt1.Flags).Area ||
+			FNavMeshNodeFlags(PathPt1.Flags).PathFlags & RECAST_STRAIGHTPATH_OFFMESH_CONNECTION || 
+			PathPt1.SmartLink.IsValid())
+		{
+			return Index;
+		}
+#endif
+	}
+
+	return Index-1;
+}
+
+void UPathFollowingComponent::UpdatePathSegment()
+{
+	if (!Path.IsValid() || MovementComp == NULL)
+	{
+		AbortMove(TEXT("no path"));
+		return;
+	}
+
+	if (!Path->IsValid())
+	{
+		if (NavComp == NULL || !NavComp->IsWaitingForRepath())
+		{
+			AbortMove(TEXT("no path"));
+		}
+		return;
+	}
+
+	// if agent has control over its movement, check finish conditions
+	const bool bCanReachTarget = MovementComp->CanStopPathFollowing();
+	if (bCanReachTarget && Status == EPathFollowingStatus::Moving)
+	{
+		const FVector CurrentLocation = MovementComp->GetActorFeetLocation();
+		const int32 LastSegmentEndIndex = Path->PathPoints.Num() - 1;
+		const bool bFollowingLastSegment = (MoveSegmentEndIndex >= LastSegmentEndIndex);
+
+		if (bCollidedWithGoal)
+		{
+			// check if collided with goal actor
+			OnSegmentFinished();
+			OnPathFinished(EPathFollowingResult::Success);
+		}
+		else if (bFollowingLastSegment)
+		{
+			// use goal actor for end of last path segment
+			// UNLESS it's partial path (can't reach goal)
+			if (DestinationActor.IsValid() && !Path->IsPartial())
+			{
+				const FVector AgentLocation = DestinationAgent ? DestinationAgent->GetNavAgentLocation() : DestinationActor->GetActorLocation();
+				const FVector GoalLocation = FRotationTranslationMatrix(DestinationActor->GetActorRotation(), AgentLocation).TransformPosition(MoveOffset);
+				CurrentDestination.Set(NULL, GoalLocation);
+			}
+
+			// include goal actor in reach test only for last segment 
+			if (HasReachedDestination(CurrentLocation))
+			{
+				OnSegmentFinished();
+				OnPathFinished(EPathFollowingResult::Success);
+			}
+			else
+			{
+				UpdateMoveFocus();
+			}
+		}
+		else
+		{
+			// check if current move segment is finished
+			if (HasReachedCurrentTarget(CurrentLocation))
+			{
+				OnSegmentFinished();
+				SetNextMoveSegment();
+			}
+		}
+	}
+
+	// gather location samples to detect if moving agent is blocked
+	if (bCanReachTarget && Status == EPathFollowingStatus::Moving)
+	{
+		const bool bHasNewSample = UpdateBlockDetection();
+		if (bHasNewSample && IsBlocked())
+		{
+			if (Path->PathPoints.IsValidIndex(MoveSegmentEndIndex) && Path->PathPoints.IsValidIndex(MoveSegmentStartIndex))
+			{
+				LogBlockHelper(GetOwner(), MovementComp, MinAgentRadiusPct, MinAgentHeightPct,
+					Path->PathPoints[MoveSegmentStartIndex].Location, Path->PathPoints[MoveSegmentEndIndex].Location);
+			}
+			else
+			{
+				if ((GetOwner() != NULL) && (MovementComp != NULL))
+				{
+					UE_VLOG(GetOwner(), LogPathFollowing, Verbose, TEXT("Path blocked, but move segment indices are not valid: start %d, end %d of %d"), MoveSegmentStartIndex, MoveSegmentEndIndex, Path->PathPoints.Num());
+				}
+			}
+			OnPathFinished(EPathFollowingResult::Blocked);
+		}
+	}
+}
+
+void UPathFollowingComponent::FollowPathSegment(float DeltaTime)
+{
+	if (MovementComp == NULL || !Path.IsValid())
+	{
+		return;
+	}
+
+	const FVector CurrentLocation = MovementComp->GetActorFeetLocation();
+
+	// check if already at acceptance radius
+	if (HasReachedDestination(CurrentLocation))
+	{
+		OnSegmentFinished();
+		OnPathFinished(EPathFollowingResult::Success);
+	}
+	else
+	{
+		const FVector CurrentTarget = GetCurrentTargetLocation();
+		FVector MoveVelocity = (CurrentTarget - CurrentLocation) / DeltaTime;
+
+		const int32 LastSegmentStartIndex = Path->PathPoints.Num() - 2;
+		const bool bNotFollowingLastSegment = (MoveSegmentStartIndex < LastSegmentStartIndex);
+
+		PostProcessMove.ExecuteIfBound(this, MoveVelocity);
+		MovementComp->RequestDirectMove(MoveVelocity, bNotFollowingLastSegment);
+	}
+}
+
+bool UPathFollowingComponent::HasReached(const FVector& TestPoint, float InAcceptanceRadius, bool bExactSpot) const
+{
+	// simple test for stationary agent, used as early finish condition
+	const FVector CurrentLocation = MovementComp ? MovementComp->GetActorFeetLocation() : FVector::ZeroVector;
+	const float GoalRadius = 0.0f;
+	const float GoalHalfHeight = 0.0f;
+	if (InAcceptanceRadius == UPathFollowingComponent::DefaultAcceptanceRadius)
+	{
+		InAcceptanceRadius = MyDefaultAcceptanceRadius;
+	}
+	return HasReachedInternal(TestPoint, GoalRadius, GoalHalfHeight, CurrentLocation, InAcceptanceRadius, !bExactSpot);
+}
+
+bool UPathFollowingComponent::HasReached(const AActor* TestGoal, float InAcceptanceRadius, bool bExactSpot) const
+{
+	if (TestGoal == NULL)
+	{
+		// we might just as well say "we're there" if "there" is NULL
+		return true;
+	}
+
+	// simple test for stationary agent, used as early finish condition
+	float GoalRadius = 0.0f;
+	float GoalHalfHeight = 0.0f;
+	FVector GoalOffset = FVector::ZeroVector;
+	FVector TestPoint = TestGoal->GetActorLocation();
+	if (InAcceptanceRadius == UPathFollowingComponent::DefaultAcceptanceRadius)
+	{
+		InAcceptanceRadius = MyDefaultAcceptanceRadius;
+	}
+
+	const INavAgentInterface* NavAgent = InterfaceCast<const INavAgentInterface>(TestGoal);
+	if (NavAgent)
+	{
+		const FVector GoalMoveOffset = NavAgent->GetMoveGoalOffset(GetOwner());
+		NavAgent->GetMoveGoalReachTest(GetOwner(), GoalMoveOffset, GoalOffset, GoalRadius, GoalHalfHeight);
+		TestPoint = FRotationTranslationMatrix(TestGoal->GetActorRotation(), NavAgent->GetNavAgentLocation()).TransformPosition(GoalOffset);
+	}
+
+	const FVector CurrentLocation = MovementComp ? MovementComp->GetActorFeetLocation() : FVector::ZeroVector;
+	return HasReachedInternal(TestPoint, GoalRadius, GoalHalfHeight, CurrentLocation, InAcceptanceRadius, !bExactSpot);
+}
+
+bool UPathFollowingComponent::HasReachedDestination(const FVector& CurrentLocation) const
+{
+	// get cylinder at goal location
+	FVector GoalLocation = Path->PathPoints[Path->PathPoints.Num() - 1].Location;
+	float GoalRadius = 0.0f;
+	float GoalHalfHeight = 0.0f;
+	
+	// take goal's current location, unless path is partial
+	if (DestinationActor.IsValid() && !Path->IsPartial())
+	{
+		if (DestinationAgent)
+		{
+			FVector GoalOffset;
+			DestinationAgent->GetMoveGoalReachTest(GetOwner(), MoveOffset, GoalOffset, GoalRadius, GoalHalfHeight);
+
+			GoalLocation = FRotationTranslationMatrix(DestinationActor->GetActorRotation(), DestinationAgent->GetNavAgentLocation()).TransformPosition(GoalOffset);
+		}
+		else
+		{
+			GoalLocation = DestinationActor->GetActorLocation();
+		}
+	}
+
+	return HasReachedInternal(GoalLocation, GoalRadius, GoalHalfHeight, CurrentLocation, AcceptanceRadius, bStopOnOverlap);
+}
+
+bool UPathFollowingComponent::HasReachedCurrentTarget(const FVector& CurrentLocation) const
+{
+	if (MovementComp == NULL)
+	{
+		return false;
+	}
+
+	const FVector CurrentTarget = GetCurrentTargetLocation();
+
+	// check if moved too far
+	const FVector ToTarget = (CurrentTarget - MovementComp->GetActorFeetLocation());
+	const float SegmentDot = FVector::DotProduct(ToTarget, MoveSegmentDirection);
+	if (SegmentDot < 0.0)
+	{
+		return true;
+	}
+
+	// or standing at target position
+	// don't use acceptance radius here, it has to be exact for moving near corners
+	const float GoalRadius = 0.0f;
+	const float GoalHalfHeight = 0.0f;
+	return HasReachedInternal(CurrentTarget, GoalRadius, GoalHalfHeight, CurrentLocation, CurrentAcceptanceRadius, false);
+}
+
+bool UPathFollowingComponent::HasReachedInternal(const FVector& Goal, float GoalRadius, float GoalHalfHeight, const FVector& AgentLocation, float RadiusThreshold, bool bUseAgentRadius) const
+{
+	if (MovementComp == NULL)
+	{
+		return false;
+	}
+
+	// get cylinder of moving agent
+	float AgentRadius = 0.0f;
+	float AgentHalfHeight = 0.0f;
+	AActor* MovingAgent = MovementComp->GetOwner();
+	MovingAgent->GetSimpleCollisionCylinder(AgentRadius, AgentHalfHeight);
+
+	// check if they overlap (with added AcceptanceRadius)
+	const FVector ToGoal = Goal - AgentLocation;
+
+	const float Dist2DSq = ToGoal.SizeSquared2D();
+	const float UseRadius = GoalRadius + RadiusThreshold + (bUseAgentRadius ? AgentRadius : 0.0f) + (AgentRadius * MinAgentRadiusPct);
+	if (Dist2DSq > FMath::Square(UseRadius))
+	{
+		return false;
+	}
+
+	const float ZDiff = FMath::Abs(ToGoal.Z);
+	const float UseHeight = GoalHalfHeight + AgentHalfHeight + (AgentHalfHeight * MinAgentHeightPct);
+	if (ZDiff > UseHeight)
+	{
+		return false;
+	}
+
+	return true;
+}
+
+void UPathFollowingComponent::DebugReachTest(float& CurrentDot, float& CurrentDistance, float& CurrentHeight, uint8& bDotFailed, uint8& bDistanceFailed, uint8& bHeightFailed) const
+{
+	if (!Path.IsValid() || MovementComp == NULL)
+	{
+		return;
+	}
+
+	const int32 LastSegmentEndIndex = Path->PathPoints.Num() - 1;
+	const bool bFollowingLastSegment = (MoveSegmentEndIndex >= LastSegmentEndIndex);
+
+	float GoalRadius = 0.0f;
+	float GoalHalfHeight = 0.0f;
+	bool bUseAgentRadius = false;
+	float RadiusThreshold = 0.0f;
+
+	FVector AgentLocation = MovementComp->GetActorFeetLocation();
+	FVector GoalLocation = GetCurrentTargetLocation();
+	RadiusThreshold = CurrentAcceptanceRadius;
+
+	if (bFollowingLastSegment)
+	{
+		GoalLocation = Path->PathPoints[Path->PathPoints.Num() - 1].Location;
+		bUseAgentRadius = bStopOnOverlap;
+
+		// take goal's current location, unless path is partial
+		if (DestinationActor.IsValid() && !Path->IsPartial())
+		{
+			if (DestinationAgent)
+			{
+				FVector GoalOffset;
+				DestinationAgent->GetMoveGoalReachTest(GetOwner(), MoveOffset, GoalOffset, GoalRadius, GoalHalfHeight);
+
+				GoalLocation = FRotationTranslationMatrix(DestinationActor->GetActorRotation(), DestinationAgent->GetNavAgentLocation()).TransformPosition(GoalOffset);
+			}
+			else
+			{
+				GoalLocation = DestinationActor->GetActorLocation();
+			}
+		}
+	}
+
+	const FVector ToGoal = (GoalLocation - AgentLocation);
+	CurrentDot = FVector::DotProduct(ToGoal.SafeNormal(), MoveSegmentDirection);
+	bDotFailed = (CurrentDot < 0.0f) ? 1 : 0;
+
+	// get cylinder of moving agent
+	float AgentRadius = 0.0f;
+	float AgentHalfHeight = 0.0f;
+	AActor* MovingAgent = MovementComp->GetOwner();
+	MovingAgent->GetSimpleCollisionCylinder(AgentRadius, AgentHalfHeight);
+
+	CurrentDistance = ToGoal.Size2D();
+	const float UseRadius = GoalRadius + RadiusThreshold + (bUseAgentRadius ? AgentRadius : 0.0f) + (AgentRadius * MinAgentRadiusPct);
+	bDistanceFailed = (CurrentDistance > UseRadius) ? 1 : 0;
+
+	CurrentHeight = FMath::Abs(ToGoal.Z);
+	const float UseHeight = GoalHalfHeight + AgentHalfHeight + (AgentHalfHeight * MinAgentHeightPct);
+	bHeightFailed = (CurrentHeight > UseHeight) ? 1 : 0;
+}
+
+void UPathFollowingComponent::SetCurrentSmartLink(USmartNavLinkComponent* InLink, const FVector& DestPoint)
+{
+	if (CurrentSmartLink)
+	{
+		CurrentSmartLink->NotifyLinkFinished(this);
+		CurrentSmartLink = NULL;
+	}
+
+	if (InLink)
+	{
+		PauseMove();
+
+		CurrentSmartLink = InLink;
+		CurrentSmartLink->NotifyLinkReached(this, DestPoint);
+	}
+}
+
+bool UPathFollowingComponent::UpdateMovementComponent(bool bForce)
+{
+	if (MovementComp == NULL || bForce == true)
+	{
+		APawn* MyPawn = Cast<APawn>(GetOwner());
+		if (MyPawn == NULL)
+		{
+			AController* MyController = Cast<AController>(GetOwner());
+			if (MyController)
+			{
+				MyPawn = MyController->GetPawn();
+			}
+		}
+
+		if (MyPawn)
+		{
+			FScriptDelegate Delegate;
+			Delegate.BindUFunction(this, "OnActorBump");
+			MyPawn->OnActorHit.AddUnique(Delegate);
+
+			SetMovementComponent(MyPawn->FindComponentByClass<UNavMovementComponent>());
+		}
+	}
+
+	return (MovementComp != NULL);
+}
+
+void UPathFollowingComponent::OnActorBump(AActor* SelfActor, AActor* OtherActor, FVector NormalImpulse, const FHitResult& Hit)
+{
+	if (Path.IsValid() && DestinationActor.IsValid() && DestinationActor.Get() == OtherActor)
+	{
+		UE_VLOG(GetOwner(), LogPathFollowing, Verbose, TEXT("Collided with goal actor"));
+		bCollidedWithGoal = true;
+	}
+}
+
+bool UPathFollowingComponent::IsOnPath() const
+{
+	bool bOnPath = false;
+	if (Path.IsValid() && Path->IsValid() && Path->GetOwner() != NULL)
+	{
+		const bool bHasNavigationCorridor = (Path->CastPath<FNavMeshPath>() != NULL);
+		if (bHasNavigationCorridor)
+		{
+			FNavLocation NavLoc = GetCurrentNavLocation();
+			bOnPath = Path->ContainsNode(NavLoc.NodeRef);
+		}
+		else
+		{
+			bOnPath = true;
+		}
+	}
+
+	return bOnPath;
+}
+
+bool UPathFollowingComponent::IsBlocked() const
+{
+	bool bBlocked = false;
+
+	if (LocationSamples.Num() == BlockDetectionSampleCount && BlockDetectionSampleCount > 0)
+	{
+		FVector Center = FVector::ZeroVector;
+		for (int32 i = 0; i < LocationSamples.Num(); i++)
+		{
+			Center += *LocationSamples[i];
+		}
+
+		Center /= LocationSamples.Num();
+		bBlocked = true;
+
+		for (int32 i = 0; i < LocationSamples.Num(); i++)
+		{
+			const float TestDistance = FVector::DistSquared(*LocationSamples[i], Center);
+			if (TestDistance > BlockDetectionDistance)
+			{
+				bBlocked = false;
+				break;
+			}
+		}
+	}
+
+	return bBlocked;
+}
+
+bool UPathFollowingComponent::UpdateBlockDetection()
+{
+	const float GameTime = GetWorld()->GetTimeSeconds();
+	if (bUseBlockDetection &&
+		MovementComp && 
+		GameTime > (LastSampleTime + BlockDetectionInterval) &&
+		BlockDetectionSampleCount > 0)
+	{
+		LastSampleTime = GameTime;
+
+		if (LocationSamples.Num() == NextSampleIdx)
+		{
+			LocationSamples.AddZeroed(1);
+		}
+
+		LocationSamples[NextSampleIdx] = MovementComp->GetActorFeetLocationBased();
+		NextSampleIdx = (NextSampleIdx + 1) % BlockDetectionSampleCount;
+		return true;
+	}
+
+	return false;
+}
+
+void UPathFollowingComponent::SetBlockDetection(float DistanceThreshold, float Interval, int32 NumSamples)
+{
+	BlockDetectionDistance = DistanceThreshold;
+	BlockDetectionInterval = Interval;
+	BlockDetectionSampleCount = NumSamples;
+	ResetBlockDetectionData();
+}
+
+void UPathFollowingComponent::SetBlockDetectionState(bool bEnable)
+{
+	bUseBlockDetection = bEnable;
+	ResetBlockDetectionData();
+}
+
+void UPathFollowingComponent::ResetBlockDetectionData()
+{
+	LastSampleTime = 0.0f;
+	NextSampleIdx = 0;
+	LocationSamples.Reset();
+}
+
+void UPathFollowingComponent::ForceBlockDetectionUpdate()
+{
+	LastSampleTime = 0.0f;
+}
+
+void UPathFollowingComponent::UpdateMoveFocus()
+{
+	AAIController* AIOwner = Cast<AAIController>(GetOwner());
+	if (AIOwner != NULL)
+	{
+		if (Status == EPathFollowingStatus::Moving)
+		{
+			const FVector MoveFocus = GetMoveFocus(AIOwner->bAllowStrafe);
+			AIOwner->SetFocalPoint(MoveFocus, false, EAIFocusPriority::Move);
+		}
+		else
+		{
+			AIOwner->ClearFocus(EAIFocusPriority::Move);
+		}
+	}
+}
+
+void UPathFollowingComponent::SetPreciseReachThreshold(float AgentRadiusMultiplier, float AgentHeightMultiplier)
+{
+	MinAgentRadiusPct = AgentRadiusMultiplier;
+	MinAgentHeightPct = AgentHeightMultiplier;
+}
+
+float UPathFollowingComponent::GetRemainingPathCost() const
+{
+	float Cost = 0.f;
+
+	if (Path.IsValid() && Path->IsValid() && Status == EPathFollowingStatus::Moving)
+	{
+		const FNavLocation& NavLocation = GetCurrentNavLocation();
+		Cost = Path->GetCostFromNode(NavLocation.NodeRef);
+	}
+
+	return Cost;
+}	
+
+FNavLocation UPathFollowingComponent::GetCurrentNavLocation() const
+{
+	// get navigation location of moved actor
+	if (MovementComp == NULL || GetWorld() == NULL || GetWorld()->GetNavigationSystem() == NULL)
+	{
+		return FVector::ZeroVector;
+	}
+
+	const FVector OwnerLoc = MovementComp->GetActorNavLocation();
+	const float Tolerance = 1.0f;
+	
+	// Using !Equals rather than != because exact equivalence isn't required.  (Equals allows a small tolerance.)
+	if (!OwnerLoc.Equals(CurrentNavLocation.Location, Tolerance))
+	{
+		const AActor* OwnerActor = MovementComp->GetOwner();
+		const FVector OwnerExtent = OwnerActor ? OwnerActor->GetSimpleCollisionCylinderExtent() : FVector::ZeroVector;
+
+		GetWorld()->GetNavigationSystem()->ProjectPointToNavigation(OwnerLoc, CurrentNavLocation, OwnerExtent, MyNavData);
+	}
+
+	return CurrentNavLocation;
+}
+
+// will be deprecated soon, please don't use it!
+EPathFollowingAction::Type UPathFollowingComponent::GetPathActionType() const
+{
+	switch (Status)
+	{
+		case EPathFollowingStatus::Idle:
+			return EPathFollowingAction::NoMove;
+
+		case EPathFollowingStatus::Waiting:
+		case EPathFollowingStatus::Paused:	
+		case EPathFollowingStatus::Moving:
+			return !Path.IsValid() ? EPathFollowingAction::Error :
+				(Path->GetOwner() == NULL) ? EPathFollowingAction::DirectMove :
+				Path->IsPartial() ? EPathFollowingAction::PartialPath : EPathFollowingAction::PathToGoal;
+
+		default:
+			break;
+	}
+
+	return EPathFollowingAction::Error;
+}
+
+// will be deprecated soon, please don't use it!
+FVector UPathFollowingComponent::GetPathDestination() const
+{
+	return Path.IsValid() ? Path->GetDestinationLocation() : FVector::ZeroVector;
+}
+
+FString UPathFollowingComponent::GetStatusDesc() const
+{
+	const static UEnum* StatusEnum = FindObject<UEnum>(ANY_PACKAGE, TEXT("EPathFollowingStatus"));
+	if (StatusEnum)
+	{
+		return StatusEnum->GetEnumName(Status);
+	}
+
+	return TEXT("Unknown");
+}
+
+FString UPathFollowingComponent::GetResultDesc(EPathFollowingResult::Type Result) const
+{
+	const static UEnum* ResultEnum = FindObject<UEnum>(ANY_PACKAGE, TEXT("EPathFollowingResult"));
+	if (ResultEnum)
+	{
+		return ResultEnum->GetEnumName(Result);
+	}
+
+	return TEXT("Unknown");
+}
+
+void UPathFollowingComponent::DisplayDebug(class UCanvas* Canvas, const FDebugDisplayInfo& DebugDisplay, float& YL, float& YPos) const
+{
+	Canvas->SetDrawColor(FColor::Blue);
+	UFont* RenderFont = GEngine->GetSmallFont();
+	FString StatusDesc = FString::Printf(TEXT("  Move status: %s"), *GetStatusDesc());
+	Canvas->DrawText(RenderFont, StatusDesc, 4.0f, YPos);
+	YPos += YL;
+
+	if (Status == EPathFollowingStatus::Moving)
+	{
+		const int32 NumMoveSegments = (Path.IsValid() && Path->IsValid()) ? Path->PathPoints.Num() : -1;
+		FString TargetDesc = FString::Printf(TEXT("  Move target [%d/%d]: %s (%s)"),
+			MoveSegmentEndIndex, NumMoveSegments, *GetCurrentTargetLocation().ToString(), *GetNameSafe(DestinationActor.Get()));
+		
+		Canvas->DrawText(RenderFont, TargetDesc, 4.0f, YPos);
+		YPos += YL;
+	}
+}
+
+#if ENABLE_VISUAL_LOG
+void UPathFollowingComponent::DescribeSelfToVisLog(struct FVisLogEntry* Snapshot) const
+{
+	FVisLogEntry::FStatusCategory Category;
+	Category.Category = TEXT("Path following");
+
+	if (DestinationActor.IsValid())
+	{
+		Category.Add(TEXT("Goal"), GetNameSafe(DestinationActor.Get()));
+	}
+	
+	FString StatusDesc = GetStatusDesc();
+	if (Status == EPathFollowingStatus::Moving)
+	{
+		const int32 NumMoveSegments = (Path.IsValid() && Path->IsValid()) ? Path->PathPoints.Num() : -1;
+		StatusDesc += FString::Printf(TEXT(" [%d..%d/%d]"), MoveSegmentStartIndex + 1, MoveSegmentEndIndex + 1, NumMoveSegments);
+	}
+
+	Category.Add(TEXT("Status"), StatusDesc);
+	Category.Add(TEXT("Path"), !Path.IsValid() ? TEXT("none") : (Path->CastPath<FNavMeshPath>() != NULL) ? TEXT("navmesh") : TEXT("direct"));
+	
+	if (CurrentSmartLink)
+	{
+		Category.Add(TEXT("SmartLink"), CurrentSmartLink->GetName());
+	}
+
+	Snapshot->Status.Add(Category);
+}
+#endif // ENABLE_VISUAL_LOG
+
+void UPathFollowingComponent::GetDebugStringTokens(TArray<FString>& Tokens, TArray<EPathFollowingDebugTokens::Type>& Flags) const
+{
+	Tokens.Add(GetStatusDesc());
+	Flags.Add(EPathFollowingDebugTokens::Description);
+
+	if (Status != EPathFollowingStatus::Moving)
+	{
+		return;
+	}
+
+	FString& StatusDesc = Tokens[0];
+	if (Path.IsValid())
+	{
+		const int32 NumMoveSegments = (Path.IsValid() && Path->IsValid()) ? Path->PathPoints.Num() : -1;
+		const bool bIsDirect = (Path->CastPath<FNavMeshPath>() == NULL);
+		const bool bIsSmartLink = (CurrentSmartLink != NULL);
+
+		if (!bIsDirect)
+		{
+			StatusDesc += FString::Printf(TEXT(" (%d..%d/%d)%s"), MoveSegmentStartIndex + 1, MoveSegmentEndIndex + 1, NumMoveSegments,
+				bIsSmartLink ? TEXT(" (SmartLink)") : TEXT(""));
+		}
+		else
+		{
+			StatusDesc += TEXT(" (direct)");
+		}
+	}
+	else
+	{
+		StatusDesc += TEXT(" (invalid path)");
+	}
+
+	// add debug params
+	float CurrentDot = 0.0f, CurrentDistance = 0.0f, CurrentHeight = 0.0f;
+	uint8 bFailedDot = 0, bFailedDistance = 0, bFailedHeight = 0;
+	DebugReachTest(CurrentDot, CurrentDistance, CurrentHeight, bFailedHeight, bFailedDistance, bFailedHeight);
+
+	Tokens.Add(TEXT("dot"));
+	Flags.Add(EPathFollowingDebugTokens::ParamName);
+	Tokens.Add(FString::Printf(TEXT("%.2f"), CurrentDot));
+	Flags.Add(bFailedDot ? EPathFollowingDebugTokens::FailedValue : EPathFollowingDebugTokens::PassedValue);
+
+	Tokens.Add(TEXT("dist2D"));
+	Flags.Add(EPathFollowingDebugTokens::ParamName);
+	Tokens.Add(FString::Printf(TEXT("%.0f"), CurrentDistance));
+	Flags.Add(bFailedDistance ? EPathFollowingDebugTokens::FailedValue : EPathFollowingDebugTokens::PassedValue);
+
+	Tokens.Add(TEXT("distZ"));
+	Flags.Add(EPathFollowingDebugTokens::ParamName);
+	Tokens.Add(FString::Printf(TEXT("%.0f"), CurrentHeight));
+	Flags.Add(bFailedHeight ? EPathFollowingDebugTokens::FailedValue : EPathFollowingDebugTokens::PassedValue);
+}
+
+FString UPathFollowingComponent::GetDebugString() const
+{
+	TArray<FString> Tokens;
+	TArray<EPathFollowingDebugTokens::Type> Flags;
+
+	GetDebugStringTokens(Tokens, Flags);
+
+	FString Desc;
+	for (int32 i = 0; i < Tokens.Num(); i++)
+	{
+		Desc += Tokens[i];
+		Desc += (Flags.IsValidIndex(i) && Flags[i] == EPathFollowingDebugTokens::ParamName) ? TEXT(": ") : TEXT(", ");
+	}
+
+	return Desc;
+}
+
+void UPathFollowingComponent::LockResource(EAILockSource::Type LockSource)
+{
+	const static UEnum* SourceEnum = FindObject<UEnum>(ANY_PACKAGE, TEXT("EAILockSource"));
+	const bool bWasLocked = ResourceLock.IsLocked();
+
+	ResourceLock.SetLock(LockSource);
+	if (bWasLocked == false)
+	{
+		UE_VLOG(GetOwner(), LogPathFollowing, Log, TEXT("Locking Move by source %s"), SourceEnum ? *SourceEnum->GetEnumName(Status) : TEXT("invalid"));
+		PauseMove();
+	}
+}
+
+void UPathFollowingComponent::ClearResourceLock(EAILockSource::Type LockSource)
+{
+	const bool bWasLocked = ResourceLock.IsLocked();
+	ResourceLock.ClearLock(LockSource);
+
+	if (bWasLocked && !ResourceLock.IsLocked())
+	{
+		UE_VLOG(GetOwner(), LogPathFollowing, Log, TEXT("Unlocking Move"));
+		ResumeMove();
+	}
+}
+
+void UPathFollowingComponent::ForceUnlockResource()
+{
+	const bool bWasLocked = ResourceLock.IsLocked();
+	ResourceLock.ForceClearAllLocks();
+
+	if (bWasLocked)
+	{
+		UE_VLOG(GetOwner(), LogPathFollowing, Log, TEXT("Unlocking Move - forced"));
+		ResumeMove();
+	}
+}
+
+bool UPathFollowingComponent::IsResourceLocked() const
+{
+	return ResourceLock.IsLocked();
+}
+
+void UPathFollowingComponent::SetLastMoveAtGoal(bool bFinishedAtGoal)
+{
+	bLastMoveReachedGoal = bFinishedAtGoal;
+}
diff --git a/Engine/Source/Runtime/AIModule/Public/BehaviorTreeDelegates.h b/Engine/Source/Runtime/AIModule/Public/BehaviorTreeDelegates.h
new file mode 100644
--- /dev/null
+++ b/Engine/Source/Runtime/AIModule/Public/BehaviorTreeDelegates.h
@@ -0,0 +1,21 @@
+// Copyright 1998-2014 Epic Games, Inc. All Rights Reserved.
+
+#pragma once
+
+struct AIMODULE_API FBehaviorTreeDelegates
+{
+	/** delegate type for tree execution events (Params: const UBehaviorTreeComponent* OwnerComp, const UBehaviorTree* TreeAsset) */
+	DECLARE_MULTICAST_DELEGATE_TwoParams(FOnTreeStarted, const class UBehaviorTreeComponent*, const class UBehaviorTree* );
+
+	/** delegate type for locking AI debugging tool on pawn */
+	DECLARE_MULTICAST_DELEGATE_OneParam(FOnTreeDebugTool, const class APawn*);
+
+	/** Called when the behavior tree starts execution. */
+	static FOnTreeStarted OnTreeStarted;
+
+	/** Called when the AI debug tool highlights a pawn. */
+	static FOnTreeDebugTool OnDebugSelected;
+
+	/** Called when the AI debug tool locks on pawn. */
+	static FOnTreeDebugTool OnDebugLocked;
+};
diff --git a/Engine/Source/Runtime/AIModule/Public/BlueprintNodeHelpers.h b/Engine/Source/Runtime/AIModule/Public/BlueprintNodeHelpers.h
new file mode 100644
--- /dev/null
+++ b/Engine/Source/Runtime/AIModule/Public/BlueprintNodeHelpers.h
@@ -0,0 +1,32 @@
+// Copyright 1998-2014 Epic Games, Inc. All Rights Reserved.
+
+#pragma once
+
+namespace BlueprintNodeHelpers
+{
+	FString CollectPropertyDescription(const UObject* Ob, const UClass* StopAtClass, const TArray<UProperty*>& PropertyData);
+	void CollectPropertyData(const UObject* Ob, const UClass* StopAtClass, TArray<UProperty*>& PropertyData);
+	uint16 GetPropertiesMemorySize(const TArray<UProperty*>& PropertyData);
+
+	void CollectBlackboardSelectors(const UObject* Ob, const UClass* StopAtClass, TArray<FName>& KeyNames);
+
+	AIMODULE_API FString DescribeProperty(const UProperty* Prop, const uint8* PropertyAddr);
+	void DescribeRuntimeValues(const UObject* Ob, const TArray<UProperty*>& PropertyData, TArray<FString>& RuntimeValues);
+
+	void CopyPropertiesToContext(const TArray<UProperty*>& PropertyData, uint8* ObjectMemory, uint8* ContextMemory);
+	void CopyPropertiesFromContext(const TArray<UProperty*>& PropertyData, uint8* ObjectMemory, uint8* ContextMemory);
+
+	bool FindNodeOwner(AActor* OwningActor, class UBTNode* Node, UBehaviorTreeComponent*& OwningComp, int32& OwningInstanceIdx);
+
+	void AbortLatentActions(AActor* OwningActor, const UObject* Ob);
+
+	FORCEINLINE bool HasBlueprintFunction(FName FuncName, const UObject* Ob, const UClass* StopAtClass)
+	{
+		return (Ob->GetClass()->FindFunctionByName(FuncName)->GetOuter() != StopAtClass);
+	}
+
+	FORCEINLINE FString GetNodeName(const UObject* Ob)
+	{
+		return Ob->GetClass()->GetName().LeftChop(2);
+	}
+}
diff --git a/Engine/Source/Runtime/Engine/Classes/AI/Navigation/NavAreas/NavAreaMeta.h b/Engine/Source/Runtime/Engine/Classes/AI/Navigation/NavAreas/NavAreaMeta.h
--- a/Engine/Source/Runtime/Engine/Classes/AI/Navigation/NavAreas/NavAreaMeta.h
+++ b/Engine/Source/Runtime/Engine/Classes/AI/Navigation/NavAreas/NavAreaMeta.h
@@ -1,6 +1,7 @@
 // Copyright 1998-2014 Epic Games, Inc. All Rights Reserved.
 
 #pragma once
+#include "AI/Navigation/NavAreas/NavArea.h"
 #include "NavAreaMeta.generated.h"
 
 /** Class containing definition of a navigation area */
diff --git a/Engine/Source/Runtime/Engine/Classes/AI/Navigation/NavAreas/NavAreaMeta_SwitchByAgent.h b/Engine/Source/Runtime/Engine/Classes/AI/Navigation/NavAreas/NavAreaMeta_SwitchByAgent.h
--- a/Engine/Source/Runtime/Engine/Classes/AI/Navigation/NavAreas/NavAreaMeta_SwitchByAgent.h
+++ b/Engine/Source/Runtime/Engine/Classes/AI/Navigation/NavAreas/NavAreaMeta_SwitchByAgent.h
@@ -1,6 +1,7 @@
 // Copyright 1998-2014 Epic Games, Inc. All Rights Reserved.
 
 #pragma once
+#include "AI/Navigation/NavAreas/NavAreaMeta.h"
 #include "NavAreaMeta_SwitchByAgent.generated.h"
 
 /** Class containing definition of a navigation area */
diff --git a/Engine/Source/Runtime/Engine/Classes/AI/Navigation/NavAreas/NavArea_Default.h b/Engine/Source/Runtime/Engine/Classes/AI/Navigation/NavAreas/NavArea_Default.h
--- a/Engine/Source/Runtime/Engine/Classes/AI/Navigation/NavAreas/NavArea_Default.h
+++ b/Engine/Source/Runtime/Engine/Classes/AI/Navigation/NavAreas/NavArea_Default.h
@@ -1,6 +1,7 @@
 // Copyright 1998-2014 Epic Games, Inc. All Rights Reserved.
 
 #pragma once
+#include "AI/Navigation/NavAreas/NavArea.h"
 #include "NavArea_Default.generated.h"
 
 /** Regular navigation area, applied to entire navigation data by default */
diff --git a/Engine/Source/Runtime/Engine/Classes/AI/Navigation/NavFilters/NavigationQueryFilter.h b/Engine/Source/Runtime/Engine/Classes/AI/Navigation/NavFilters/NavigationQueryFilter.h
--- a/Engine/Source/Runtime/Engine/Classes/AI/Navigation/NavFilters/NavigationQueryFilter.h
+++ b/Engine/Source/Runtime/Engine/Classes/AI/Navigation/NavFilters/NavigationQueryFilter.h
@@ -108,7 +108,6 @@ class ENGINE_API UNavigationQueryFilter : public UObject
 	TSharedPtr<const struct FNavigationQueryFilter> GetQueryFilter(const class ANavigationData* NavData) const;
 	
 	/** helper functions for accessing filter */
-	static TSharedPtr<const struct FNavigationQueryFilter> GetQueryFilter(const class AAIController* AI, UClass* FilterClass);
 	static TSharedPtr<const struct FNavigationQueryFilter> GetQueryFilter(const class ANavigationData* NavData, UClass* FilterClass);
 
 	template<class T>
diff --git a/Engine/Source/Runtime/Engine/Classes/AI/Navigation/NavigationSystem.h b/Engine/Source/Runtime/Engine/Classes/AI/Navigation/NavigationSystem.h
--- a/Engine/Source/Runtime/Engine/Classes/AI/Navigation/NavigationSystem.h
+++ b/Engine/Source/Runtime/Engine/Classes/AI/Navigation/NavigationSystem.h
@@ -2,7 +2,8 @@
 
 #pragma once
 
-#include "NavigationTypes.h"
+#include "AI/Navigation/NavFilters/NavigationQueryFilter.h"
+#include "AI/Navigation/NavigationTypes.h"
 #include "GenericOctreePublic.h"
 #include "NavigationSystem.generated.h"
 
@@ -583,11 +584,9 @@ class ENGINE_API UNavigationSystem : public UBlueprintFunctionLibrary
  	UPROPERTY(/*BlueprintAssignable, */Transient)
 	FOnNavDataRegistered OnNavDataRegisteredEvent;
 
-protected:
-	UPROPERTY(transient)
-	class UCrowdManager* CrowdManager;
-
 private:
+	TWeakObjectPtr<class UCrowdManager> CrowdManager;
+
 	// required navigation data 
 	UPROPERTY(config)
 	TArray<FStringClassReference> RequiredNavigationDataClassNames;
@@ -631,7 +630,7 @@ public:
 	 *	@param HitLocation if line was obstructed this will be set to hit location. Otherwise it contains SegmentEnd
 	 *	@return true if line from RayStart to RayEnd was obstructed. Also, true when no navigation data present */
 	UFUNCTION(BlueprintCallable, Category="AI|Navigation", meta=(HidePin="WorldContext", DefaultToSelf="WorldContext" ))
-	static bool NavigationRaycast(UObject* WorldContext, const FVector& RayStart, const FVector& RayEnd, FVector& HitLocation, TSubclassOf<class UNavigationQueryFilter> FilterClass = NULL, class AAIController* Querier = NULL);
+	static bool NavigationRaycast(UObject* WorldContext, const FVector& RayStart, const FVector& RayEnd, FVector& HitLocation, TSubclassOf<class UNavigationQueryFilter> FilterClass = NULL, class AController* Querier = NULL);
 
 	/** delegate type for events that dirty the navigation data ( Params: const FBox& DirtyBounds ) */
 	DECLARE_MULTICAST_DELEGATE_OneParam(FOnNavigationDirty, const FBox&);
@@ -664,7 +663,7 @@ public:
 
 	UWorld* GetWorld() const { return GetOuterUWorld(); }
 
-	class UCrowdManager* GetCrowdManager() const { return CrowdManager; }
+	class UCrowdManager* GetCrowdManager() const { return CrowdManager.Get(); }
 
 	/** spawn new crowd manager */
 	virtual void CreateCrowdManager();
@@ -1062,6 +1061,8 @@ protected:
 	void AddActorElementToNavOctree(class AActor* Actor, const FNavigationDirtyElement& DirtyElement);
 	void AddComponentElementToNavOctree(class UActorComponent* ActorComp, const FNavigationDirtyElement& DirtyElement, const FBox& Bounds);
 
+	void SetCrowdManager(class UCrowdManager* NewCrowdManager); 
+
 #if WITH_NAVIGATION_GENERATOR
 
 	/** Add BSP collision data to navigation octree */
diff --git a/Engine/Source/Runtime/Engine/Classes/AI/Navigation/SmartNavLinkComponent.h b/Engine/Source/Runtime/Engine/Classes/AI/Navigation/SmartNavLinkComponent.h
--- a/Engine/Source/Runtime/Engine/Classes/AI/Navigation/SmartNavLinkComponent.h
+++ b/Engine/Source/Runtime/Engine/Classes/AI/Navigation/SmartNavLinkComponent.h
@@ -2,6 +2,7 @@
 
 #pragma once
 #include "NavigationTypes.h"
+#include "AI/Navigation/NavRelevantComponent.h"
 #include "SmartNavLinkComponent.generated.h"
 
 /**
diff --git a/Engine/Source/Runtime/Engine/Classes/Animation/AnimInstance.h b/Engine/Source/Runtime/Engine/Classes/Animation/AnimInstance.h
--- a/Engine/Source/Runtime/Engine/Classes/Animation/AnimInstance.h
+++ b/Engine/Source/Runtime/Engine/Classes/Animation/AnimInstance.h
@@ -275,16 +275,16 @@ public:
 	void ClearSlotNodeWeights();
 	bool IsActiveSlotNode(FName SlotNodeName) const;
 
-protected:
 	// kismet event functions
 
+	UFUNCTION(BlueprintPure, Category = "Animation")
+	virtual class APawn* TryGetPawnOwner();
+
+protected:
 	/** Returns the owning actor of this AnimInstance */
 	UFUNCTION(BlueprintPure, Category="Animation")
 	class AActor* GetOwningActor();
-
-	UFUNCTION(BlueprintPure, Category="Animation")
-	virtual class APawn* TryGetPawnOwner();
-
+	
 	// Returns the skeletal mesh component that has created this AnimInstance
 	UFUNCTION(BlueprintPure, Category="Animation")
 	class USkeletalMeshComponent* GetOwningComponent();
@@ -406,12 +406,14 @@ public:
 	float CalculateDirection(const FVector & Velocity, const FRotator & BaseRotation);
 
 	//--- AI communication start ---//
-	/** locks indicated AI resources of animated pawn */
-	UFUNCTION(BlueprintCallable, Category="Animation", BlueprintAuthorityOnly)
+	/** locks indicated AI resources of animated pawn
+	 *	DEPRECATED. Use LockAIResourcesWithAnimation instead */
+	UFUNCTION(BlueprintCallable, Category = "Animation", BlueprintAuthorityOnly, Meta=(DeprecatedFunction, DeprecationMessage="Use LockAIResourcesWithAnimation instead"))
 	void LockAIResources(bool bLockMovement, bool LockAILogic);
 
-	/** unlocks indicated AI resources of animated pawn. Will unlock only animation-locked resources */
-	UFUNCTION(BlueprintCallable, Category="Animation", BlueprintAuthorityOnly)
+	/** unlocks indicated AI resources of animated pawn. Will unlock only animation-locked resources.
+	 *	DEPRECATED. Use UnlockAIResourcesWithAnimation instead */
+	UFUNCTION(BlueprintCallable, Category = "Animation", BlueprintAuthorityOnly, Meta=(DeprecatedFunction, DeprecationMessage="Use UnlockAIResourcesWithAnimation instead"))
 	void UnlockAIResources(bool bUnlockMovement, bool UnlockAILogic);
 	//--- AI communication end ---//
 
diff --git a/Engine/Source/Runtime/Engine/Classes/Debug/GameplayDebuggingComponent.h b/Engine/Source/Runtime/Engine/Classes/Debug/GameplayDebuggingComponent.h
--- a/Engine/Source/Runtime/Engine/Classes/Debug/GameplayDebuggingComponent.h
+++ b/Engine/Source/Runtime/Engine/Classes/Debug/GameplayDebuggingComponent.h
@@ -6,8 +6,11 @@
 
 #pragma once
 #include "GameplayDebuggingControllerComponent.h"
+#include "DebugRenderSceneProxy.h"
 #include "GameplayDebuggingComponent.generated.h"
 
+#define WITH_EQS 0
+
 DECLARE_LOG_CATEGORY_EXTERN(LogGDT, Display, All);
 
 UENUM()
@@ -56,8 +59,8 @@ namespace EQSDebug
 	{
 		TArray<FItemData> Items;
 		TArray<FTestData> Tests;
-		TArray<FEQSSceneProxy::FSphere> SolidSpheres;
-		TArray<FEQSSceneProxy::FText3d> Texts;
+		TArray<FDebugRenderSceneProxy::FSphere> SolidSpheres;
+		TArray<FDebugRenderSceneProxy::FText3d> Texts;
 		int32 NumValidItems;
 	};
 }
@@ -65,7 +68,7 @@ namespace EQSDebug
 DECLARE_MULTICAST_DELEGATE_TwoParams(FOnDebuggingTargetChanged, class AActor* /*Owner of debugging component*/, bool /*is being debugged now*/);
 
 UCLASS()
-class ENGINE_API UGameplayDebuggingComponent : public UPrimitiveComponent, public IEQSQueryResultSourceInterface
+class ENGINE_API UGameplayDebuggingComponent : public UPrimitiveComponent//, public IEQSQueryResultSourceInterface
 {
 	GENERATED_UCLASS_BODY()
 
@@ -106,6 +109,7 @@ class ENGINE_API UGameplayDebuggingComponent : public UPrimitiveComponent, publi
 	UPROPERTY(ReplicatedUsing = OnRep_UpdateNavmesh)
 	TArray<uint8> NavmeshRepData;
 	
+#if WITH_EQS
 	/** Begin EQS replication data */
 	UPROPERTY(Replicated)
 	float EQSTimestamp;
@@ -128,10 +132,11 @@ class ENGINE_API UGameplayDebuggingComponent : public UPrimitiveComponent, publi
 	uint32 bDrawEQSLabels:1;
 	uint32 bDrawEQSFailedItems : 1;
 
-	virtual bool GetComponentClassCanReplicate() const OVERRIDE{ return true; }
-
 	UFUNCTION()
 	virtual void OnRep_UpdateEQS();
+#endif // WITH_EQS
+
+	virtual bool GetComponentClassCanReplicate() const OVERRIDE{ return true; }
 
 	UFUNCTION()
 	virtual void OnRep_UpdateNavmesh();
@@ -175,6 +180,7 @@ class ENGINE_API UGameplayDebuggingComponent : public UPrimitiveComponent, publi
 	UFUNCTION(reliable, server, WithValidation)
 	void ServerEnableTargetSelection(bool bEnable);
 
+#if WITH_EQS
 	//=============================================================================
 	// EQS debugging
 	//=============================================================================
@@ -190,6 +196,10 @@ class ENGINE_API UGameplayDebuggingComponent : public UPrimitiveComponent, publi
 	virtual bool GetShouldDebugDrawLabels() const OVERRIDE { return bDrawEQSLabels; }
 	virtual bool GetShouldDrawFailedItems() const OVERRIDE{ return bDrawEQSFailedItems; }
 	// IEQSQueryResultSourceInterface end
+protected:
+	virtual void CollectEQSData();
+public:
+#endif //WITH_EQS
 
 	//=============================================================================
 	// Rendering
@@ -208,7 +218,6 @@ protected:
 	virtual void CollectPathData();
 	virtual void CollectBasicData();
 	virtual void CollectBehaviorTreeData();
-	virtual void CollectEQSData();
 
 	virtual void GetKeyboardDesc(TArray<FDebugCategoryView>& Categories);
 
diff --git a/Engine/Source/Runtime/Engine/Classes/Engine/Engine.h b/Engine/Source/Runtime/Engine/Classes/Engine/Engine.h
--- a/Engine/Source/Runtime/Engine/Classes/Engine/Engine.h
+++ b/Engine/Source/Runtime/Engine/Classes/Engine/Engine.h
@@ -637,23 +637,7 @@ public:
 	/** The class for NavigationSystem **/
 	UPROPERTY()
 	TSubclassOf<class UNavigationSystem>  NavigationSystemClass;
-
-	/** Name of behavior tree manager class */
-	UPROPERTY(globalconfig, noclear, meta=(MetaClass="BehaviorTreeManager", DisplayName="Behavior Tree Manager Class"))
-	FStringClassReference BehaviorTreeManagerClassName;
-
-	/** The class for behavior tree manager **/
-	UPROPERTY()
-	TSubclassOf<class UBehaviorTreeManager>  BehaviorTreeManagerClass;
-
-	/** Name of environment query manager class */
-	UPROPERTY(globalconfig, noclear, meta=(MetaClass="EnvQueryManager", DisplayName="Environment Query Manager Class"))
-	FStringClassReference EnvironmentQueryManagerClassName;
-
-	/** The class for environment query manager **/
-	UPROPERTY()
-	TSubclassOf<class UEnvQueryManager>  EnvironmentQueryManagerClass;
-
+	
 	/** Name of behavior tree manager class */
 	UPROPERTY(globalconfig, noclear, meta=(MetaClass="AvoidanceManager", DisplayName="Avoidance Manager Class"))
 	FStringClassReference AvoidanceManagerClassName;
diff --git a/Engine/Source/Runtime/Engine/Classes/Engine/EngineTypes.h b/Engine/Source/Runtime/Engine/Classes/Engine/EngineTypes.h
--- a/Engine/Source/Runtime/Engine/Classes/Engine/EngineTypes.h
+++ b/Engine/Source/Runtime/Engine/Classes/Engine/EngineTypes.h
@@ -2348,7 +2348,7 @@ namespace EComponentMobility
 }
 
 UCLASS(abstract, config=Engine)
-class UEngineTypes : public UObject
+class ENGINE_API UEngineTypes : public UObject
 {
 	GENERATED_UCLASS_BODY()
 
diff --git a/Engine/Source/Runtime/Engine/Classes/Engine/World.h b/Engine/Source/Runtime/Engine/Classes/Engine/World.h
--- a/Engine/Source/Runtime/Engine/Classes/Engine/World.h
+++ b/Engine/Source/Runtime/Engine/Classes/Engine/World.h
@@ -565,13 +565,11 @@ private:
 	UPROPERTY(Transient)
 	class AGameMode*							AuthorityGameMode;
 
-	/** Behavior tree manager used by game */
-	UPROPERTY(Transient)
-	class UBehaviorTreeManager*					BehaviorTreeManager;
+	/** Behavior tree manager used by game */	
+	TWeakObjectPtr<class UBehaviorTreeManager> BehaviorTreeManager;
 
 	/** Environment query manager used by game */
-	UPROPERTY(Transient)
-	class UEnvQueryManager*						EnvironmentQueryManager;
+	TWeakObjectPtr<class UEnvQueryManager>		EnvironmentQueryManager;
 
 	/** RVO avoidance manager used by game */
 	UPROPERTY(Transient)
@@ -1413,14 +1411,14 @@ public:
 	FORCEINLINE const UNavigationSystem* GetNavigationSystem() const { return NavigationSystem; }
 
 	/** Behavior tree manager getter */
-	FORCEINLINE class UBehaviorTreeManager* GetBehaviorTreeManager() { return BehaviorTreeManager; }
+	FORCEINLINE class UBehaviorTreeManager* GetBehaviorTreeManager() { return BehaviorTreeManager.Get(); }
 	/** Behavior tree manager const getter */
-	FORCEINLINE const class UBehaviorTreeManager* GetBehaviorTreeManager() const { return BehaviorTreeManager; }
+	FORCEINLINE const class UBehaviorTreeManager* GetBehaviorTreeManager() const { return BehaviorTreeManager.Get(); }
 
 	/** Behavior tree manager getter */
-	FORCEINLINE class UEnvQueryManager* GetEnvironmentQueryManager() { return EnvironmentQueryManager; }
+	FORCEINLINE class UEnvQueryManager* GetEnvironmentQueryManager() { return EnvironmentQueryManager.Get(); }
 	/** Behavior tree manager const getter */
-	FORCEINLINE const class UEnvQueryManager* GetEnvironmentQueryManager() const { return EnvironmentQueryManager; }
+	FORCEINLINE const class UEnvQueryManager* GetEnvironmentQueryManager() const { return EnvironmentQueryManager.Get(); }
 
 	/** Behavior tree manager getter */
 	FORCEINLINE class UAvoidanceManager* GetAvoidanceManager() { return AvoidanceManager; }
diff --git a/Engine/Source/Runtime/Engine/Classes/GameFramework/CharacterMovementComponent.h b/Engine/Source/Runtime/Engine/Classes/GameFramework/CharacterMovementComponent.h
--- a/Engine/Source/Runtime/Engine/Classes/GameFramework/CharacterMovementComponent.h
+++ b/Engine/Source/Runtime/Engine/Classes/GameFramework/CharacterMovementComponent.h
@@ -6,7 +6,7 @@
  */
 
 #pragma once
-#include "../AI/AITypes.h"
+
 #include "CharacterMovementComponent.generated.h"
 
 /** Movement modes for Characters.  */
@@ -22,6 +22,112 @@ enum EMovementMode
 	MOVE_MAX		UMETA(Hidden),
 };
 
+USTRUCT()
+struct FNavAvoidanceMask
+{
+	GENERATED_USTRUCT_BODY()
+
+#if CPP
+	union
+	{
+		struct
+		{
+#endif
+		UPROPERTY(Category="Group", EditAnywhere, BlueprintReadOnly)
+		uint32 bGroup0 : 1;
+		UPROPERTY(Category="Group", EditAnywhere, BlueprintReadOnly)
+		uint32 bGroup1 : 1;
+		UPROPERTY(Category="Group", EditAnywhere, BlueprintReadOnly)
+		uint32 bGroup2 : 1;
+		UPROPERTY(Category="Group", EditAnywhere, BlueprintReadOnly)
+		uint32 bGroup3 : 1;
+		UPROPERTY(Category="Group", EditAnywhere, BlueprintReadOnly)
+		uint32 bGroup4 : 1;
+		UPROPERTY(Category="Group", EditAnywhere, BlueprintReadOnly)
+		uint32 bGroup5 : 1;
+		UPROPERTY(Category="Group", EditAnywhere, BlueprintReadOnly)
+		uint32 bGroup6 : 1;
+		UPROPERTY(Category="Group", EditAnywhere, BlueprintReadOnly)
+		uint32 bGroup7 : 1;
+		UPROPERTY(Category="Group", EditAnywhere, BlueprintReadOnly)
+		uint32 bGroup8 : 1;
+		UPROPERTY(Category="Group", EditAnywhere, BlueprintReadOnly)
+		uint32 bGroup9 : 1;
+		UPROPERTY(Category="Group", EditAnywhere, BlueprintReadOnly)
+		uint32 bGroup10 : 1;
+		UPROPERTY(Category="Group", EditAnywhere, BlueprintReadOnly)
+		uint32 bGroup11 : 1;
+		UPROPERTY(Category="Group", EditAnywhere, BlueprintReadOnly)
+		uint32 bGroup12 : 1;
+		UPROPERTY(Category="Group", EditAnywhere, BlueprintReadOnly)
+		uint32 bGroup13 : 1;
+		UPROPERTY(Category="Group", EditAnywhere, BlueprintReadOnly)
+		uint32 bGroup14 : 1;
+		UPROPERTY(Category="Group", EditAnywhere, BlueprintReadOnly)
+		uint32 bGroup15 : 1;
+		UPROPERTY(Category="Group", EditAnywhere, BlueprintReadOnly)
+		uint32 bGroup16 : 1;
+		UPROPERTY(Category="Group", EditAnywhere, BlueprintReadOnly)
+		uint32 bGroup17 : 1;
+		UPROPERTY(Category="Group", EditAnywhere, BlueprintReadOnly)
+		uint32 bGroup18 : 1;
+		UPROPERTY(Category="Group", EditAnywhere, BlueprintReadOnly)
+		uint32 bGroup19 : 1;
+		UPROPERTY(Category="Group", EditAnywhere, BlueprintReadOnly)
+		uint32 bGroup20 : 1;
+		UPROPERTY(Category="Group", EditAnywhere, BlueprintReadOnly)
+		uint32 bGroup21 : 1;
+		UPROPERTY(Category="Group", EditAnywhere, BlueprintReadOnly)
+		uint32 bGroup22 : 1;
+		UPROPERTY(Category="Group", EditAnywhere, BlueprintReadOnly)
+		uint32 bGroup23 : 1;
+		UPROPERTY(Category="Group", EditAnywhere, BlueprintReadOnly)
+		uint32 bGroup24 : 1;
+		UPROPERTY(Category="Group", EditAnywhere, BlueprintReadOnly)
+		uint32 bGroup25 : 1;
+		UPROPERTY(Category="Group", EditAnywhere, BlueprintReadOnly)
+		uint32 bGroup26 : 1;
+		UPROPERTY(Category="Group", EditAnywhere, BlueprintReadOnly)
+		uint32 bGroup27 : 1;
+		UPROPERTY(Category="Group", EditAnywhere, BlueprintReadOnly)
+		uint32 bGroup28 : 1;
+		UPROPERTY(Category="Group", EditAnywhere, BlueprintReadOnly)
+		uint32 bGroup29 : 1;
+		UPROPERTY(Category="Group", EditAnywhere, BlueprintReadOnly)
+		uint32 bGroup30 : 1;
+		UPROPERTY(Category="Group", EditAnywhere, BlueprintReadOnly)
+		uint32 bGroup31 : 1;
+#if CPP
+		};
+		int32 Packed;
+	};
+#endif
+
+	FORCEINLINE bool HasGroup(uint8 GroupID) const
+	{
+		return (Packed & (1 << GroupID)) != 0;
+	}
+
+	FORCEINLINE void SetGroup(uint8 GroupID)
+	{
+		Packed |= (1 << GroupID);
+	}
+
+	FORCEINLINE void ClearGroup(uint8 GroupID)
+	{
+		Packed &= ~(1 << GroupID);
+	}
+
+	FORCEINLINE void ClearAll()
+	{
+		Packed = 0;
+	}
+
+	FORCEINLINE void SetFlagsDirectly(uint32 NewFlagset)
+	{
+		Packed = NewFlagset;
+	}
+};
 
 // Data about the floor for walking movement.
 USTRUCT(BlueprintType)
diff --git a/Engine/Source/Runtime/Engine/Classes/GameFramework/NavMovementComponent.h b/Engine/Source/Runtime/Engine/Classes/GameFramework/NavMovementComponent.h
--- a/Engine/Source/Runtime/Engine/Classes/GameFramework/NavMovementComponent.h
+++ b/Engine/Source/Runtime/Engine/Classes/GameFramework/NavMovementComponent.h
@@ -35,7 +35,7 @@ public:
 	void UpdateNavAgent(class UCapsuleComponent* CapsuleComponent);
 
 	/** @returns location of controlled actor - meaning center of collision bounding box */
-	FORCEINLINE FVector GetActorLocation() const { return UpdatedComponent ? UpdatedComponent->GetComponentLocation() : FAISystem::InvalidLocation; }
+	FORCEINLINE FVector GetActorLocation() const { return UpdatedComponent ? UpdatedComponent->GetComponentLocation() : FVector(FLT_MAX); }
 	/** @returns location of controlled actor's "feet" meaning center of bottom of collision bounding box */
 	FORCEINLINE FVector GetActorFeetLocation() const { return UpdatedComponent ? (UpdatedComponent->GetComponentLocation() - FVector(0,0,UpdatedComponent->Bounds.BoxExtent.Z)) : FVector::ZeroVector; }
 	/** @returns based location of controlled actor */
