Removing AIModule dependencies from the Engine, WIP #UE4

diff --git a/Engine/Source/Runtime/AIModule/Private/AISystemExec.cpp b/Engine/Source/Runtime/AIModule/Private/AISystemExec.cpp
new file mode 100644
--- /dev/null
+++ b/Engine/Source/Runtime/AIModule/Private/AISystemExec.cpp
@@ -0,0 +1,45 @@
+// Copyright 1998-2014 Epic Games, Inc. All Rights Reserved.
+
+#include "AIModulePrivate.h"
+#if !(UE_BUILD_SHIPPING || UE_BUILD_TEST)
+#include "Misc/CoreMisc.h"
+
+struct FAISystemExec : public FSelfRegisteringExec
+{
+	FAISystemExec()
+	{
+	}
+
+	// Begin FExec Interface
+	virtual bool Exec(UWorld* Inworld, const TCHAR* Cmd, FOutputDevice& Ar) OVERRIDE;
+	// End FExec Interface
+};
+
+FAISystemExec AISystemExecInstance;
+
+bool FAISystemExec::Exec(UWorld* Inworld, const TCHAR* Cmd, FOutputDevice& Ar)
+{
+	bool bHandled = false;
+
+	if (FParse::Command(&Cmd, TEXT("AIIgnorePlayers")))
+	{
+		AAIController::ToggleAIIgnorePlayers();		
+		bHandled = true; 
+	}
+	else if (FParse::Command(&Cmd, TEXT("AILoggingVerbose")))
+	{
+		if (Inworld)
+		{
+			APlayerController* PC = Inworld->GetFirstPlayerController();
+			if (PC)
+			{
+				PC->ConsoleCommand(TEXT("log lognavigation verbose | log logpathfollowing verbose | log LogCharacter verbose | log LogBehaviorTree verbose | log LogPawnAction verbose|"));
+			}
+		}
+		bHandled = true;
+	}
+
+	return bHandled;
+}
+
+#endif // !(UE_BUILD_SHIPPING || UE_BUILD_TEST)
\ No newline at end of file
diff --git a/Engine/Source/Runtime/Engine/Classes/GameFramework/CheatManager.h b/Engine/Source/Runtime/Engine/Classes/GameFramework/CheatManager.h
--- a/Engine/Source/Runtime/Engine/Classes/GameFramework/CheatManager.h
+++ b/Engine/Source/Runtime/Engine/Classes/GameFramework/CheatManager.h
@@ -159,10 +159,6 @@ class ENGINE_API UCheatManager : public UObject
 	UFUNCTION(exec)
 	virtual void Summon(const FString& ClassName);
 
-	// Toggle AI ignoring the player.
-	UFUNCTION(exec)
-	virtual void AIIgnorePlayers();
-
 	/** Freeze everything in the level except for players. */
 	UFUNCTION(exec)
 	virtual void PlayersOnly();
@@ -208,10 +204,6 @@ class ENGINE_API UCheatManager : public UObject
 
 	void DumpAILogs();
 
-	/** makes various AI logging categories verbose */
-	UFUNCTION(exec)
-	virtual void AILoggingVerbose();
-
 	/** Toggle capsule trace debugging. Will trace a capsule from current view point and show where it hits the world */
 	UFUNCTION(exec)
 	virtual void DebugCapsuleSweep();
diff --git a/Engine/Source/Runtime/Engine/Private/AI/Navigation/NavFilters/NavigationQueryFilter.cpp b/Engine/Source/Runtime/Engine/Private/AI/Navigation/NavFilters/NavigationQueryFilter.cpp
--- a/Engine/Source/Runtime/Engine/Private/AI/Navigation/NavFilters/NavigationQueryFilter.cpp
+++ b/Engine/Source/Runtime/Engine/Private/AI/Navigation/NavFilters/NavigationQueryFilter.cpp
@@ -2,7 +2,6 @@
 
 #include "EnginePrivate.h"
 // @todo AIModule circular dependency
-#include "AIController.h"
 #include "Navigation/NavigationComponent.h"
 #include "AI/Navigation/NavFilters/NavigationQueryFilter.h"
 
diff --git a/Engine/Source/Runtime/Engine/Private/AI/Navigation/NavigationSystem.cpp b/Engine/Source/Runtime/Engine/Private/AI/Navigation/NavigationSystem.cpp
--- a/Engine/Source/Runtime/Engine/Private/AI/Navigation/NavigationSystem.cpp
+++ b/Engine/Source/Runtime/Engine/Private/AI/Navigation/NavigationSystem.cpp
@@ -262,7 +262,7 @@ TSubclassOf<class UNavArea> UNavigationSystem::DefaultWalkableArea = NULL;
 TSubclassOf<class UNavArea> UNavigationSystem::DefaultObstacleArea = NULL;
 FOnNavAreaChanged UNavigationSystem::NavAreaAddedObservers;
 FOnNavAreaChanged UNavigationSystem::NavAreaRemovedObservers;
-#if !UE_BUILD_SHIPPING
+#if !(UE_BUILD_SHIPPING || UE_BUILD_TEST)
 FNavigationSystemExec UNavigationSystem::ExecHandler;
 #endif // !UE_BUILD_SHIPPING
 
diff --git a/Engine/Source/Runtime/Engine/Private/CheatManager.cpp b/Engine/Source/Runtime/Engine/Private/CheatManager.cpp
--- a/Engine/Source/Runtime/Engine/Private/CheatManager.cpp
+++ b/Engine/Source/Runtime/Engine/Private/CheatManager.cpp
@@ -371,11 +371,6 @@ void UCheatManager::Summon( const FString& ClassName )
 	}
 }
 
-void UCheatManager::AIIgnorePlayers()
-{
-	AAIController::ToggleAIIgnorePlayers();
-}
-
 void UCheatManager::PlayersOnly()
 {
 	check( GetWorld() );
@@ -644,16 +639,6 @@ void UCheatManager::ToggleAILogging()
 #endif
 }
 
-void UCheatManager::AILoggingVerbose()
-{
-	APlayerController* PC = GetOuterAPlayerController();
-	if (PC)
-	{
-		PC->ConsoleCommand(TEXT("log lognavigation verbose | log logpathfollowing verbose | log LogCharacter verbose | log LogBehaviorTree verbose | log LogPawnAction verbose|"));
-	}
-}
-
-
 #define SAFE_TRACEINDEX_DECREASE(x) ((--x)<0)? 9:(x)
 
 #if !(UE_BUILD_SHIPPING || UE_BUILD_TEST)
