Fixed EQS manager's default object ticking #UE4

diff --git a/Engine/Source/Runtime/AIModule/Classes/EnvironmentQuery/EnvQueryManager.h b/Engine/Source/Runtime/AIModule/Classes/EnvironmentQuery/EnvQueryManager.h
--- a/Engine/Source/Runtime/AIModule/Classes/EnvironmentQuery/EnvQueryManager.h
+++ b/Engine/Source/Runtime/AIModule/Classes/EnvironmentQuery/EnvQueryManager.h
@@ -92,8 +92,8 @@ class AIMODULE_API UEnvQueryManager : public UObject, public FTickableGameObject
 	/** [FTickableGameObject] tick function */
 	virtual void Tick(float DeltaTime) OVERRIDE;
 
-	/** [FTickableGameObject] always tick */
-	virtual bool IsTickable() const OVERRIDE { return true; }
+	/** [FTickableGameObject] always tick, unless it's the default object */
+	virtual bool IsTickable() const OVERRIDE { return HasAnyFlags(RF_ClassDefaultObject) == false; }
 
 	/** [FTickableGameObject] tick stats */
 	virtual TStatId GetStatId() const OVERRIDE;
