Fixed TTP #337747 "[AI] AI Move To BP node's OnSuccess and OnFail no longer get triggered"

diff --git a/Engine/Source/Runtime/AIModule/Classes/AITypes.h b/Engine/Source/Runtime/AIModule/Classes/AITypes.h
--- a/Engine/Source/Runtime/AIModule/Classes/AITypes.h
+++ b/Engine/Source/Runtime/AIModule/Classes/AITypes.h
@@ -177,6 +177,7 @@ private:
 	static const uint32 AnyRequestID = 0;
 	static const uint32 InvalidRequestID = uint32(-1);
 
+	UPROPERTY()
 	uint32 RequestID;
 
 public:
