fixed crash on loading broken behavior tree asset

diff --git a/Engine/Source/Runtime/AIModule/Private/BehaviorTree/BehaviorTreeManager.cpp b/Engine/Source/Runtime/AIModule/Private/BehaviorTree/BehaviorTreeManager.cpp
--- a/Engine/Source/Runtime/AIModule/Private/BehaviorTree/BehaviorTreeManager.cpp
+++ b/Engine/Source/Runtime/AIModule/Private/BehaviorTree/BehaviorTreeManager.cpp
@@ -108,6 +108,16 @@ static void InitializeNodeHelper(UBTCompositeNode* ParentNode, UBTNode* NodeOb,
 	{
 		for (int32 i = 0; i < CompositeOb->Services.Num(); i++)
 		{
+			if (CompositeOb->Services[i] == NULL)
+			{
+				UE_LOG(LogBehaviorTree, Warning, TEXT("%s has missing service node! (parent: %s)"),
+					*GetNameSafe(TreeAsset), *UBehaviorTreeTypes::DescribeNodeHelper(CompositeOb));
+
+				CompositeOb->Services.RemoveAt(i, 1, false);
+				i--;
+				continue;
+			}
+
 			UBTService* Service = Cast<UBTService>(StaticDuplicateObject(CompositeOb->Services[i], NodeOuter, TEXT("None")));;
 			CompositeOb->Services[i] = Service;
 
@@ -123,6 +133,16 @@ static void InitializeNodeHelper(UBTCompositeNode* ParentNode, UBTNode* NodeOb,
 			FBTCompositeChild& ChildInfo = CompositeOb->Children[i];
 			for (int32 j = 0; j < ChildInfo.Decorators.Num(); j++)
 			{
+				if (ChildInfo.Decorators[j] == NULL)
+				{
+					UE_LOG(LogBehaviorTree, Warning, TEXT("%s has missing decorator node! (parent: %s, branch: %d)"),
+						*GetNameSafe(TreeAsset), *UBehaviorTreeTypes::DescribeNodeHelper(CompositeOb), i);
+
+					ChildInfo.Decorators.RemoveAt(j, 1, false);
+					j--;
+					continue;
+				}
+
 				UBTDecorator* Decorator = Cast<UBTDecorator>(StaticDuplicateObject(ChildInfo.Decorators[j], NodeOuter, TEXT("None")));
 				ChildInfo.Decorators[j] = Decorator;
 
