Fixed differen issues for GameplayDebugger and fixed missing features we had before moving to separate module.

diff --git a/Engine/Source/Developer/GameplayDebugger/Classes/GameplayDebuggingReplicator.h b/Engine/Source/Developer/GameplayDebugger/Classes/GameplayDebuggingReplicator.h
--- a/Engine/Source/Developer/GameplayDebugger/Classes/GameplayDebuggingReplicator.h
+++ b/Engine/Source/Developer/GameplayDebugger/Classes/GameplayDebuggingReplicator.h
@@ -90,4 +90,5 @@ private:
 	TWeakObjectPtr<AGameplayDebuggingHUDComponent>	DebugRenderer;
 	TArray<FDebugContext> Clients;
 	TWeakObjectPtr<APlayerController> LocalPlayerOwner;
+	TWeakObjectPtr<AActor>	LastSelectedActorInSimulation;
 };
diff --git a/Engine/Source/Developer/GameplayDebugger/Classes/GameplayDebuggingTypes.h b/Engine/Source/Developer/GameplayDebugger/Classes/GameplayDebuggingTypes.h
--- a/Engine/Source/Developer/GameplayDebugger/Classes/GameplayDebuggingTypes.h
+++ b/Engine/Source/Developer/GameplayDebugger/Classes/GameplayDebuggingTypes.h
@@ -44,5 +44,6 @@ struct GAMEPLAYDEBUGGER_API FGameplayDebuggerSettings
 	}
 
 	static uint32 DebuggerShowFlags;
+	static uint32 ShowFlagIndex;
 };
 
diff --git a/Engine/Source/Developer/GameplayDebugger/Private/GameplayDebugger.cpp b/Engine/Source/Developer/GameplayDebugger/Private/GameplayDebugger.cpp
--- a/Engine/Source/Developer/GameplayDebugger/Private/GameplayDebugger.cpp
+++ b/Engine/Source/Developer/GameplayDebugger/Private/GameplayDebugger.cpp
@@ -6,6 +6,7 @@
 #include "GameplayDebuggingReplicator.h"
 
 uint32 FGameplayDebuggerSettings::DebuggerShowFlags = (1 << EAIDebugDrawDataView::Basic) | (1 << EAIDebugDrawDataView::OverHead);
+uint32 FGameplayDebuggerSettings::ShowFlagIndex = 0;
 
 class FGameplayDebugger : public IGameplayDebugger
 {
diff --git a/Engine/Source/Developer/GameplayDebugger/Private/GameplayDebuggerExec.cpp b/Engine/Source/Developer/GameplayDebugger/Private/GameplayDebuggerExec.cpp
--- a/Engine/Source/Developer/GameplayDebugger/Private/GameplayDebuggerExec.cpp
+++ b/Engine/Source/Developer/GameplayDebugger/Private/GameplayDebuggerExec.cpp
@@ -78,14 +78,14 @@ bool FGameplayDebuggerExec::Exec(UWorld* Inworld, const TCHAR* Cmd, FOutputDevic
 
 		FString InViewName = FParse::Token(Cmd, 0);
 		int32 ViewIndex = ViewNames.Find(InViewName);
-		const static uint32 ShowFlagIndex = FEngineShowFlags::FindIndexByName(TEXT("GameplayDebug"));
+		FGameplayDebuggerSettings::ShowFlagIndex = FEngineShowFlags::FindIndexByName(TEXT("GameplayDebug"));
 		if (ViewIndex != INDEX_NONE)
 		{
 			FGameplayDebuggerSettings::CheckFlag(ViewIndex) ? FGameplayDebuggerSettings::ClearFlag(ViewIndex) : FGameplayDebuggerSettings::SetFlag(ViewIndex);
 #if WITH_EDITOR
 			if (ViewIndex == EAIDebugDrawDataView::EQS && GCurrentLevelEditingViewportClient)
 			{
-				GCurrentLevelEditingViewportClient->EngineShowFlags.SetSingleFlag(ShowFlagIndex, FGameplayDebuggerSettings::CheckFlag(ViewIndex));
+				GCurrentLevelEditingViewportClient->EngineShowFlags.SetSingleFlag(FGameplayDebuggerSettings::ShowFlagIndex, FGameplayDebuggerSettings::CheckFlag(ViewIndex));
 			}
 #endif
 			PC->ClientMessage(FString::Printf(TEXT("View %s %s")
diff --git a/Engine/Source/Developer/GameplayDebugger/Private/GameplayDebuggingComponent.cpp b/Engine/Source/Developer/GameplayDebugger/Private/GameplayDebuggingComponent.cpp
--- a/Engine/Source/Developer/GameplayDebugger/Private/GameplayDebuggingComponent.cpp
+++ b/Engine/Source/Developer/GameplayDebugger/Private/GameplayDebuggingComponent.cpp
@@ -234,11 +234,6 @@ void UGameplayDebuggingComponent::SetActorToDebug(AActor* Actor)
 {
 #if !(UE_BUILD_SHIPPING || UE_BUILD_TEST)
 	TargetActor = Actor;
-	APawn* TargetPawn = Cast<APawn>(TargetActor);
-	if (TargetPawn)
-	{
-		FBehaviorTreeDelegates::OnDebugSelected.Broadcast(TargetPawn);
-	}
 #endif //!(UE_BUILD_SHIPPING || UE_BUILD_TEST)
 }
 
@@ -459,8 +454,6 @@ void UGameplayDebuggingComponent::SelectForDebugging(bool bNewStatus)
 	{
 		bIsSelectedForDebugging = bNewStatus;
 		
-		OnDebuggingTargetChangedDelegate.Broadcast(GetSelectedActor(), bNewStatus);
-
 #if 0
 		// temp: avoidance debug
 		UAvoidanceManager* Avoidance = GetWorld()->GetAvoidanceManager();
diff --git a/Engine/Source/Developer/GameplayDebugger/Private/GameplayDebuggingControllerComponent.cpp b/Engine/Source/Developer/GameplayDebugger/Private/GameplayDebuggingControllerComponent.cpp
--- a/Engine/Source/Developer/GameplayDebugger/Private/GameplayDebuggingControllerComponent.cpp
+++ b/Engine/Source/Developer/GameplayDebugger/Private/GameplayDebuggingControllerComponent.cpp
@@ -331,13 +331,17 @@ void UGameplayDebuggingControllerComponent::TickComponent(float DeltaTime, enum
 
 void UGameplayDebuggingControllerComponent::UpdateNavMeshTimer()
 {
-	APawn* Pawn = PlayerOwner.IsValid() ? PlayerOwner->GetPawn() : NULL;
+	const APawn* PlayerPawn = PlayerOwner.IsValid() ? PlayerOwner->GetPawn() : NULL;
 	UGameplayDebuggingComponent* DebuggingComponent = GetDebuggingReplicator() ? GetDebuggingReplicator()->GetDebugComponent() : NULL;
 	if (DebuggingComponent)
 	{
+		const AActor* SelectedActor = DebuggingComponent->GetSelectedActor();
+		const APawn* SelectedActorAsPawn = Cast<APawn>(SelectedActor);
+
 		const FVector AdditionalTargetLoc =
-			DebugAITargetActor ? DebugAITargetActor->GetActorLocation() :
-			Pawn ? Pawn->GetActorLocation() :
+			SelectedActorAsPawn ? SelectedActorAsPawn->GetNavAgentLocation() :
+			SelectedActor ? SelectedActor->GetActorLocation() :
+			PlayerPawn ? PlayerPawn->GetNavAgentLocation() :
 			FVector::ZeroVector;
 
 		if (AdditionalTargetLoc != FVector::ZeroVector)
diff --git a/Engine/Source/Developer/GameplayDebugger/Private/GameplayDebuggingReplicator.cpp b/Engine/Source/Developer/GameplayDebugger/Private/GameplayDebuggingReplicator.cpp
--- a/Engine/Source/Developer/GameplayDebugger/Private/GameplayDebuggingReplicator.cpp
+++ b/Engine/Source/Developer/GameplayDebugger/Private/GameplayDebuggingReplicator.cpp
@@ -249,6 +249,8 @@ void AGameplayDebuggingReplicator::OnDebugAIDelegate(class UCanvas* Canvas, clas
 	UWorld* World = GetWorld();
 	if (World && DebugComponent && DebugComponent->GetOwnerRole() == ROLE_Authority)
 	{
+		bool bBroadcastedNewSelection = false;
+
 		// looks like Simulate in UE4 Editor - let's find selected Pawn to debug
 		for (FConstPawnIterator Iterator = World->GetPawnIterator(); Iterator; ++Iterator)
 		{
@@ -261,8 +263,25 @@ void AGameplayDebuggingReplicator::OnDebugAIDelegate(class UCanvas* Canvas, clas
 			DebugComponent->SelectForDebugging(NewTarget->IsSelected());
 			DebugComponent->CollectDataToReplicate(NewTarget->IsSelected());
 
+			if (NewTarget->IsSelected() && LastSelectedActorInSimulation.Get() != NewTarget)
+			{
+				LastSelectedActorInSimulation = NewTarget;
+				UGameplayDebuggingComponent::OnDebuggingTargetChangedDelegate.Broadcast(NewTarget, NewTarget->IsSelected());
+				APawn* TargetPawn = Cast<APawn>(NewTarget);
+				if (TargetPawn)
+				{
+					FBehaviorTreeDelegates::OnDebugSelected.Broadcast(TargetPawn);
+				}
+				bBroadcastedNewSelection = true;
+			}
 			DrawDebugData(Canvas, PC);
 		}
+
+		if (!bBroadcastedNewSelection && LastSelectedActorInSimulation.IsValid() && !LastSelectedActorInSimulation->IsSelected() )
+		{
+			UGameplayDebuggingComponent::OnDebuggingTargetChangedDelegate.Broadcast(LastSelectedActorInSimulation.Get(), false);
+			LastSelectedActorInSimulation = NULL;
+		}
 	}
 #endif
 }
diff --git a/Engine/Source/Developer/LogVisualizer/LogVisualizer.Build.cs b/Engine/Source/Developer/LogVisualizer/LogVisualizer.Build.cs
--- a/Engine/Source/Developer/LogVisualizer/LogVisualizer.Build.cs
+++ b/Engine/Source/Developer/LogVisualizer/LogVisualizer.Build.cs
@@ -31,6 +31,7 @@ public class LogVisualizer : ModuleRules
 				"Engine",
 				"TaskGraph",
 				"UnrealEd",
+                "GameplayDebugger",
 			}
 		);
 
diff --git a/Engine/Source/Developer/LogVisualizer/Private/SLogVisualizer.cpp b/Engine/Source/Developer/LogVisualizer/Private/SLogVisualizer.cpp
--- a/Engine/Source/Developer/LogVisualizer/Private/SLogVisualizer.cpp
+++ b/Engine/Source/Developer/LogVisualizer/Private/SLogVisualizer.cpp
@@ -19,6 +19,8 @@
 #	include "Editor/UnrealEd/Public/EditorViewportClient.h"
 #endif
 
+#include "GameplayDebuggingComponent.h"
+
 #include "SFilterList.h"
 
 #define LOCTEXT_NAMESPACE "SLogVisualizer"
@@ -443,13 +445,13 @@ void SLogVisualizer::Construct(const FArguments& InArgs, FLogVisualizer* InLogVi
 
 	DrawingOnCanvasDelegate = FDebugDrawDelegate::CreateSP(this, &SLogVisualizer::DrawOnCanvas);
 	UDebugDrawService::Register(TEXT("VisLog"), DrawingOnCanvasDelegate);
-	//UGameplayDebuggingComponent::OnDebuggingTargetChangedDelegate.AddSP(this, &SLogVisualizer::SelectionChanged);
+	UGameplayDebuggingComponent::OnDebuggingTargetChangedDelegate.AddSP(this, &SLogVisualizer::SelectionChanged);
 }
 END_SLATE_FUNCTION_BUILD_OPTIMIZATION
 
 SLogVisualizer::~SLogVisualizer()
 {
-	//UGameplayDebuggingComponent::OnDebuggingTargetChangedDelegate.RemoveAll(this);
+	UGameplayDebuggingComponent::OnDebuggingTargetChangedDelegate.RemoveAll(this);
 	LogVisualizer->OnLogAdded().RemoveAll(this);
 	UDebugDrawService::Unregister(DrawingOnCanvasDelegate);
 }
diff --git a/Engine/Source/Editor/BehaviorTreeEditor/BehaviorTreeEditor.Build.cs b/Engine/Source/Editor/BehaviorTreeEditor/BehaviorTreeEditor.Build.cs
--- a/Engine/Source/Editor/BehaviorTreeEditor/BehaviorTreeEditor.Build.cs
+++ b/Engine/Source/Editor/BehaviorTreeEditor/BehaviorTreeEditor.Build.cs
@@ -40,6 +40,7 @@ public class BehaviorTreeEditor : ModuleRules
 				"AnimGraph",
 				"BlueprintGraph",
                 "AIModule",
+                "GameplayDebugger",
 			}
 		);
 
diff --git a/Engine/Source/Editor/BehaviorTreeEditor/Private/BehaviorTreeDebugger.cpp b/Engine/Source/Editor/BehaviorTreeEditor/Private/BehaviorTreeDebugger.cpp
--- a/Engine/Source/Editor/BehaviorTreeEditor/Private/BehaviorTreeDebugger.cpp
+++ b/Engine/Source/Editor/BehaviorTreeEditor/Private/BehaviorTreeDebugger.cpp
@@ -7,6 +7,8 @@
 #include "BehaviorTree/BTAuxiliaryNode.h"
 #include "BehaviorTreeDelegates.h"
 #include "SBehaviorTreeDebuggerView.h"
+#include "GameplayDebuggingReplicator.h"
+#include "GameplayDebuggingComponent.h"
 
 FBehaviorTreeDebugger::FBehaviorTreeDebugger()
 {
@@ -624,17 +626,28 @@ void FBehaviorTreeDebugger::FindLockedDebugActor(UWorld* World)
 	APlayerController* LocalPC = GEngine->GetFirstLocalPlayerController(World);
 	if (LocalPC && LocalPC->GetHUD() && LocalPC->GetPawnOrSpectator())
 	{
-//		UGameplayDebuggingControllerComponent* DebuggingController = LocalPC->FindComponentByClass<UGameplayDebuggingControllerComponent>();
-//		const APawn* LockedPawn = DebuggingController != NULL ? DebuggingController->GetCurrentDebugTarget() : NULL;
-//		UBehaviorTreeComponent* TestInstance = FindInstanceInActor((APawn*)LockedPawn);
-//		if (TestInstance)
-//		{
-//			TreeInstance = TestInstance;
-//
-//#if USE_BEHAVIORTREE_DEBUGGER
-//			ActiveStepIndex = TestInstance->DebuggerSteps.Num() - 1;
-//#endif
-//		}
+		AGameplayDebuggingReplicator* DebuggingReplicator = NULL;
+		for (FActorIterator It(World); It; ++It)
+		{
+			AActor* A = *It;
+			if (A && A->IsA(AGameplayDebuggingReplicator::StaticClass()) && !A->IsPendingKill())
+			{
+				DebuggingReplicator = Cast<AGameplayDebuggingReplicator>(A);
+				break;
+			}
+		}
+
+		UGameplayDebuggingComponent* DebuggingComponent = DebuggingReplicator != NULL ? DebuggingReplicator->GetDebugComponent() : NULL;
+		const APawn* LockedPawn = DebuggingComponent != NULL ? Cast<APawn>(DebuggingComponent->GetSelectedActor()) : NULL;
+		UBehaviorTreeComponent* TestInstance = FindInstanceInActor((APawn*)LockedPawn);
+		if (TestInstance)
+		{
+			TreeInstance = TestInstance;
+
+#if USE_BEHAVIORTREE_DEBUGGER
+			ActiveStepIndex = TestInstance->DebuggerSteps.Num() - 1;
+#endif
+		}
 	}
 }
 
diff --git a/Engine/Source/Runtime/AIModule/Private/AISystem.cpp b/Engine/Source/Runtime/AIModule/Private/AISystem.cpp
--- a/Engine/Source/Runtime/AIModule/Private/AISystem.cpp
+++ b/Engine/Source/Runtime/AIModule/Private/AISystem.cpp
@@ -1,10 +1,12 @@
 // Copyright 1998-2014 Epic Games, Inc. All Rights Reserved.
 
 #include "AIModulePrivate.h"
+#include "EngineUtils.h"
 #include "BehaviorTree/BehaviorTreeManager.h"
 #include "EnvironmentQuery/EnvQueryManager.h"
 #include "GameFramework/PlayerController.h"
-//#include "Debug/GameplayDebuggingControllerComponent.h"
+#include "GameplayDebuggingReplicator.h"
+#include "GameplayDebuggingComponent.h"
 #include "AISystem.h"
 
 UAISystem::UAISystem(const FPostConstructInitializeProperties& PCIP) : Super(PCIP)
@@ -68,47 +70,34 @@ void UAISystem::RunEQS(const FString& QueryName)
 	UEnvQueryManager* EQS = GetEnvironmentQueryManager();
 	if (MyPC && EQS)
 	{
-		//AActor* Target = NULL;
-		//UGameplayDebuggingControllerComponent* DebugComp = MyPC->FindComponentByClass<UGameplayDebuggingControllerComponent>();
-		//if (DebugComp)
-		//{
-		//	Target = DebugComp->GetCurrentDebugTarget();
-		//}
-
-//#if WITH_EDITOR
-//		if (Target == NULL && GEditor != NULL)
-//		{
-//			Target = GEditor->GetSelectedObjects()->GetTop<AActor>();
-//
-//			// this part should not be needed, but is due to gameplay debugging messed up design
-//			if (Target == NULL)
-//			{
-//				for (FObjectIterator It(UGameplayDebuggingComponent::StaticClass()); It && (Target == NULL); ++It)
-//				{
-//					UGameplayDebuggingComponent* Comp = ((UGameplayDebuggingComponent*)*It);
-//					if (Comp->IsSelected())
-//					{
-//						Target = Comp->GetOwner();
-//					}
-//				}
-//			}
-//		}
-//#endif // WITH_EDITOR
-
-		//if (Target)
-		//{
-		//	UEnvQuery* QueryTemplate = EQS->FindQueryTemplate(QueryName);
-
-		//	if (QueryTemplate)
-		//	{
-		//		EQS->RunInstantQuery(FEnvQueryRequest(QueryTemplate, Target), EEnvQueryRunMode::AllMatching);
-		//	}
-		//	else
-		//	{
-		//		MyPC->ClientMessage(FString::Printf(TEXT("Unable to fing query template \'%s\'"), *QueryName));
-		//	}
-		//}
-		//else
+		AGameplayDebuggingReplicator* DebuggingReplicator = NULL;
+		for (FActorIterator It(OuterWorld); It; ++It)
+		{
+			AActor* A = *It;
+			if (A && A->IsA(AGameplayDebuggingReplicator::StaticClass()) && !A->IsPendingKill())
+			{
+				DebuggingReplicator = Cast<AGameplayDebuggingReplicator>(A);
+				break;
+			}
+		}
+
+		UGameplayDebuggingComponent* DebuggingComponent = DebuggingReplicator != NULL ? DebuggingReplicator->GetDebugComponent() : NULL;
+		UObject* Target = DebuggingComponent != NULL ? DebuggingComponent->GetSelectedActor() : NULL;
+
+		if (Target)
+		{
+			const UEnvQuery* QueryTemplate = EQS->FindQueryTemplate(QueryName);
+
+			if (QueryTemplate)
+			{
+				EQS->RunInstantQuery(FEnvQueryRequest(QueryTemplate, Target), EEnvQueryRunMode::AllMatching);
+			}
+			else
+			{
+				MyPC->ClientMessage(FString::Printf(TEXT("Unable to fing query template \'%s\'"), *QueryName));
+			}
+		}
+		else
 		{
 			MyPC->ClientMessage(TEXT("No debugging target"));
 		}
diff --git a/Engine/Source/Runtime/AIModule/Private/EnvironmentQuery/EQSTestingPawn.cpp b/Engine/Source/Runtime/AIModule/Private/EnvironmentQuery/EQSTestingPawn.cpp
--- a/Engine/Source/Runtime/AIModule/Private/EnvironmentQuery/EQSTestingPawn.cpp
+++ b/Engine/Source/Runtime/AIModule/Private/EnvironmentQuery/EQSTestingPawn.cpp
@@ -1,6 +1,7 @@
 // Copyright 1998-2014 Epic Games, Inc. All Rights Reserved.
 
 #include "AIModulePrivate.h"
+#include "GameplayDebuggingTypes.h"
 #include "EnvironmentQuery/EnvQuery.h"
 #include "EnvironmentQuery/EQSTestingPawn.h"
 #include "EnvironmentQuery/EnvQueryTypes.h"
@@ -95,10 +96,10 @@ void AEQSTestingPawn::OnEditorSelectionChanged(UObject* NewSelection)
 	}
 
 #if WITH_EDITOR
-	//if (GCurrentLevelEditingViewportClient != NULL && UGameplayDebuggingComponent::ShowFlagIndex != INDEX_NONE)
-	//{
-	//	GCurrentLevelEditingViewportClient->EngineShowFlags.SetSingleFlag(UGameplayDebuggingComponent::ShowFlagIndex, bEQSPawnSelected);
-	//}
+	if (GCurrentLevelEditingViewportClient != NULL && FGameplayDebuggerSettings::ShowFlagIndex != INDEX_NONE)
+	{
+		GCurrentLevelEditingViewportClient->EngineShowFlags.SetSingleFlag(FGameplayDebuggerSettings::ShowFlagIndex, bEQSPawnSelected);
+	}
 #endif // WITH_EDITOR
 }
 
