GitHub pull request #164: Blackboard none indexes

https://github.com/EpicGames/UnrealEngine/pull/164 ([REDACTED_PROFILE_URL])

Where possible (everywhere except for u-properties) replaces uses of uint8 as Blackboard key ID to FBlackboard::FKey. Also all occurrences of INDEX_NONE as key ID have been replaced with FBlackboard::InvalidKey.

#github #UE4

diff --git a/Engine/Source/Runtime/AIModule/Classes/BehaviorTree/BehaviorTreeTypes.h b/Engine/Source/Runtime/AIModule/Classes/BehaviorTree/BehaviorTreeTypes.h
--- a/Engine/Source/Runtime/AIModule/Classes/BehaviorTree/BehaviorTreeTypes.h
+++ b/Engine/Source/Runtime/AIModule/Classes/BehaviorTree/BehaviorTreeTypes.h
@@ -134,12 +134,10 @@ namespace EBTDescriptionVerbosity
 	};
 }
 
-/** @TODO might wanna move this to a dedicated header */
 namespace FBlackboard
 {
 	const FName KeySelf = TEXT("SelfActor");
 
-	/** @todo FBlackboard::FKey needs to be used instead of uint8 for keys */
 	typedef uint8 FKey;
 
 	const FKey InvalidKey = FKey(-1);
@@ -370,6 +368,9 @@ protected:
 	/** ID of selected key */
 	UPROPERTY(transient, EditInstanceOnly, Category=Blackboard)
 	uint8 SelectedKeyID;
+	// SelectedKeyId type should be FBlackboard::FKey, but typedefs are not supported by UHT
+	// Please also check key id in FOnBlackboardChange delegate
+	checkAtCompileTime(sizeof(FBlackboard::FKey) == 1, FBlackboardKeySelector_SelectedKeyId_should_be_of_FBlackboard_FKey_type);
 
 	UPROPERTY(transient, EditDefaultsOnly, Category=Blackboard)
 	uint32 bNoneIsAllowedValue:1;
@@ -383,7 +384,7 @@ public:
 
 	void AllowNoneAsValue(bool bNewVal ) { bNoneIsAllowedValue = bNewVal; }
 
-	FORCEINLINE uint8 GetSelectedKeyID() const { return SelectedKeyID; }
+	FORCEINLINE FBlackboard::FKey GetSelectedKeyID() const { return SelectedKeyID; }
 
 	/** helper functions for setting basic filters */
 	void AddObjectFilter(UObject* Owner, TSubclassOf<UObject> AllowedClass);
diff --git a/Engine/Source/Runtime/AIModule/Classes/BehaviorTree/BlackboardComponent.h b/Engine/Source/Runtime/AIModule/Classes/BehaviorTree/BlackboardComponent.h
--- a/Engine/Source/Runtime/AIModule/Classes/BehaviorTree/BlackboardComponent.h
+++ b/Engine/Source/Runtime/AIModule/Classes/BehaviorTree/BlackboardComponent.h
@@ -38,13 +38,13 @@ class AIMODULE_API UBlackboardComponent : public UActorComponent
 	/** END UActorComponent overrides */
 
 	/** @return name of key */
-	FName GetKeyName(uint8 KeyID) const;
+	FName GetKeyName(FBlackboard::FKey KeyID) const;
 
 	/** @return key ID from name */
-	uint8 GetKeyID(const FName& KeyName) const;
+	FBlackboard::FKey GetKeyID(const FName& KeyName) const;
 
 	/** @return class of value for given key */
-	TSubclassOf<UBlackboardKeyType> GetKeyType(uint8 KeyID) const;
+	TSubclassOf<UBlackboardKeyType> GetKeyType(FBlackboard::FKey KeyID) const;
 
 	/** @return number of entries in data asset */
 	int32 GetNumKeys() const;
@@ -53,10 +53,10 @@ class AIMODULE_API UBlackboardComponent : public UActorComponent
 	bool HasValidAsset() const;
 
 	/** register observer for blackboard key */
-	void RegisterObserver(uint8 KeyID, FOnBlackboardChange ObserverDelegate);
+	void RegisterObserver(FBlackboard::FKey KeyID, FOnBlackboardChange ObserverDelegate);
 
 	/** unregister observer from blackboard key */
-	void UnregisterObserver(uint8 KeyID, FOnBlackboardChange ObserverDelegate);
+	void UnregisterObserver(FBlackboard::FKey KeyID, FOnBlackboardChange ObserverDelegate);
 
 	/** pause change notifies and add them to queue */
 	void PauseUpdates();
@@ -81,106 +81,106 @@ class AIMODULE_API UBlackboardComponent : public UActorComponent
 
 	UFUNCTION(BlueprintCallable, Category="AI|Components|Blackboard")
 	UObject* GetValueAsObject(const FName& KeyName) const;
-	UObject* GetValueAsObject(uint8 KeyID) const;
+	UObject* GetValueAsObject(FBlackboard::FKey KeyID) const;
 
 	UFUNCTION(BlueprintCallable, Category="AI|Components|Blackboard")
 	UClass* GetValueAsClass(const FName& KeyName) const;
-	UClass* GetValueAsClass(uint8 KeyID) const;
+	UClass* GetValueAsClass(FBlackboard::FKey KeyID) const;
 
 	UFUNCTION(BlueprintCallable, Category="AI|Components|Blackboard")
 	uint8 GetValueAsEnum(const FName& KeyName) const;
-	uint8 GetValueAsEnum(uint8 KeyID) const;
+	uint8 GetValueAsEnum(FBlackboard::FKey KeyID) const;
 
 	UFUNCTION(BlueprintCallable, Category="AI|Components|Blackboard")
 	int32 GetValueAsInt(const FName& KeyName) const;
-	int32 GetValueAsInt(uint8 KeyID) const;
+	int32 GetValueAsInt(FBlackboard::FKey KeyID) const;
 
 	UFUNCTION(BlueprintCallable, Category="AI|Components|Blackboard")
 	float GetValueAsFloat(const FName& KeyName) const;
-	float GetValueAsFloat(uint8 KeyID) const;
+	float GetValueAsFloat(FBlackboard::FKey KeyID) const;
 
 	UFUNCTION(BlueprintCallable, Category="AI|Components|Blackboard")
 	bool GetValueAsBool(const FName& KeyName) const;
-	bool GetValueAsBool(uint8 KeyID) const;
+	bool GetValueAsBool(FBlackboard::FKey KeyID) const;
 
 	UFUNCTION(BlueprintCallable, Category="AI|Components|Blackboard")
 	FString GetValueAsString(const FName& KeyName) const;
-	FString GetValueAsString(uint8 KeyID) const;
+	FString GetValueAsString(FBlackboard::FKey KeyID) const;
 	
 	UFUNCTION(BlueprintCallable, Category="AI|Components|Blackboard")
 	FName GetValueAsName(const FName& KeyName) const;
-	FName GetValueAsName(uint8 KeyID) const;
+	FName GetValueAsName(FBlackboard::FKey KeyID) const;
 
 	UFUNCTION(BlueprintCallable, Category="AI|Components|Blackboard")
 	FVector GetValueAsVector(const FName& KeyName) const;
-	FVector GetValueAsVector(uint8 KeyID) const;
+	FVector GetValueAsVector(FBlackboard::FKey KeyID) const;
 
 	UFUNCTION(BlueprintCallable, Category="AI|Components|Blackboard")
 	void SetValueAsObject(const FName& KeyName, UObject* ObjectValue);
-	void SetValueAsObject(uint8 KeyID, UObject* ObjectValue);
+	void SetValueAsObject(FBlackboard::FKey KeyID, UObject* ObjectValue);
 
 	UFUNCTION(BlueprintCallable, Category="AI|Components|Blackboard")
 	void SetValueAsClass(const FName& KeyName, UClass* ClassValue);
-	void SetValueAsClass(uint8 KeyID, UClass* ClassValue);
+	void SetValueAsClass(FBlackboard::FKey KeyID, UClass* ClassValue);
 
 	UFUNCTION(BlueprintCallable, Category="AI|Components|Blackboard")
 	void SetValueAsEnum(const FName& KeyName, uint8 EnumValue);
-	void SetValueAsEnum(uint8 KeyID, uint8 EnumValue);
+	void SetValueAsEnum(FBlackboard::FKey KeyID, uint8 EnumValue);
 
 	UFUNCTION(BlueprintCallable, Category="AI|Components|Blackboard")
 	void SetValueAsInt(const FName& KeyName, int32 IntValue);
-	void SetValueAsInt(uint8 KeyID, int32 IntValue);
+	void SetValueAsInt(FBlackboard::FKey KeyID, int32 IntValue);
 
 	UFUNCTION(BlueprintCallable, Category="AI|Components|Blackboard")
 	void SetValueAsFloat(const FName& KeyName, float FloatValue);
-	void SetValueAsFloat(uint8 KeyID, float FloatValue);
+	void SetValueAsFloat(FBlackboard::FKey KeyID, float FloatValue);
 
 	UFUNCTION(BlueprintCallable, Category="AI|Components|Blackboard")
 	void SetValueAsBool(const FName& KeyName, bool BoolValue);
-	void SetValueAsBool(uint8 KeyID, bool BoolValue);
+	void SetValueAsBool(FBlackboard::FKey KeyID, bool BoolValue);
 
 	UFUNCTION(BlueprintCallable, Category="AI|Components|Blackboard")
 	void SetValueAsString(const FName& KeyName, FString StringValue);
-	void SetValueAsString(uint8 KeyID, const FString& StringValue);
+	void SetValueAsString(FBlackboard::FKey KeyID, const FString& StringValue);
 
 	UFUNCTION(BlueprintCallable, Category="AI|Components|Blackboard")
 	void SetValueAsName(const FName& KeyName, FName NameValue);
-	void SetValueAsName(uint8 KeyID, const FName& NameValue);
+	void SetValueAsName(FBlackboard::FKey KeyID, const FName& NameValue);
 
 	UFUNCTION(BlueprintCallable, Category="AI|Components|Blackboard")
 	void SetValueAsVector(const FName& KeyName, FVector VectorValue);
-	void SetValueAsVector(uint8 KeyID, const FVector& VectorValue);
+	void SetValueAsVector(FBlackboard::FKey KeyID, const FVector& VectorValue);
 
 	UFUNCTION(BlueprintCallable, Category="Blackboard")
 	void ClearValueAsVector(const FName& KeyName);
-	void ClearValueAsVector(uint8 KeyID);
+	void ClearValueAsVector(FBlackboard::FKey KeyID);
 
 	/** return false if call failed (most probably no such entry in BB) */
 	UFUNCTION(BlueprintCallable, Category="AI|Components|Blackboard")
 	bool GetLocationFromEntry(const FName& KeyName, FVector& ResultLocation) const;
-	bool GetLocationFromEntry(uint8 KeyID, FVector& ResultLocation) const;
+	bool GetLocationFromEntry(FBlackboard::FKey KeyID, FVector& ResultLocation) const;
 
 	/** return false if call failed (most probably no such entry in BB) */
 	UFUNCTION(BlueprintCallable, Category="Blackboard")
 	bool GetRotationFromEntry(const FName& KeyName, FRotator& ResultRotation) const;
-	bool GetRotationFromEntry(uint8 KeyID, FRotator& ResultRotation) const;
+	bool GetRotationFromEntry(FBlackboard::FKey KeyID, FRotator& ResultRotation) const;
 
 	/** get pointer to raw data for given key */
 	FORCEINLINE uint8* GetKeyRawData(const FName& KeyName) { return GetKeyRawData(GetKeyID(KeyName)); }
-	FORCEINLINE uint8* GetKeyRawData(uint8 KeyID) { return ValueMemory.Num() && ValueOffsets.IsValidIndex(KeyID) ? (ValueMemory.GetTypedData() + ValueOffsets[KeyID]) : NULL; }
+	FORCEINLINE uint8* GetKeyRawData(FBlackboard::FKey KeyID) { return ValueMemory.Num() && ValueOffsets.IsValidIndex(KeyID) ? (ValueMemory.GetTypedData() + ValueOffsets[KeyID]) : NULL; }
 
 	FORCEINLINE const uint8* GetKeyRawData(const FName& KeyName) const { return GetKeyRawData(GetKeyID(KeyName)); }
-	FORCEINLINE const uint8* GetKeyRawData(uint8 KeyID) const { return ValueMemory.Num() && ValueOffsets.IsValidIndex(KeyID) ? (ValueMemory.GetTypedData() + ValueOffsets[KeyID]) : NULL; }
+	FORCEINLINE const uint8* GetKeyRawData(FBlackboard::FKey KeyID) const { return ValueMemory.Num() && ValueOffsets.IsValidIndex(KeyID) ? (ValueMemory.GetTypedData() + ValueOffsets[KeyID]) : NULL; }
 
-	FORCEINLINE bool IsValidKey(uint8 KeyID) const { check(BlackboardAsset); return KeyID != UBlackboardData::InvalidKeyID && BlackboardAsset->Keys.IsValidIndex(KeyID); }
+	FORCEINLINE bool IsValidKey(FBlackboard::FKey KeyID) const { check(BlackboardAsset); return KeyID != FBlackboard::InvalidKey && BlackboardAsset->Keys.IsValidIndex(KeyID); }
 
 	/** compares blackboard's values under specified keys */
-	EBlackboardCompare::Type CompareKeyValues(TSubclassOf<class UBlackboardKeyType> KeyType, uint8 KeyA, uint8 KeyB) const;
+	EBlackboardCompare::Type CompareKeyValues(TSubclassOf<class UBlackboardKeyType> KeyType, FBlackboard::FKey KeyA, FBlackboard::FKey KeyB) const;
 
 	FString GetDebugInfoString(EBlackboardDescription::Type Mode) const;
 
 	/** get description of value under given key */
-	FString DescribeKeyValue(uint8 KeyID, EBlackboardDescription::Type Mode) const;
+	FString DescribeKeyValue(FBlackboard::FKey KeyID, EBlackboardDescription::Type Mode) const;
 
 #if ENABLE_VISUAL_LOG
 	/** prepare blackboard snapshot for logs */
@@ -213,7 +213,7 @@ protected:
 	uint32 bPausedNotifies : 1;
 
 	/** notifies behavior tree decorators about change in blackboard */
-	void NotifyObservers(uint8 KeyID) const;
+	void NotifyObservers(FBlackboard::FKey KeyID) const;
 };
 
 //////////////////////////////////////////////////////////////////////////
diff --git a/Engine/Source/Runtime/AIModule/Classes/BehaviorTree/BlackboardData.h b/Engine/Source/Runtime/AIModule/Classes/BehaviorTree/BlackboardData.h
--- a/Engine/Source/Runtime/AIModule/Classes/BehaviorTree/BlackboardData.h
+++ b/Engine/Source/Runtime/AIModule/Classes/BehaviorTree/BlackboardData.h
@@ -24,11 +24,6 @@ class AIMODULE_API UBlackboardData : public UDataAsset
 	GENERATED_UCLASS_BODY()
 	DECLARE_MULTICAST_DELEGATE_OneParam(FKeyUpdate, UBlackboardData* /*asset*/);
 
-	enum 
-	{
-		InvalidKeyID = 0xFF
-	};
-
 	/** parent blackboard (keys can be overridden) */
 	UPROPERTY(EditAnywhere, Category=Parent)
 	class UBlackboardData* Parent;
@@ -44,21 +39,21 @@ class AIMODULE_API UBlackboardData : public UDataAsset
 	TArray<FBlackboardEntry> Keys;
 
 	/** @return key ID from name */
-	uint8 GetKeyID(const FName& KeyName) const;
+	FBlackboard::FKey GetKeyID(const FName& KeyName) const;
 
 	/** @return name of key */
-	FName GetKeyName(uint8 KeyID) const;
+	FName GetKeyName(FBlackboard::FKey KeyID) const;
 
 	/** @return class of value for given key */
-	TSubclassOf<UBlackboardKeyType> GetKeyType(uint8 KeyID) const;
+	TSubclassOf<UBlackboardKeyType> GetKeyType(FBlackboard::FKey KeyID) const;
 
 	/** @return number of defined keys, including parent chain */
 	int32 GetNumKeys() const;
 
-	FORCEINLINE int32 GetFirstKeyID() const { return FirstKeyID; }
+	FORCEINLINE FBlackboard::FKey GetFirstKeyID() const { return FirstKeyID; }
 
 	/** @return key data */
-	const FBlackboardEntry* GetKey(uint8 KeyID) const;
+	const FBlackboardEntry* GetKey(FBlackboard::FKey KeyID) const;
 
 	virtual void PostLoad() OVERRIDE;
 #if WITH_EDITOR
@@ -76,8 +71,8 @@ class AIMODULE_API UBlackboardData : public UDataAsset
 	{
 		T* CreatedKeyType = NULL;
 
-		const uint8 KeyID = InternalGetKeyID(KeyName, DontCheckParentKeys);
-		if (KeyID == InvalidKeyID && Parent == NULL)
+		const FBlackboard::FKey KeyID = InternalGetKeyID(KeyName, DontCheckParentKeys);
+		if (KeyID == FBlackboard::InvalidKey && Parent == NULL)
 		{
 			FBlackboardEntry Entry;
 			Entry.EntryName = KeyName;
@@ -88,9 +83,9 @@ class AIMODULE_API UBlackboardData : public UDataAsset
 			Keys.Add(Entry);
 			MarkPackageDirty();
 		}
-		else if (KeyID != InvalidKeyID && Parent != NULL)
+		else if (KeyID != FBlackboard::InvalidKey && Parent != NULL)
 		{
-			const int32 KeyIndex = KeyID - FirstKeyID;
+			const FBlackboard::FKey KeyIndex = KeyID - FirstKeyID;
 			Keys.RemoveAt(KeyIndex);
 			MarkPackageDirty();
 		}
@@ -111,10 +106,10 @@ protected:
 	};
 
 	/** @return first ID for keys of this asset (parent keys goes first) */
-	int32 FirstKeyID;
+	FBlackboard::FKey FirstKeyID;
 
 	/** @return key ID from name */
-	uint8 InternalGetKeyID(const FName& KeyName, EKeyLookupMode LookupMode) const;
+	FBlackboard::FKey InternalGetKeyID(const FName& KeyName, EKeyLookupMode LookupMode) const;
 
 	/** updates parent key cache for editor */
 	void UpdateParentKeys();
diff --git a/Engine/Source/Runtime/AIModule/Private/BehaviorTree/BlackboardComponent.cpp b/Engine/Source/Runtime/AIModule/Private/BehaviorTree/BlackboardComponent.cpp
--- a/Engine/Source/Runtime/AIModule/Private/BehaviorTree/BlackboardComponent.cpp
+++ b/Engine/Source/Runtime/AIModule/Private/BehaviorTree/BlackboardComponent.cpp
@@ -36,11 +36,11 @@ void UBlackboardComponent::CacheBrainComponent(UBrainComponent* BrainComponent)
 
 struct FBlackboardInitializationData
 {
-	uint8 KeyID;
+	FBlackboard::FKey KeyID;
 	uint16 DataSize;
 
 	FBlackboardInitializationData() {}
-	FBlackboardInitializationData(uint8 InKeyID, uint16 InDataSize) : KeyID(InKeyID)
+	FBlackboardInitializationData(FBlackboard::FKey InKeyID, uint16 InDataSize) : KeyID(InKeyID)
 	{
 		DataSize = (InDataSize <= 2) ? InDataSize : ((InDataSize + 3) & ~3);
 	}
@@ -126,17 +126,17 @@ class UBlackboardData* UBlackboardComponent::GetBlackboardAsset() const
 	return BlackboardAsset;
 }
 
-FName UBlackboardComponent::GetKeyName(uint8 KeyID) const
+FName UBlackboardComponent::GetKeyName(FBlackboard::FKey KeyID) const
 {
 	return BlackboardAsset ? BlackboardAsset->GetKeyName(KeyID) : NAME_None;
 }
 
-uint8 UBlackboardComponent::GetKeyID(const FName& KeyName) const
+FBlackboard::FKey UBlackboardComponent::GetKeyID(const FName& KeyName) const
 {
-	return BlackboardAsset ? BlackboardAsset->GetKeyID(KeyName) : 0xFF;
+	return BlackboardAsset ? BlackboardAsset->GetKeyID(KeyName) : FBlackboard::InvalidKey;
 }
 
-TSubclassOf<UBlackboardKeyType> UBlackboardComponent::GetKeyType(uint8 KeyID) const
+TSubclassOf<UBlackboardKeyType> UBlackboardComponent::GetKeyType(FBlackboard::FKey KeyID) const
 {
 	return BlackboardAsset ? BlackboardAsset->GetKeyType(KeyID) : NULL;
 }
@@ -146,12 +146,12 @@ int32 UBlackboardComponent::GetNumKeys() const
 	return BlackboardAsset ? BlackboardAsset->GetNumKeys() : 0;
 }
 
-void UBlackboardComponent::RegisterObserver(uint8 KeyID, FOnBlackboardChange ObserverDelegate)
+void UBlackboardComponent::RegisterObserver(FBlackboard::FKey KeyID, FOnBlackboardChange ObserverDelegate)
 {
 	Observers.AddUnique(KeyID, ObserverDelegate);
 }
 
-void UBlackboardComponent::UnregisterObserver(uint8 KeyID, FOnBlackboardChange ObserverDelegate)
+void UBlackboardComponent::UnregisterObserver(FBlackboard::FKey KeyID, FOnBlackboardChange ObserverDelegate)
 {
 	Observers.RemoveSingle(KeyID, ObserverDelegate);
 }
@@ -173,7 +173,7 @@ void UBlackboardComponent::ResumeUpdates()
 	QueuedUpdates.Empty();
 }
 
-void UBlackboardComponent::NotifyObservers(uint8 KeyID) const
+void UBlackboardComponent::NotifyObservers(FBlackboard::FKey KeyID) const
 {
 	if (bPausedNotifies)
 	{
@@ -202,7 +202,7 @@ bool UBlackboardComponent::IsCompatibileWith(class UBlackboardData* TestAsset) c
 	return false;
 }
 
-EBlackboardCompare::Type UBlackboardComponent::CompareKeyValues(TSubclassOf<class UBlackboardKeyType> KeyType, uint8 KeyA, uint8 KeyB) const
+EBlackboardCompare::Type UBlackboardComponent::CompareKeyValues(TSubclassOf<class UBlackboardKeyType> KeyType, FBlackboard::FKey KeyA, FBlackboard::FKey KeyB) const
 {
 	UBlackboardKeyType* KeyCDO = KeyType->GetDefaultObject<UBlackboardKeyType>();
 	return KeyCDO->Compare(GetKeyRawData(KeyA), GetKeyRawData(KeyB));
@@ -242,7 +242,7 @@ FString UBlackboardComponent::GetDebugInfoString(EBlackboardDescription::Type Mo
 		
 			for (int32 KeyIndex = 0; KeyIndex < ObserversKeys.Num(); ++KeyIndex)
 			{
-				const uint8 KeyID = ObserversKeys[KeyIndex];
+				const FBlackboard::FKey KeyID = ObserversKeys[KeyIndex];
 				//@todo shouldn't be using a localized value?; GetKeyName() [10/11/2013 justin.sargent]
 				DebugString += FString::Printf(TEXT("  %s:\n"), *BlackboardAsset->GetKeyName(KeyID).ToString());
 			}
@@ -256,7 +256,7 @@ FString UBlackboardComponent::GetDebugInfoString(EBlackboardDescription::Type Mo
 	return DebugString;
 }
 
-FString UBlackboardComponent::DescribeKeyValue(uint8 KeyID, EBlackboardDescription::Type Mode) const
+FString UBlackboardComponent::DescribeKeyValue(FBlackboard::FKey KeyID, EBlackboardDescription::Type Mode) const
 {
 	FString Description;
 
@@ -315,16 +315,16 @@ void UBlackboardComponent::DescribeSelfToVisLog(struct FVisLogEntry* Snapshot) c
 
 UObject* UBlackboardComponent::GetValueAsObject(const FName& KeyName) const
 {
-	uint8 KeyID = GetKeyID(KeyName);
+	FBlackboard::FKey KeyID = GetKeyID(KeyName);
 	if (GetKeyType(KeyID) != UBlackboardKeyType_Object::StaticClass())
 	{
-		KeyID = INDEX_NONE;
+		KeyID = FBlackboard::InvalidKey;
 	}
 
 	return GetValueAsObject(KeyID);
 }
 
-UObject* UBlackboardComponent::GetValueAsObject(uint8 KeyID) const
+UObject* UBlackboardComponent::GetValueAsObject(FBlackboard::FKey KeyID) const
 {
 	const uint8* RawData = GetKeyRawData(KeyID);
 	return RawData ? UBlackboardKeyType_Object::GetValue(RawData) : NULL;
@@ -332,16 +332,16 @@ UObject* UBlackboardComponent::GetValueAsObject(uint8 KeyID) const
 
 UClass* UBlackboardComponent::GetValueAsClass(const FName& KeyName) const
 {
-	uint8 KeyID = GetKeyID(KeyName);
+	FBlackboard::FKey KeyID = GetKeyID(KeyName);
 	if (GetKeyType(KeyID) != UBlackboardKeyType_Class::StaticClass())
 	{
-		KeyID = INDEX_NONE;
+		KeyID = FBlackboard::InvalidKey;
 	}
 
 	return GetValueAsClass(KeyID);
 }
 
-UClass* UBlackboardComponent::GetValueAsClass(uint8 KeyID) const
+UClass* UBlackboardComponent::GetValueAsClass(FBlackboard::FKey KeyID) const
 {
 	const uint8* RawData = GetKeyRawData(KeyID);
 	return RawData ? UBlackboardKeyType_Class::GetValue(RawData) : NULL;
@@ -349,17 +349,17 @@ UClass* UBlackboardComponent::GetValueAsClass(uint8 KeyID) const
 
 uint8 UBlackboardComponent::GetValueAsEnum(const FName& KeyName) const
 {
-	uint8 KeyID = GetKeyID(KeyName);
+	FBlackboard::FKey KeyID = GetKeyID(KeyName);
 	if (GetKeyType(KeyID) != UBlackboardKeyType_Enum::StaticClass() &&
 		GetKeyType(KeyID) != UBlackboardKeyType_NativeEnum::StaticClass())
 	{
-		KeyID = INDEX_NONE;
+		KeyID = FBlackboard::InvalidKey;
 	}
 
 	return GetValueAsEnum(KeyID);
 }
 
-uint8 UBlackboardComponent::GetValueAsEnum(uint8 KeyID) const
+uint8 UBlackboardComponent::GetValueAsEnum(FBlackboard::FKey KeyID) const
 {
 	const uint8* RawData = GetKeyRawData(KeyID);
 	return RawData ? UBlackboardKeyType_Enum::GetValue(RawData) : 0;
@@ -367,16 +367,16 @@ uint8 UBlackboardComponent::GetValueAsEnum(uint8 KeyID) const
 
 int32 UBlackboardComponent::GetValueAsInt(const FName& KeyName) const
 {
-	uint8 KeyID = GetKeyID(KeyName);
+	FBlackboard::FKey KeyID = GetKeyID(KeyName);
 	if (GetKeyType(KeyID) != UBlackboardKeyType_Int::StaticClass())
 	{
-		KeyID = INDEX_NONE;
+		KeyID = FBlackboard::InvalidKey;
 	}
 
 	return GetValueAsInt(KeyID);
 }
 
-int32 UBlackboardComponent::GetValueAsInt(uint8 KeyID) const
+int32 UBlackboardComponent::GetValueAsInt(FBlackboard::FKey KeyID) const
 {
 	const uint8* RawData = GetKeyRawData(KeyID);
 	return RawData ? UBlackboardKeyType_Int::GetValue(RawData) : 0;
@@ -384,16 +384,16 @@ int32 UBlackboardComponent::GetValueAsInt(uint8 KeyID) const
 
 float UBlackboardComponent::GetValueAsFloat(const FName& KeyName) const
 {
-	uint8 KeyID = GetKeyID(KeyName);
+	FBlackboard::FKey KeyID = GetKeyID(KeyName);
 	if (GetKeyType(KeyID) != UBlackboardKeyType_Float::StaticClass())
 	{
-		KeyID = INDEX_NONE;
+		KeyID = FBlackboard::InvalidKey;
 	}
 
 	return GetValueAsFloat(KeyID);
 }
 
-float UBlackboardComponent::GetValueAsFloat(uint8 KeyID) const
+float UBlackboardComponent::GetValueAsFloat(FBlackboard::FKey KeyID) const
 {
 	const uint8* RawData = GetKeyRawData(KeyID);
 	return RawData ? UBlackboardKeyType_Float::GetValue(RawData) : 0.0f;
@@ -401,16 +401,16 @@ float UBlackboardComponent::GetValueAsFloat(uint8 KeyID) const
 
 bool UBlackboardComponent::GetValueAsBool(const FName& KeyName) const
 {
-	uint8 KeyID = GetKeyID(KeyName);
+	FBlackboard::FKey KeyID = GetKeyID(KeyName);
 	if (GetKeyType(KeyID) != UBlackboardKeyType_Bool::StaticClass())
 	{
-		KeyID = INDEX_NONE;
+		KeyID = FBlackboard::InvalidKey;
 	}
 
 	return GetValueAsBool(KeyID);
 }
 
-bool UBlackboardComponent::GetValueAsBool(uint8 KeyID) const
+bool UBlackboardComponent::GetValueAsBool(FBlackboard::FKey KeyID) const
 {
 	const uint8* RawData = GetKeyRawData(KeyID);
 	return RawData ? UBlackboardKeyType_Bool::GetValue(RawData) : false;
@@ -418,16 +418,16 @@ bool UBlackboardComponent::GetValueAsBool(uint8 KeyID) const
 
 FString UBlackboardComponent::GetValueAsString(const FName& KeyName) const
 {
-	uint8 KeyID = GetKeyID(KeyName);
+	FBlackboard::FKey KeyID = GetKeyID(KeyName);
 	if (GetKeyType(KeyID) != UBlackboardKeyType_String::StaticClass())
 	{
-		KeyID = INDEX_NONE;
+		KeyID = FBlackboard::InvalidKey;
 	}
 
 	return GetValueAsString(KeyID);
 }
 
-FString UBlackboardComponent::GetValueAsString(uint8 KeyID) const
+FString UBlackboardComponent::GetValueAsString(FBlackboard::FKey KeyID) const
 {
 	const uint8* RawData = GetKeyRawData(KeyID);
 	return RawData ? UBlackboardKeyType_String::GetValue(RawData) : FString();
@@ -435,16 +435,16 @@ FString UBlackboardComponent::GetValueAsString(uint8 KeyID) const
 
 FName UBlackboardComponent::GetValueAsName(const FName& KeyName) const
 {
-	uint8 KeyID = GetKeyID(KeyName);
+	FBlackboard::FKey KeyID = GetKeyID(KeyName);
 	if (GetKeyType(KeyID) != UBlackboardKeyType_Name::StaticClass())
 	{
-		KeyID = INDEX_NONE;
+		KeyID = FBlackboard::InvalidKey;
 	}
 
 	return GetValueAsName(KeyID);
 }
 
-FName UBlackboardComponent::GetValueAsName(uint8 KeyID) const
+FName UBlackboardComponent::GetValueAsName(FBlackboard::FKey KeyID) const
 {
 	const uint8* RawData = GetKeyRawData(KeyID);
 	return RawData ? UBlackboardKeyType_Name::GetValue(RawData) : NAME_None;
@@ -452,16 +452,16 @@ FName UBlackboardComponent::GetValueAsName(uint8 KeyID) const
 
 FVector UBlackboardComponent::GetValueAsVector(const FName& KeyName) const
 {
-	uint8 KeyID = GetKeyID(KeyName);
+	FBlackboard::FKey KeyID = GetKeyID(KeyName);
 	if (GetKeyType(KeyID) != UBlackboardKeyType_Vector::StaticClass())
 	{
-		KeyID = INDEX_NONE;
+		KeyID = FBlackboard::InvalidKey;
 	}
 
 	return GetValueAsVector(KeyID);
 }
 
-FVector UBlackboardComponent::GetValueAsVector(uint8 KeyID) const
+FVector UBlackboardComponent::GetValueAsVector(FBlackboard::FKey KeyID) const
 {
 	const uint8* RawData = GetKeyRawData(KeyID);
 	return RawData ? UBlackboardKeyType_Vector::GetValue(RawData) : FVector::ZeroVector;
@@ -469,14 +469,14 @@ FVector UBlackboardComponent::GetValueAsVector(uint8 KeyID) const
 
 void UBlackboardComponent::SetValueAsObject(const FName& KeyName, UObject* ObjectValue)
 {
-	const uint8 KeyID = GetKeyID(KeyName);
+	const FBlackboard::FKey KeyID = GetKeyID(KeyName);
 	if (GetKeyType(KeyID) == UBlackboardKeyType_Object::StaticClass())
 	{
 		SetValueAsObject(KeyID, ObjectValue);
 	}
 }
 
-void UBlackboardComponent::SetValueAsObject(uint8 KeyID, UObject* ObjectValue)
+void UBlackboardComponent::SetValueAsObject(FBlackboard::FKey KeyID, UObject* ObjectValue)
 {
 	uint8* RawData = GetKeyRawData(KeyID);
 	if (RawData)
@@ -491,14 +491,14 @@ void UBlackboardComponent::SetValueAsObject(uint8 KeyID, UObject* ObjectValue)
 
 void UBlackboardComponent::SetValueAsClass(const FName& KeyName, UClass* ClassValue)
 {
-	const uint8 KeyID = GetKeyID(KeyName);
+	const FBlackboard::FKey KeyID = GetKeyID(KeyName);
 	if (GetKeyType(KeyID) == UBlackboardKeyType_Class::StaticClass())
 	{
 		SetValueAsClass(KeyID, ClassValue);
 	}
 }
 
-void UBlackboardComponent::SetValueAsClass(uint8 KeyID, UClass* ClassValue)
+void UBlackboardComponent::SetValueAsClass(FBlackboard::FKey KeyID, UClass* ClassValue)
 {
 	uint8* RawData = GetKeyRawData(KeyID);
 	if (RawData)
@@ -513,7 +513,7 @@ void UBlackboardComponent::SetValueAsClass(uint8 KeyID, UClass* ClassValue)
 
 void UBlackboardComponent::SetValueAsEnum(const FName& KeyName, uint8 EnumValue)
 {
-	const uint8 KeyID = GetKeyID(KeyName);
+	const FBlackboard::FKey KeyID = GetKeyID(KeyName);
 	if (GetKeyType(KeyID) == UBlackboardKeyType_Enum::StaticClass() ||
 		GetKeyType(KeyID) == UBlackboardKeyType_NativeEnum::StaticClass())
 	{
@@ -521,7 +521,7 @@ void UBlackboardComponent::SetValueAsEnum(const FName& KeyName, uint8 EnumValue)
 	}
 }
 
-void UBlackboardComponent::SetValueAsEnum(uint8 KeyID, uint8 EnumValue)
+void UBlackboardComponent::SetValueAsEnum(FBlackboard::FKey KeyID, uint8 EnumValue)
 {
 	uint8* RawData = GetKeyRawData(KeyID);
 	if (RawData)
@@ -536,14 +536,14 @@ void UBlackboardComponent::SetValueAsEnum(uint8 KeyID, uint8 EnumValue)
 
 void UBlackboardComponent::SetValueAsInt(const FName& KeyName, int32 IntValue)
 {
-	const uint8 KeyID = GetKeyID(KeyName);
+	const FBlackboard::FKey KeyID = GetKeyID(KeyName);
 	if (GetKeyType(KeyID) == UBlackboardKeyType_Int::StaticClass())
 	{
 		SetValueAsInt(KeyID, IntValue);
 	}
 }
 
-void UBlackboardComponent::SetValueAsInt(uint8 KeyID, int32 IntValue)
+void UBlackboardComponent::SetValueAsInt(FBlackboard::FKey KeyID, int32 IntValue)
 {
 	uint8* RawData = GetKeyRawData(KeyID);
 	if (RawData)
@@ -558,14 +558,14 @@ void UBlackboardComponent::SetValueAsInt(uint8 KeyID, int32 IntValue)
 
 void UBlackboardComponent::SetValueAsFloat(const FName& KeyName, float FloatValue)
 {
-	const uint8 KeyID = GetKeyID(KeyName);
+	const FBlackboard::FKey KeyID = GetKeyID(KeyName);
 	if (GetKeyType(KeyID) == UBlackboardKeyType_Float::StaticClass())
 	{
 		SetValueAsFloat(KeyID, FloatValue);
 	}
 }
 
-void UBlackboardComponent::SetValueAsFloat(uint8 KeyID, float FloatValue)
+void UBlackboardComponent::SetValueAsFloat(FBlackboard::FKey KeyID, float FloatValue)
 {
 	uint8* RawData = GetKeyRawData(KeyID);
 	if (RawData)
@@ -580,14 +580,14 @@ void UBlackboardComponent::SetValueAsFloat(uint8 KeyID, float FloatValue)
 
 void UBlackboardComponent::SetValueAsBool(const FName& KeyName, bool BoolValue)
 {
-	const uint8 KeyID = GetKeyID(KeyName);
+	const FBlackboard::FKey KeyID = GetKeyID(KeyName);
 	if (GetKeyType(KeyID) == UBlackboardKeyType_Bool::StaticClass())
 	{
 		SetValueAsBool(KeyID, BoolValue);
 	}
 }
 
-void UBlackboardComponent::SetValueAsBool(uint8 KeyID, bool BoolValue)
+void UBlackboardComponent::SetValueAsBool(FBlackboard::FKey KeyID, bool BoolValue)
 {
 	uint8* RawData = GetKeyRawData(KeyID);
 	if (RawData)
@@ -602,14 +602,14 @@ void UBlackboardComponent::SetValueAsBool(uint8 KeyID, bool BoolValue)
 
 void UBlackboardComponent::SetValueAsString(const FName& KeyName, FString StringValue)
 {
-	const uint8 KeyID = GetKeyID(KeyName);
+	const FBlackboard::FKey KeyID = GetKeyID(KeyName);
 	if (GetKeyType(KeyID) == UBlackboardKeyType_String::StaticClass())
 	{
 		SetValueAsString(KeyID, StringValue);
 	}
 }
 
-void UBlackboardComponent::SetValueAsString(uint8 KeyID, const FString& StringValue)
+void UBlackboardComponent::SetValueAsString(FBlackboard::FKey KeyID, const FString& StringValue)
 {
 	uint8* RawData = GetKeyRawData(KeyID);
 	if (RawData)
@@ -624,14 +624,14 @@ void UBlackboardComponent::SetValueAsString(uint8 KeyID, const FString& StringVa
 
 void UBlackboardComponent::SetValueAsName(const FName& KeyName, FName NameValue)
 {
-	const uint8 KeyID = GetKeyID(KeyName);
+	const FBlackboard::FKey KeyID = GetKeyID(KeyName);
 	if (GetKeyType(KeyID) == UBlackboardKeyType_Name::StaticClass())
 	{
 		SetValueAsName(KeyID, NameValue);
 	}
 }
 
-void UBlackboardComponent::SetValueAsName(uint8 KeyID, const FName& NameValue)
+void UBlackboardComponent::SetValueAsName(FBlackboard::FKey KeyID, const FName& NameValue)
 {
 	uint8* RawData = GetKeyRawData(KeyID);
 	if (RawData)
@@ -646,14 +646,14 @@ void UBlackboardComponent::SetValueAsName(uint8 KeyID, const FName& NameValue)
 
 void UBlackboardComponent::SetValueAsVector(const FName& KeyName, FVector VectorValue)
 {
-	const uint8 KeyID = GetKeyID(KeyName);
+	const FBlackboard::FKey KeyID = GetKeyID(KeyName);
 	if (GetKeyType(KeyID) == UBlackboardKeyType_Vector::StaticClass())
 	{
 		SetValueAsVector(KeyID, VectorValue);
 	}
 }
 
-void UBlackboardComponent::SetValueAsVector(uint8 KeyID, const FVector& VectorValue)
+void UBlackboardComponent::SetValueAsVector(FBlackboard::FKey KeyID, const FVector& VectorValue)
 {
 	uint8* RawData = GetKeyRawData(KeyID);
 	if (RawData)
@@ -668,14 +668,14 @@ void UBlackboardComponent::SetValueAsVector(uint8 KeyID, const FVector& VectorVa
 
 void UBlackboardComponent::ClearValueAsVector(const FName& KeyName)
 {
-	const uint8 KeyID = GetKeyID(KeyName);
+	const FBlackboard::FKey KeyID = GetKeyID(KeyName);
 	if (GetKeyType(KeyID) == UBlackboardKeyType_Vector::StaticClass())
 	{
 		ClearValueAsVector(KeyID);
 	}
 }
 
-void UBlackboardComponent::ClearValueAsVector(uint8 KeyID)
+void UBlackboardComponent::ClearValueAsVector(FBlackboard::FKey KeyID)
 {
 	uint8* RawData = GetKeyRawData(KeyID);
 	if (RawData)
@@ -690,11 +690,11 @@ void UBlackboardComponent::ClearValueAsVector(uint8 KeyID)
 
 bool UBlackboardComponent::GetLocationFromEntry(const FName& KeyName, FVector& ResultLocation) const
 {
-	const uint8 KeyID = GetKeyID(KeyName);
+	const FBlackboard::FKey KeyID = GetKeyID(KeyName);
 	return GetLocationFromEntry(KeyID, ResultLocation);
 }
 
-bool UBlackboardComponent::GetLocationFromEntry(uint8 KeyID, FVector& ResultLocation) const
+bool UBlackboardComponent::GetLocationFromEntry(FBlackboard::FKey KeyID, FVector& ResultLocation) const
 {
 	if (BlackboardAsset && ValueOffsets.IsValidIndex(KeyID))
 	{
@@ -711,11 +711,11 @@ bool UBlackboardComponent::GetLocationFromEntry(uint8 KeyID, FVector& ResultLoca
 
 bool UBlackboardComponent::GetRotationFromEntry(const FName& KeyName, FRotator& ResultRotation) const
 {
-	const uint8 KeyID = GetKeyID(KeyName);
+	const FBlackboard::FKey KeyID = GetKeyID(KeyName);
 	return GetRotationFromEntry(KeyID, ResultRotation);
 }
 
-bool UBlackboardComponent::GetRotationFromEntry(uint8 KeyID, FRotator& ResultRotation) const
+bool UBlackboardComponent::GetRotationFromEntry(FBlackboard::FKey KeyID, FRotator& ResultRotation) const
 {
 	if (BlackboardAsset && ValueOffsets.IsValidIndex(KeyID))
 	{
diff --git a/Engine/Source/Runtime/AIModule/Private/BehaviorTree/BlackboardData.cpp b/Engine/Source/Runtime/AIModule/Private/BehaviorTree/BlackboardData.cpp
--- a/Engine/Source/Runtime/AIModule/Private/BehaviorTree/BlackboardData.cpp
+++ b/Engine/Source/Runtime/AIModule/Private/BehaviorTree/BlackboardData.cpp
@@ -23,26 +23,26 @@ UBlackboardData::UBlackboardData(const class FPostConstructInitializeProperties&
 	}
 }
 
-uint8 UBlackboardData::GetKeyID(const FName& KeyName) const
+FBlackboard::FKey UBlackboardData::GetKeyID(const FName& KeyName) const
 {
 	return InternalGetKeyID(KeyName, CheckParentKeys);
 }
 
-FName UBlackboardData::GetKeyName(uint8 KeyID) const
+FName UBlackboardData::GetKeyName(FBlackboard::FKey KeyID) const
 {
 	const FBlackboardEntry* KeyEntry = GetKey(KeyID);
 	return KeyEntry ? KeyEntry->EntryName : NAME_None;
 }
 
-TSubclassOf<UBlackboardKeyType> UBlackboardData::GetKeyType(uint8 KeyID) const
+TSubclassOf<UBlackboardKeyType> UBlackboardData::GetKeyType(FBlackboard::FKey KeyID) const
 {
 	const FBlackboardEntry* KeyEntry = GetKey(KeyID);
 	return KeyEntry ? KeyEntry->KeyType->GetClass() : NULL;
 }
 
-const FBlackboardEntry* UBlackboardData::GetKey(uint8 KeyID) const
+const FBlackboardEntry* UBlackboardData::GetKey(FBlackboard::FKey KeyID) const
 {
-	if (KeyID != InvalidKeyID)
+	if (KeyID != FBlackboard::InvalidKey)
 	{
 		if (KeyID >= FirstKeyID)
 		{
@@ -62,7 +62,7 @@ int32 UBlackboardData::GetNumKeys() const
 	return FirstKeyID + Keys.Num();
 }
 
-uint8 UBlackboardData::InternalGetKeyID(const FName& KeyName, EKeyLookupMode LookupMode) const
+FBlackboard::FKey UBlackboardData::InternalGetKeyID(const FName& KeyName, EKeyLookupMode LookupMode) const
 {
 	for (int32 i = 0; i < Keys.Num(); i++)
 	{
@@ -72,7 +72,7 @@ uint8 UBlackboardData::InternalGetKeyID(const FName& KeyName, EKeyLookupMode Loo
 		}
 	}
 
-	return Parent && (LookupMode == CheckParentKeys) ? Parent->InternalGetKeyID(KeyName, LookupMode) : InvalidKeyID;
+	return Parent && (LookupMode == CheckParentKeys) ? Parent->InternalGetKeyID(KeyName, LookupMode) : FBlackboard::InvalidKey;
 }
 
 bool UBlackboardData::IsValid() const
@@ -81,8 +81,8 @@ bool UBlackboardData::IsValid() const
 	{
 		for (int32 i = 0; i < Keys.Num(); i++)
 		{
-			const uint8 KeyID = Parent->InternalGetKeyID(Keys[i].EntryName, CheckParentKeys);
-			if (KeyID != InvalidKeyID)
+			const FBlackboard::FKey KeyID = Parent->InternalGetKeyID(Keys[i].EntryName, CheckParentKeys);
+			if (KeyID != FBlackboard::InvalidKey)
 			{
 				UE_LOG(LogBehaviorTree, Warning, TEXT("Blackboard asset (%s) has duplicated key (%s) in parent chain!"),
 					*GetName(), *Keys[i].EntryName.ToString());
diff --git a/Engine/Source/Runtime/AIModule/Private/BehaviorTree/Decorators/BTDecorator_BlueprintBase.cpp b/Engine/Source/Runtime/AIModule/Private/BehaviorTree/Decorators/BTDecorator_BlueprintBase.cpp
--- a/Engine/Source/Runtime/AIModule/Private/BehaviorTree/Decorators/BTDecorator_BlueprintBase.cpp
+++ b/Engine/Source/Runtime/AIModule/Private/BehaviorTree/Decorators/BTDecorator_BlueprintBase.cpp
@@ -73,8 +73,8 @@ void UBTDecorator_BlueprintBase::OnBecomeRelevant(UBehaviorTreeComponent* OwnerC
 	{
 		for (int32 i = 0; i < CDO->ObservedKeyNames.Num(); i++)
 		{
-			const uint8 KeyID = BlackboardComp->GetKeyID(CDO->ObservedKeyNames[i]);
-			if (KeyID != UBlackboardData::InvalidKeyID)
+			const FBlackboard::FKey KeyID = BlackboardComp->GetKeyID(CDO->ObservedKeyNames[i]);
+			if (KeyID != FBlackboard::InvalidKey)
 			{
 				BlackboardComp->RegisterObserver(KeyID, BBKeyObserver);
 			}
@@ -90,8 +90,8 @@ void UBTDecorator_BlueprintBase::OnCeaseRelevant(UBehaviorTreeComponent* OwnerCo
 	{
 		for (int32 i = 0; i < CDO->ObservedKeyNames.Num(); i++)
 		{
-			const uint8 KeyID = BlackboardComp->GetKeyID(CDO->ObservedKeyNames[i]);
-			if (KeyID != UBlackboardData::InvalidKeyID)
+			const FBlackboard::FKey KeyID = BlackboardComp->GetKeyID(CDO->ObservedKeyNames[i]);
+			if (KeyID != FBlackboard::InvalidKey)
 			{
 				BlackboardComp->UnregisterObserver(KeyID, BBKeyObserver);
 			}
