fixed crash in CompareBBEntries decorator

diff --git a/Engine/Source/Runtime/AIModule/Private/BehaviorTree/Decorators/BTDecorator_CompareBBEntries.cpp b/Engine/Source/Runtime/AIModule/Private/BehaviorTree/Decorators/BTDecorator_CompareBBEntries.cpp
--- a/Engine/Source/Runtime/AIModule/Private/BehaviorTree/Decorators/BTDecorator_CompareBBEntries.cpp
+++ b/Engine/Source/Runtime/AIModule/Private/BehaviorTree/Decorators/BTDecorator_CompareBBEntries.cpp
@@ -4,12 +4,11 @@
 #include "BehaviorTree/BlackboardComponent.h"
 #include "BehaviorTree/Decorators/BTDecorator_CompareBBEntries.h"
 
-UBTDecorator_CompareBBEntries::UBTDecorator_CompareBBEntries(const class FPostConstructInitializeProperties& PCIP) 
-	: Super(PCIP)
-	, Operator(EArithmeticKeyOperation::LessOrEqual)
+UBTDecorator_CompareBBEntries::UBTDecorator_CompareBBEntries(const class FPostConstructInitializeProperties& PCIP) : Super(PCIP)
 {
 	NodeName = "Compare Blackboard entries";
 
+	Operator = EBlackBoardEntryComparison::Equal;
 	bNotifyBecomeRelevant = true;
 	bNotifyCeaseRelevant = true;
 }
