Implemented a cheat to dump information about BT nodes usage to log #UE4

diff --git a/Engine/Source/Editor/UnrealEd/Public/Kismet2/KismetEditorUtilities.h b/Engine/Source/Editor/UnrealEd/Public/Kismet2/KismetEditorUtilities.h
--- a/Engine/Source/Editor/UnrealEd/Public/Kismet2/KismetEditorUtilities.h
+++ b/Engine/Source/Editor/UnrealEd/Public/Kismet2/KismetEditorUtilities.h
@@ -6,6 +6,8 @@
 
 #pragma once
 
+class UBlueprintGeneratedClass;
+
 //////////////////////////////////////////////////////////////////////////
 // FKismetEditorUtilities
 
diff --git a/Engine/Source/Runtime/AIModule/Classes/BehaviorTree/BehaviorTreeManager.h b/Engine/Source/Runtime/AIModule/Classes/BehaviorTree/BehaviorTreeManager.h
--- a/Engine/Source/Runtime/AIModule/Classes/BehaviorTree/BehaviorTreeManager.h
+++ b/Engine/Source/Runtime/AIModule/Classes/BehaviorTree/BehaviorTreeManager.h
@@ -44,6 +44,8 @@ class AIMODULE_API UBehaviorTreeManager : public UObject
 	/** cleanup hooks for map loading */
 	virtual void FinishDestroy() OVERRIDE;
 
+	void DumpUsageStats() const;
+
 protected:
 
 	/** initialized tree templates */
diff --git a/Engine/Source/Runtime/AIModule/Private/AISystemExec.cpp b/Engine/Source/Runtime/AIModule/Private/AISystemExec.cpp
--- a/Engine/Source/Runtime/AIModule/Private/AISystemExec.cpp
+++ b/Engine/Source/Runtime/AIModule/Private/AISystemExec.cpp
@@ -3,6 +3,7 @@
 #include "AIModulePrivate.h"
 #if !(UE_BUILD_SHIPPING || UE_BUILD_TEST)
 #include "Misc/CoreMisc.h"
+#include "BehaviorTree/BehaviorTreeManager.h"
 
 struct FAISystemExec : public FSelfRegisteringExec
 {
@@ -39,6 +40,14 @@ bool FAISystemExec::Exec(UWorld* Inworld, const TCHAR* Cmd, FOutputDevice& Ar)
 			AISys->AILoggingVerbose();
 			bHandled = true;
 		}
+		else if (FParse::Command(&Cmd, TEXT("DumpBTUsageStats")))
+		{
+			if (AISys->GetBehaviorTreeManager())
+			{
+				AISys->GetBehaviorTreeManager()->DumpUsageStats();
+				bHandled = true;
+			}
+		}
 
 		// void UAISystem::RunEQS(const FString& QueryName)
 	}
diff --git a/Engine/Source/Runtime/AIModule/Private/BehaviorTree/BehaviorTreeManager.cpp b/Engine/Source/Runtime/AIModule/Private/BehaviorTree/BehaviorTreeManager.cpp
--- a/Engine/Source/Runtime/AIModule/Private/BehaviorTree/BehaviorTreeManager.cpp
+++ b/Engine/Source/Runtime/AIModule/Private/BehaviorTree/BehaviorTreeManager.cpp
@@ -4,6 +4,9 @@
 #include "BehaviorTree/Tasks/BTTask_RunBehavior.h"
 #include "BehaviorTree/BlackboardComponent.h"
 #include "BehaviorTree/BehaviorTreeManager.h"
+#if WITH_EDITOR
+#include "Kismet2/KismetEditorUtilities.h"
+#endif // WITH_EDITOR
 
 DEFINE_STAT(STAT_AI_BehaviorTree_Tick);
 DEFINE_STAT(STAT_AI_BehaviorTree_LoadTime);
@@ -254,3 +257,126 @@ void UBehaviorTreeManager::InitializeMemoryHelper(const TArray<UBTDecorator*>& N
 
 	MemorySize = MemoryOffset;
 }
+
+//----------------------------------------------------------------------//
+// stats dumping
+//----------------------------------------------------------------------//
+struct FNodeClassCounter
+{
+	TMap<UClass*, uint32> NodeClassUsage;
+
+	FNodeClassCounter()
+	{}
+
+	void Declare(UClass* NodeClass)
+	{
+		NodeClassUsage.FindOrAdd(NodeClass);
+	}
+
+	void CountNode(UBTNode* Node)
+	{
+		uint32& Count = NodeClassUsage.FindOrAdd(Node->GetClass());
+		++Count;
+	}
+
+	void Append(const FNodeClassCounter& Other)
+	{
+		for (auto Iterator : Other.NodeClassUsage)
+		{
+			uint32& Count = NodeClassUsage.FindOrAdd(Iterator.Key);
+			Count += Iterator.Value;
+		}
+	}
+
+	void Print(TCHAR* Separator=TEXT(" "))
+	{
+		for (auto Iterator : NodeClassUsage)
+		{
+			UE_LOG(LogBehaviorTree, Display, TEXT("%s%s%s(%s)%s%d")
+				, Separator
+				, *Iterator.Key->GetName()
+				, Separator
+				, Iterator.Key->HasAnyClassFlags(CLASS_CompiledFromBlueprint) ? TEXT("BP") : TEXT("C++")
+				, Separator
+				, Iterator.Value);
+		}
+	}
+};
+
+void StatNodeUsage(UBTNode* Node, FNodeClassCounter& NodeCounter)
+{
+	NodeCounter.CountNode(Node);
+
+	UBTCompositeNode* CompositeOb = Cast<UBTCompositeNode>(Node);
+	if (CompositeOb)
+	{
+		for (int32 ServiceIndex = 0; ServiceIndex < CompositeOb->Services.Num(); ServiceIndex++)
+		{
+			if (CompositeOb->Services[ServiceIndex] == NULL)
+			{
+				continue;
+			}
+			NodeCounter.CountNode(CompositeOb->Services[ServiceIndex]);
+		}
+
+		for (int32 ChildIndex = 0; ChildIndex < CompositeOb->Children.Num(); ChildIndex++)
+		{
+			FBTCompositeChild& ChildInfo = CompositeOb->Children[ChildIndex];
+			for (int32 DecoratorIndex = 0; DecoratorIndex < ChildInfo.Decorators.Num(); DecoratorIndex++)
+			{
+				if (ChildInfo.Decorators[DecoratorIndex] == NULL)
+				{
+					continue;
+				}
+
+				NodeCounter.CountNode(ChildInfo.Decorators[DecoratorIndex]);
+			}
+
+			UBTNode* ChildNode = NULL;
+
+			if (ChildInfo.ChildComposite)
+			{
+				StatNodeUsage(ChildInfo.ChildComposite, NodeCounter);
+			}
+			else if (ChildInfo.ChildTask)
+			{
+				NodeCounter.CountNode(ChildInfo.ChildTask);
+			}
+		}
+	}
+}
+
+void UBehaviorTreeManager::DumpUsageStats() const
+{
+	FNodeClassCounter AllNodesCounter;
+	
+	for (TObjectIterator<UClass> It; It; ++It)
+	{
+		if (It->IsChildOf(UBTNode::StaticClass()) && It->HasAnyClassFlags(CLASS_Abstract) == false
+#if WITH_EDITOR
+			&& !(FKismetEditorUtilities::IsClassABlueprintSkeleton(*It)
+				|| It->HasAnyClassFlags(CLASS_NewerVersionExists))
+#endif
+			)
+
+		{
+			AllNodesCounter.Declare(*It);
+		}
+	}
+
+	UE_LOG(LogBehaviorTree, Display, TEXT("----------------------UBehaviorTreeManager::DumpUsageStats----------------------\nBehavior Trees:"));
+
+	// get all BTNode classes
+	
+	for (TObjectIterator<UBehaviorTree> It; It; ++It)
+	{
+		FNodeClassCounter TreeNodeCounter;
+		UE_LOG(LogBehaviorTree, Display, TEXT("--- %s ---"), *(It->GetName()));
+		StatNodeUsage(It->RootNode, TreeNodeCounter);
+		TreeNodeCounter.Print();
+		AllNodesCounter.Append(TreeNodeCounter);
+	}
+	
+	UE_LOG(LogBehaviorTree, Display, TEXT("--- Total Nodes class usage:"));
+	AllNodesCounter.Print(TEXT(","));
+}
\ No newline at end of file
