Making sure Default Subobjects are not created using NewObject (which produces inconsistent names)

diff --git a/Engine/Source/Runtime/AIModule/Private/BehaviorTree/BehaviorTreeTypes.cpp b/Engine/Source/Runtime/AIModule/Private/BehaviorTree/BehaviorTreeTypes.cpp
--- a/Engine/Source/Runtime/AIModule/Private/BehaviorTree/BehaviorTreeTypes.cpp
+++ b/Engine/Source/Runtime/AIModule/Private/BehaviorTree/BehaviorTreeTypes.cpp
@@ -193,60 +193,60 @@ void FBlackboardKeySelector::InitSelectedKey(class UBlackboardData* BlackboardAs
 
 void FBlackboardKeySelector::AddObjectFilter(UObject* Owner, TSubclassOf<UObject> AllowedClass)
 {
-	UBlackboardKeyType_Object* FilterOb = NewObject<UBlackboardKeyType_Object>(Owner);
+	UBlackboardKeyType_Object* FilterOb = NewNamedObject<UBlackboardKeyType_Object>(Owner, TEXT("BlackboardKeyType_Object"));
 	FilterOb->BaseClass = AllowedClass;
 	AllowedTypes.Add(FilterOb);
 }
 
 void FBlackboardKeySelector::AddClassFilter(UObject* Owner, TSubclassOf<UClass> AllowedClass)
 {
-	UBlackboardKeyType_Class* FilterOb = NewObject<UBlackboardKeyType_Class>(Owner);
+	UBlackboardKeyType_Class* FilterOb = NewNamedObject<UBlackboardKeyType_Class>(Owner, TEXT("BlackboardKeyType_Class"));
 	FilterOb->BaseClass = AllowedClass;
 	AllowedTypes.Add(FilterOb);
 }
 
 void FBlackboardKeySelector::AddEnumFilter(UObject* Owner, UEnum* AllowedEnum)
 {
-	UBlackboardKeyType_Enum* FilterOb = NewObject<UBlackboardKeyType_Enum>(Owner);
+	UBlackboardKeyType_Enum* FilterOb = NewNamedObject<UBlackboardKeyType_Enum>(Owner, TEXT("BlackboardKeyType_Enum"));
 	FilterOb->EnumType = AllowedEnum;
 	AllowedTypes.Add(FilterOb);
 }
 
 void FBlackboardKeySelector::AddNativeEnumFilter(UObject* Owner, const FString& AllowedEnumName)
 {
-	UBlackboardKeyType_NativeEnum* FilterOb = NewObject<UBlackboardKeyType_NativeEnum>(Owner);
+	UBlackboardKeyType_NativeEnum* FilterOb = NewNamedObject<UBlackboardKeyType_NativeEnum>(Owner, TEXT("BlackboardKeyType_NativeEnum"));
 	FilterOb->EnumName = AllowedEnumName;
 	AllowedTypes.Add(FilterOb);
 }
 
 void FBlackboardKeySelector::AddIntFilter(UObject* Owner)
 {
-	AllowedTypes.Add(NewObject<UBlackboardKeyType_Int>(Owner));
+	AllowedTypes.Add(NewNamedObject<UBlackboardKeyType_Int>(Owner, TEXT("BlackboardKeyType_Int")));
 }
 
 void FBlackboardKeySelector::AddFloatFilter(UObject* Owner)
 {
-	AllowedTypes.Add(NewObject<UBlackboardKeyType_Float>(Owner));
+	AllowedTypes.Add(NewNamedObject<UBlackboardKeyType_Float>(Owner, TEXT("BlackboardKeyType_Float")));
 }
 
 void FBlackboardKeySelector::AddBoolFilter(UObject* Owner)
 {
-	AllowedTypes.Add(NewObject<UBlackboardKeyType_Bool>(Owner));
+	AllowedTypes.Add(NewNamedObject<UBlackboardKeyType_Bool>(Owner, TEXT("BlackboardKeyType_Bool")));
 }
 
 void FBlackboardKeySelector::AddVectorFilter(UObject* Owner)
 {
-	AllowedTypes.Add(NewObject<UBlackboardKeyType_Vector>(Owner));
+	AllowedTypes.Add(NewNamedObject<UBlackboardKeyType_Vector>(Owner, TEXT("BlackboardKeyType_Vector")));
 }
 
 void FBlackboardKeySelector::AddStringFilter(UObject* Owner)
 {
-	AllowedTypes.Add(NewObject<UBlackboardKeyType_String>(Owner));
+	AllowedTypes.Add(NewNamedObject<UBlackboardKeyType_String>(Owner, TEXT("BlackboardKeyType_String")));
 }
 
 void FBlackboardKeySelector::AddNameFilter(UObject* Owner)
 {
-	AllowedTypes.Add(NewObject<UBlackboardKeyType_Name>(Owner));
+	AllowedTypes.Add(NewNamedObject<UBlackboardKeyType_Name>(Owner, TEXT("BlackboardKeyType_Name")));
 }
 
 //----------------------------------------------------------------------//
diff --git a/Engine/Source/Runtime/CoreUObject/Private/UObject/UObjectGlobals.cpp b/Engine/Source/Runtime/CoreUObject/Private/UObject/UObjectGlobals.cpp
--- a/Engine/Source/Runtime/CoreUObject/Private/UObject/UObjectGlobals.cpp
+++ b/Engine/Source/Runtime/CoreUObject/Private/UObject/UObjectGlobals.cpp
@@ -1850,13 +1850,16 @@ UObject::UObject(const class FPostConstructInitializeProperties& PCIP)
 
 /* Global flag so that FObjectFinders know if they are called from inside the UObject constructors or not. */
 static int32 GIsInConstructor = 0;
+/* Object that is currently being constructed with PCIP */
+static UObject* GConstructedObject = NULL;
 
 FPostConstructInitializeProperties::FPostConstructInitializeProperties() :
 	Obj(NULL),
 	ObjectArchetype(NULL),
 	bCopyTransientsFromClassDefaults(false),
 	bShouldIntializePropsFromArchetype(true),
-	InstanceGraph(NULL)
+	InstanceGraph(NULL),
+	LastConstructedObject(GConstructedObject, Obj)
 {
 	// Mark we're in the constructor now.
 	GIsInConstructor++;
@@ -1868,7 +1871,8 @@ FPostConstructInitializeProperties::FPostConstructInitializeProperties(UObject*
 	// if the SubobjectRoot NULL, then we want to copy the transients from the template, otherwise we are doing a duplicate and we want to copy the transients from the class defaults
 	bCopyTransientsFromClassDefaults(bInCopyTransientsFromClassDefaults),
 	bShouldIntializePropsFromArchetype(bInShouldIntializeProps),
-	InstanceGraph(InInstanceGraph)
+	InstanceGraph(InInstanceGraph),
+	LastConstructedObject(GConstructedObject, Obj)
 {
 	// Mark we're in the constructor now.
 	GIsInConstructor++;
@@ -2159,6 +2163,10 @@ UObject* StaticConstructObject
 	return Result;
 }
 
+void FPostConstructInitializeProperties::AssertIfInConstructor(UObject* Outer, const TCHAR* ErrorMessage)
+{
+	UE_CLOG(GIsInConstructor && Outer == GConstructedObject, LogUObjectGlobals, Fatal, TEXT("%s"), ErrorMessage);
+}
 
 /**
  * Stores the object flags for all objects in the tracking array.
diff --git a/Engine/Source/Runtime/CoreUObject/Public/UObject/UObjectGlobals.h b/Engine/Source/Runtime/CoreUObject/Public/UObject/UObjectGlobals.h
--- a/Engine/Source/Runtime/CoreUObject/Public/UObject/UObjectGlobals.h
+++ b/Engine/Source/Runtime/CoreUObject/Public/UObject/UObjectGlobals.h
@@ -827,6 +827,11 @@ public:
 	**/
 	bool IslegalOverride(FName InComponentName, class UClass *DerivedComponentClass, class UClass *BaseComponentClass) const;
 
+	/**
+	 * Asserts with the specified message if code is executed inside UObject constructor
+	 **/
+	static void AssertIfInConstructor(UObject* Outer, const TCHAR* ErrorMessage);
+
 private:
 
 	friend class UObject; 
@@ -959,6 +964,8 @@ private:
 	mutable FOverrides ComponentOverrides;
 	/**  List of component classes to intialize after the C++ constructors **/
 	mutable FSubobjectsToInit ComponentInits;
+	/**  Previously constructed object in the callstack */
+	TGuardValue<UObject*> LastConstructedObject;
 };
 
 /**
@@ -993,7 +1000,8 @@ T* ConstructObject(UClass* Class, UObject* Outer = (UObject*)GetTransientPackage
 template< class T >
 T* NewObject(UObject* Outer=(UObject*)GetTransientPackage(),UClass* Class=T::StaticClass() )
 {
-       return ConstructObject<T>(Class,Outer);
+	FPostConstructInitializeProperties::AssertIfInConstructor(Outer, TEXT("NewObject can't be used to create default subobjects (inside of UObject derived class constructor) as it produces inconsistent object names. Use PCIP.CreateDefaultSuobject<> instead."));
+	return ConstructObject<T>(Class,Outer);
 }
 
 /**
