UE4 Refactoring. Changed OVERRIDE and FINAL macros to keywords override and final.

diff --git a/Engine/Source/Editor/WorldBrowser/Private/Tiles/WorldTileModel.h b/Engine/Source/Editor/WorldBrowser/Private/Tiles/WorldTileModel.h
--- a/Engine/Source/Editor/WorldBrowser/Private/Tiles/WorldTileModel.h
+++ b/Engine/Source/Editor/WorldBrowser/Private/Tiles/WorldTileModel.h
@@ -68,21 +68,21 @@ public:
 
 public:
 	// FLevelModel interface
-	virtual UObject* GetNodeObject() OVERRIDE;
-	virtual ULevel* GetLevelObject() const OVERRIDE;
-	virtual FName GetAssetName() const OVERRIDE;
-	virtual FName GetLongPackageName() const OVERRIDE;
-	virtual void Update() OVERRIDE;
-	virtual void LoadLevel() OVERRIDE;
-	virtual void SetVisible(bool bVisible) OVERRIDE;
-	virtual bool HitTest2D(const FVector2D& Point) const OVERRIDE;
-	virtual FVector2D GetLevelPosition2D() const OVERRIDE;
-	virtual FVector2D GetLevelSize2D() const OVERRIDE;
-	virtual void OnDrop(const TSharedPtr<FLevelDragDropOp>& Op) OVERRIDE;
-	virtual bool IsGoodToDrop(const TSharedPtr<FLevelDragDropOp>& Op) const OVERRIDE;
-	virtual void OnLevelAddedToWorld(ULevel* InLevel) OVERRIDE;
-	virtual void OnLevelRemovedFromWorld() OVERRIDE;
-	virtual void OnParentChanged() OVERRIDE;
+	virtual UObject* GetNodeObject() override;
+	virtual ULevel* GetLevelObject() const override;
+	virtual FName GetAssetName() const override;
+	virtual FName GetLongPackageName() const override;
+	virtual void Update() override;
+	virtual void LoadLevel() override;
+	virtual void SetVisible(bool bVisible) override;
+	virtual bool HitTest2D(const FVector2D& Point) const override;
+	virtual FVector2D GetLevelPosition2D() const override;
+	virtual FVector2D GetLevelSize2D() const override;
+	virtual void OnDrop(const TSharedPtr<FLevelDragDropOp>& Op) override;
+	virtual bool IsGoodToDrop(const TSharedPtr<FLevelDragDropOp>& Op) const override;
+	virtual void OnLevelAddedToWorld(ULevel* InLevel) override;
+	virtual void OnLevelRemovedFromWorld() override;
+	virtual void OnParentChanged() override;
 	// FLevelModel interface end
 	
 	/** Adds new streaming level*/
diff --git a/Engine/Source/Programs/CrashReporter/CrashReportClient/Private/CrashReportUtil.h b/Engine/Source/Programs/CrashReporter/CrashReportClient/Private/CrashReportUtil.h
--- a/Engine/Source/Programs/CrashReporter/CrashReportClient/Private/CrashReportUtil.h
+++ b/Engine/Source/Programs/CrashReporter/CrashReportClient/Private/CrashReportUtil.h
@@ -48,7 +48,7 @@ public:
 	}
 
 private:
-	virtual uint32 Run() OVERRIDE
+	virtual uint32 Run() override
 	{
 		Worker.DoWork();
 		FPlatformAtomics::InterlockedIncrement(&bIsDone);
@@ -74,7 +74,7 @@ public:
 	 * @param bIsDirectory Whether the path refers to a file or directory
 	 * @return Whether to carry on iterating
 	 */
-	virtual bool Visit(const TCHAR* FilenameOrDirectory, bool bIsDirectory) OVERRIDE
+	virtual bool Visit(const TCHAR* FilenameOrDirectory, bool bIsDirectory) override
 	{
 		return Functor(FilenameOrDirectory, bIsDirectory);
 	}
diff --git a/Engine/Source/Programs/UnrealHeaderTool/Private/CodeGenerator.cpp b/Engine/Source/Programs/UnrealHeaderTool/Private/CodeGenerator.cpp
--- a/Engine/Source/Programs/UnrealHeaderTool/Private/CodeGenerator.cpp
+++ b/Engine/Source/Programs/UnrealHeaderTool/Private/CodeGenerator.cpp
@@ -4600,7 +4600,7 @@ ECompilationResult::Type UnrealHeaderTool_Main(const FString& ModuleInfoFilename
 						ScriptClass(Class)
 					{ }
 
-					virtual FString GetContext() OVERRIDE
+					virtual FString GetContext() override
 					{
 						FString Filename = IFileManager::Get().ConvertToAbsolutePathForExternalAppForRead(*GClassSourceFileMap[ScriptClass]);
 						int32 LineNumber = GClassDeclarationLineNumber[ScriptClass];
diff --git a/Engine/Source/Programs/UnrealHeaderTool/Private/HeaderParser.h b/Engine/Source/Programs/UnrealHeaderTool/Private/HeaderParser.h
--- a/Engine/Source/Programs/UnrealHeaderTool/Private/HeaderParser.h
+++ b/Engine/Source/Programs/UnrealHeaderTool/Private/HeaderParser.h
@@ -298,7 +298,7 @@ protected:
 	static void ExportNativeHeaders( UPackage* CurrentPackage, FClasses& AllClasses, bool bAllowSaveExportedHeaders );
 
 	// FContextSupplier interface.
-	virtual FString GetContext() OVERRIDE;
+	virtual FString GetContext() override;
 	// End of FContextSupplier interface.
 
 	// High-level compiling functions.
diff --git a/Engine/Source/Programs/UnrealLightmass/Private/LightmassCore/Misc/LMHelpers.h b/Engine/Source/Programs/UnrealLightmass/Private/LightmassCore/Misc/LMHelpers.h
--- a/Engine/Source/Programs/UnrealLightmass/Private/LightmassCore/Misc/LMHelpers.h
+++ b/Engine/Source/Programs/UnrealLightmass/Private/LightmassCore/Misc/LMHelpers.h
@@ -13,8 +13,8 @@ public:
 	~FLightmassLog();
 
 	// BEGIN FOutputDevice Interface 
-	virtual void Serialize( const TCHAR* V, ELogVerbosity::Type Verbosity, const class FName& Category ) OVERRIDE;
-	virtual void Flush() OVERRIDE;
+	virtual void Serialize( const TCHAR* V, ELogVerbosity::Type Verbosity, const class FName& Category ) override;
+	virtual void Flush() override;
 	// END FOutputDevice Interface
 
 	/**
diff --git a/Engine/Source/Programs/UnrealPak/Private/KeyGenerator.cpp b/Engine/Source/Programs/UnrealPak/Private/KeyGenerator.cpp
--- a/Engine/Source/Programs/UnrealPak/Private/KeyGenerator.cpp
+++ b/Engine/Source/Programs/UnrealPak/Private/KeyGenerator.cpp
@@ -58,8 +58,8 @@ public:
 	}
 
 	// Start FRunnable interface
-	virtual bool Init() OVERRIDE { return true; }
-	virtual uint32 Run() OVERRIDE
+	virtual bool Init() override { return true; }
+	virtual uint32 Run() override
 	{
 		int256 Remainder;
 		int32 FactorCheckTimer = 0;
@@ -320,8 +320,8 @@ public:
 	}
 
 	// Start FRunnable interface
-	virtual bool Init() OVERRIDE { return true; }
-	virtual uint32 Run() OVERRIDE
+	virtual bool Init() override { return true; }
+	virtual uint32 Run() override
 	{
 		int256 Remainder;
 		int32 FactorCheckTimer = 0;
diff --git a/Engine/Source/Programs/UnrealPak/Private/SignedArchiveWriter.h b/Engine/Source/Programs/UnrealPak/Private/SignedArchiveWriter.h
--- a/Engine/Source/Programs/UnrealPak/Private/SignedArchiveWriter.h
+++ b/Engine/Source/Programs/UnrealPak/Private/SignedArchiveWriter.h
@@ -36,9 +36,9 @@ public:
 	virtual ~FSignedArchiveWriter();
 
 	// FArchive interface
-	virtual bool Close() OVERRIDE;
-	virtual void Serialize(void* Data, int64 Length) OVERRIDE;
-	virtual int64 Tell() OVERRIDE;
-	virtual int64 TotalSize() OVERRIDE;
-	virtual void Seek(int64 InPos) OVERRIDE;
+	virtual bool Close() override;
+	virtual void Serialize(void* Data, int64 Length) override;
+	virtual int64 Tell() override;
+	virtual int64 TotalSize() override;
+	virtual void Seek(int64 InPos) override;
 };
diff --git a/Engine/Source/Programs/UnrealSync/Private/GUI.cpp b/Engine/Source/Programs/UnrealSync/Private/GUI.cpp
--- a/Engine/Source/Programs/UnrealSync/Private/GUI.cpp
+++ b/Engine/Source/Programs/UnrealSync/Private/GUI.cpp
@@ -402,7 +402,7 @@ public:
 	 *
 	 * @returns UAT UnrealSync command line.
 	 */
-	FString GetLabelName() const OVERRIDE
+	FString GetLabelName() const override
 	{
 		DataReady->Wait();
 
@@ -425,7 +425,7 @@ public:
 		return "Sync to latest finished.";
 	}
 
-	FString GetGameName() const OVERRIDE
+	FString GetGameName() const override
 	{
 		DataReady->Wait();
 
@@ -437,7 +437,7 @@ public:
 	 *
 	 * @param GameName Current game name.
 	 */
-	void RefreshData(const FString& GameName) OVERRIDE
+	void RefreshData(const FString& GameName) override
 	{
 		ILabelNameProvider::RefreshData(GameName);
 
@@ -455,7 +455,7 @@ public:
 	 *
 	 * @param GameName Current game name.
 	 */
-	void ResetData(const FString& GameName) OVERRIDE
+	void ResetData(const FString& GameName) override
 	{
 		ILabelNameProvider::ResetData(GameName);
 
@@ -467,7 +467,7 @@ public:
 	 *
 	 * @returns True if ready. False otherwise.
 	 */
-	bool IsReadyForSync() const OVERRIDE
+	bool IsReadyForSync() const override
 	{
 		return true;
 	}
@@ -519,7 +519,7 @@ public:
 	 *
 	 * @returns UAT UnrealSync command line.
 	 */
-	FString GetLabelName() const OVERRIDE
+	FString GetLabelName() const override
 	{
 		return LabelsCombo->GetCurrentOption();
 	}
@@ -529,7 +529,7 @@ public:
 	 *
 	 * @returns True if ready. False otherwise.
 	 */
-	bool IsReadyForSync() const OVERRIDE
+	bool IsReadyForSync() const override
 	{
 		return FUnrealSync::HasValidData() && !LabelsCombo->IsEmpty();
 	}
@@ -539,7 +539,7 @@ public:
 	 *
 	 * @param GameName Current game name.
 	 */
-	void ResetData(const FString& GameName) OVERRIDE
+	void ResetData(const FString& GameName) override
 	{
 		ILabelNameProvider::ResetData(GameName);
 
@@ -590,7 +590,7 @@ public:
 	 *
 	 * @param GameName Game name to refresh for.
 	 */
-	void RefreshData(const FString& GameName) OVERRIDE
+	void RefreshData(const FString& GameName) override
 	{
 		SetLabelOptions(*FUnrealSync::GetPromotedLabelsForGame(GameName));
 
@@ -609,7 +609,7 @@ public:
 	 *
 	 * @param GameName Game name to refresh for.
 	 */
-	void RefreshData(const FString& GameName) OVERRIDE
+	void RefreshData(const FString& GameName) override
 	{
 		SetLabelOptions(*FUnrealSync::GetPromotableLabelsForGame(GameName));
 
@@ -628,7 +628,7 @@ public:
 	 *
 	 * @param GameName Game name to refresh for.
 	 */
-	void RefreshData(const FString& GameName) OVERRIDE
+	void RefreshData(const FString& GameName) override
 	{
 		SetLabelOptions(*FUnrealSync::GetAllLabels());
 
diff --git a/Engine/Source/Programs/UnrealSync/Private/P4DataCache.h b/Engine/Source/Programs/UnrealSync/Private/P4DataCache.h
--- a/Engine/Source/Programs/UnrealSync/Private/P4DataCache.h
+++ b/Engine/Source/Programs/UnrealSync/Private/P4DataCache.h
@@ -89,14 +89,14 @@ public:
 	 *
 	 * @return True if initialization was successful, false otherwise
 	 */
-	virtual bool Init() OVERRIDE;
+	virtual bool Init() override;
 
 	/**
 	 * Main function to run for this thread.
 	 *
 	 * @returns Exit code.
 	 */
-	uint32 Run() OVERRIDE;
+	uint32 Run() override;
 
 	/**
 	 * Exits the runnable object.
diff --git a/Engine/Source/Programs/UnrealSync/Private/UnrealSync.cpp b/Engine/Source/Programs/UnrealSync/Private/UnrealSync.cpp
--- a/Engine/Source/Programs/UnrealSync/Private/UnrealSync.cpp
+++ b/Engine/Source/Programs/UnrealSync/Private/UnrealSync.cpp
@@ -257,7 +257,7 @@ public:
 	/**
 	 * Main thread function.
 	 */
-	uint32 Run() OVERRIDE
+	uint32 Run() override
 	{
 		if (OnSyncProgress.IsBound())
 		{
diff --git a/Engine/Source/Runtime/AIModule/Classes/AIController.h b/Engine/Source/Runtime/AIModule/Classes/AIController.h
--- a/Engine/Source/Runtime/AIModule/Classes/AIController.h
+++ b/Engine/Source/Runtime/AIModule/Classes/AIController.h
@@ -146,7 +146,7 @@ public:
 	bool ResumeMove(FAIRequestID RequestToResume);
 
 	/** Aborts the move the controller is currently performing */
-	virtual void StopMovement() OVERRIDE;
+	virtual void StopMovement() override;
 
 	/** Called on completing current movement request */
 	virtual void OnMoveCompleted(FAIRequestID RequestID, EPathFollowingResult::Type Result);
@@ -181,7 +181,7 @@ public:
 	void SetMoveBlockDetection(bool bEnable);
 
 	/** Prepares path finding and path following components. */
-	virtual void InitNavigationControl(UNavigationComponent*& PathFindingComp, UPathFollowingComponent*& PathFollowingComp) OVERRIDE;
+	virtual void InitNavigationControl(UNavigationComponent*& PathFindingComp, UPathFollowingComponent*& PathFollowingComp) override;
 
 	/** Starts executing behavior tree. */
 	UFUNCTION(BlueprintCallable, Category="AI")
@@ -230,21 +230,21 @@ public:
 	bool SuggestTossVelocity(FVector& OutTossVelocity, FVector Start, FVector End, float TossSpeed, bool bPreferHighArc, float CollisionRadius=0, bool bOnlyTraceUp=false);
 
 	// Begin AActor Interface
-	virtual void Tick(float DeltaTime) OVERRIDE;
-	virtual void PostInitializeComponents() OVERRIDE;
+	virtual void Tick(float DeltaTime) override;
+	virtual void PostInitializeComponents() override;
 	// End AActor Interface
 
 	// Begin AController Interface
-	virtual void Possess(class APawn* InPawn) OVERRIDE;
-	virtual void UnPossess() OVERRIDE;
-	virtual void DisplayDebug(class UCanvas* Canvas, const FDebugDisplayInfo& DebugDisplay, float& YL, float& YPos) OVERRIDE;
+	virtual void Possess(class APawn* InPawn) override;
+	virtual void UnPossess() override;
+	virtual void DisplayDebug(class UCanvas* Canvas, const FDebugDisplayInfo& DebugDisplay, float& YL, float& YPos) override;
 
 #if ENABLE_VISUAL_LOG
-	virtual void GrabDebugSnapshot(struct FVisLogEntry* Snapshot) const OVERRIDE;
+	virtual void GrabDebugSnapshot(struct FVisLogEntry* Snapshot) const override;
 #endif
 
-	virtual void Reset() OVERRIDE;
-	virtual void GetPlayerViewPoint(FVector& out_Location, FRotator& out_Rotation) const OVERRIDE;
+	virtual void Reset() override;
+	virtual void GetPlayerViewPoint(FVector& out_Location, FRotator& out_Rotation) const override;
 
 	/**
 	 * Checks line to center and top of other actor
@@ -253,7 +253,7 @@ public:
 	 * @param bAlternateChecks used only in AIController implementation
 	 * @return true if controller's pawn can see Other actor.
 	 */
-	virtual bool LineOfSightTo(const class AActor* Other, FVector ViewPoint = FVector(ForceInit), bool bAlternateChecks = false) OVERRIDE;
+	virtual bool LineOfSightTo(const class AActor* Other, FVector ViewPoint = FVector(ForceInit), bool bAlternateChecks = false) override;
 	// End AController Interface
 
 	/** Update direction AI is looking based on FocalPoint */
diff --git a/Engine/Source/Runtime/AIModule/Classes/AISystem.h b/Engine/Source/Runtime/AIModule/Classes/AISystem.h
--- a/Engine/Source/Runtime/AIModule/Classes/AISystem.h
+++ b/Engine/Source/Runtime/AIModule/Classes/AISystem.h
@@ -26,9 +26,9 @@ class AIMODULE_API UAISystem : public UAISystemBase
 	~UAISystem();
 
 	// IAISystemInterface begin		
-	virtual void InitializeActorsForPlay(bool bTimeGotReset) OVERRIDE;
-	virtual void WorldOriginChanged(FIntPoint OldOrigin, FIntPoint NewOrigin) OVERRIDE;
-	virtual void CleanupWorld(bool bSessionEnded = true, bool bCleanupResources = true, UWorld* NewWorld = NULL) OVERRIDE;
+	virtual void InitializeActorsForPlay(bool bTimeGotReset) override;
+	virtual void WorldOriginChanged(FIntPoint OldOrigin, FIntPoint NewOrigin) override;
+	virtual void CleanupWorld(bool bSessionEnded = true, bool bCleanupResources = true, UWorld* NewWorld = NULL) override;
 	// IAISystemInterface end
 	
 	/** Behavior tree manager getter */
diff --git a/Engine/Source/Runtime/AIModule/Classes/BehaviorTree/BTAuxiliaryNode.h b/Engine/Source/Runtime/AIModule/Classes/BehaviorTree/BTAuxiliaryNode.h
--- a/Engine/Source/Runtime/AIModule/Classes/BehaviorTree/BTAuxiliaryNode.h
+++ b/Engine/Source/Runtime/AIModule/Classes/BehaviorTree/BTAuxiliaryNode.h
@@ -38,8 +38,8 @@ class AIMODULE_API UBTAuxiliaryNode : public UBTNode
 	/** wrapper for node instancing: TickNode */
 	void WrappedTickNode(class UBehaviorTreeComponent* OwnerComp, uint8* NodeMemory, float DeltaSeconds) const;
 
-	virtual void DescribeRuntimeValues(const class UBehaviorTreeComponent* OwnerComp, uint8* NodeMemory, EBTDescriptionVerbosity::Type Verbosity, TArray<FString>& Values) const OVERRIDE;
-	virtual uint16 GetSpecialMemorySize() const OVERRIDE;
+	virtual void DescribeRuntimeValues(const class UBehaviorTreeComponent* OwnerComp, uint8* NodeMemory, EBTDescriptionVerbosity::Type Verbosity, TArray<FString>& Values) const override;
+	virtual uint16 GetSpecialMemorySize() const override;
 
 protected:
 
diff --git a/Engine/Source/Runtime/AIModule/Classes/BehaviorTree/BTCompositeNode.h b/Engine/Source/Runtime/AIModule/Classes/BehaviorTree/BTCompositeNode.h
--- a/Engine/Source/Runtime/AIModule/Classes/BehaviorTree/BTCompositeNode.h
+++ b/Engine/Source/Runtime/AIModule/Classes/BehaviorTree/BTCompositeNode.h
@@ -124,7 +124,7 @@ class AIMODULE_API UBTCompositeNode : public UBTNode
 	void ConditionalNotifyChildExecution(class UBehaviorTreeComponent* OwnerComp, uint8* NodeMemory, const class UBTNode* ChildNode, EBTNodeResult::Type& NodeResult) const;
 
 	/** size of instance memory */
-	virtual uint16 GetInstanceMemorySize() const OVERRIDE;
+	virtual uint16 GetInstanceMemorySize() const override;
 
 	/** @return child node at given index */
 	class UBTNode* GetChildNode(int32 Index) const;
@@ -142,7 +142,7 @@ class AIMODULE_API UBTCompositeNode : public UBTNode
 	void SetChildOverride(FBehaviorTreeSearchData& SearchData, int8 Index) const;
 
 	/** gathers description of all runtime parameters */
-	virtual void DescribeRuntimeValues(const class UBehaviorTreeComponent* OwnerComp, uint8* NodeMemory, EBTDescriptionVerbosity::Type Verbosity, TArray<FString>& Values) const OVERRIDE;
+	virtual void DescribeRuntimeValues(const class UBehaviorTreeComponent* OwnerComp, uint8* NodeMemory, EBTDescriptionVerbosity::Type Verbosity, TArray<FString>& Values) const override;
 
 	/** check if child node can execute new subtree */
 	virtual bool CanPushSubtree(class UBehaviorTreeComponent* OwnerComp, uint8* NodeMemory, int32 ChildIdx) const;
diff --git a/Engine/Source/Runtime/AIModule/Classes/BehaviorTree/BTDecorator.h b/Engine/Source/Runtime/AIModule/Classes/BehaviorTree/BTDecorator.h
--- a/Engine/Source/Runtime/AIModule/Classes/BehaviorTree/BTDecorator.h
+++ b/Engine/Source/Runtime/AIModule/Classes/BehaviorTree/BTDecorator.h
@@ -54,7 +54,7 @@ class AIMODULE_API UBTDecorator : public UBTAuxiliaryNode
 	/** @return true if condition should be inversed */
 	bool IsInversed() const;
 
-	virtual FString GetStaticDescription() const OVERRIDE;
+	virtual FString GetStaticDescription() const override;
 
 	/** modify current flow abort mode, so it can be used with parent composite */
 	void UpdateFlowAbortMode();
diff --git a/Engine/Source/Runtime/AIModule/Classes/BehaviorTree/BTNode.h b/Engine/Source/Runtime/AIModule/Classes/BehaviorTree/BTNode.h
--- a/Engine/Source/Runtime/AIModule/Classes/BehaviorTree/BTNode.h
+++ b/Engine/Source/Runtime/AIModule/Classes/BehaviorTree/BTNode.h
@@ -16,7 +16,7 @@ class AIMODULE_API UBTNode : public UObject
 {
 	GENERATED_UCLASS_BODY()
 
-	virtual class UWorld* GetWorld() const OVERRIDE;
+	virtual class UWorld* GetWorld() const override;
 
 	/** fill in data about tree structure */
 	void InitializeNode(class UBTCompositeNode* InParentNode, uint16 InExecutionIndex, uint16 InMemoryOffset, uint8 InTreeDepth);
diff --git a/Engine/Source/Runtime/AIModule/Classes/BehaviorTree/BTService.h b/Engine/Source/Runtime/AIModule/Classes/BehaviorTree/BTService.h
--- a/Engine/Source/Runtime/AIModule/Classes/BehaviorTree/BTService.h
+++ b/Engine/Source/Runtime/AIModule/Classes/BehaviorTree/BTService.h
@@ -34,7 +34,7 @@ class AIMODULE_API UBTService : public UBTAuxiliaryNode
 {
 	GENERATED_UCLASS_BODY()
 
-	virtual FString GetStaticDescription() const OVERRIDE;
+	virtual FString GetStaticDescription() const override;
 
 protected:
 
@@ -48,11 +48,11 @@ protected:
 
 	/** update next tick interval
 	 * this function should be considered as const (don't modify state of object) if node is not instanced! */
-	virtual void TickNode(UBehaviorTreeComponent* OwnerComp, uint8* NodeMemory, float DeltaSeconds) OVERRIDE;
+	virtual void TickNode(UBehaviorTreeComponent* OwnerComp, uint8* NodeMemory, float DeltaSeconds) override;
 
 	/** called when auxiliary node becomes active
 	 * this function should be considered as const (don't modify state of object) if node is not instanced! */
-	virtual void OnBecomeRelevant(UBehaviorTreeComponent* OwnerComp, uint8* NodeMemory) OVERRIDE;
+	virtual void OnBecomeRelevant(UBehaviorTreeComponent* OwnerComp, uint8* NodeMemory) override;
 
 #if WITH_EDITOR
 	virtual FName GetNodeIconName() const override;
diff --git a/Engine/Source/Runtime/AIModule/Classes/BehaviorTree/BehaviorTreeComponent.h b/Engine/Source/Runtime/AIModule/Classes/BehaviorTree/BehaviorTreeComponent.h
--- a/Engine/Source/Runtime/AIModule/Classes/BehaviorTree/BehaviorTreeComponent.h
+++ b/Engine/Source/Runtime/AIModule/Classes/BehaviorTree/BehaviorTreeComponent.h
@@ -33,18 +33,18 @@ class AIMODULE_API UBehaviorTreeComponent : public UBrainComponent
 	GENERATED_UCLASS_BODY()
 
 	// Begin UBrainComponent overrides
-	virtual void RestartLogic() OVERRIDE;
+	virtual void RestartLogic() override;
 protected:
-	virtual void StopLogic(const FString& Reason) OVERRIDE;
-	virtual void PauseLogic(const FString& Reason) OVERRIDE;
-	virtual EAILogicResuming::Type ResumeLogic(const FString& Reason) OVERRIDE;
+	virtual void StopLogic(const FString& Reason) override;
+	virtual void PauseLogic(const FString& Reason) override;
+	virtual EAILogicResuming::Type ResumeLogic(const FString& Reason) override;
 
 	/** indicates instance has been initialized to work with specific BT asset */
 	bool TreeHasBeenStarted() const;
 
 public:
-	virtual bool IsRunning() const OVERRIDE;
-	virtual bool IsPaused() const OVERRIDE;
+	virtual bool IsRunning() const override;
+	virtual bool IsPaused() const override;
 	// End UBrainComponent overrides
 
 	/** starts execution from root */
@@ -88,7 +88,7 @@ public:
 	void UnregisterAuxNodesUpTo(const struct FBTNodeIndex& Index);
 
 	/** BEGIN UActorComponent overrides */
-	virtual void TickComponent(float DeltaTime, enum ELevelTick TickType, FActorComponentTickFunction *ThisTickFunction) OVERRIDE;
+	virtual void TickComponent(float DeltaTime, enum ELevelTick TickType, FActorComponentTickFunction *ThisTickFunction) override;
 	/** END UActorComponent overrides */
 
 	/** process execution flow */
@@ -130,12 +130,12 @@ public:
 	/** @return status of speficied task */
 	EBTTaskStatus::Type GetTaskStatus(const class UBTTaskNode* TaskNode) const;
 
-	virtual FString GetDebugInfoString() const OVERRIDE;
+	virtual FString GetDebugInfoString() const override;
 	virtual FString DescribeActiveTasks() const;
 	virtual FString DescribeActiveTrees() const;
 
 #if ENABLE_VISUAL_LOG
-	virtual void DescribeSelfToVisLog(struct FVisLogEntry* Snapshot) const OVERRIDE;
+	virtual void DescribeSelfToVisLog(struct FVisLogEntry* Snapshot) const override;
 #endif
 
 protected:
diff --git a/Engine/Source/Runtime/AIModule/Classes/BehaviorTree/BehaviorTreeManager.h b/Engine/Source/Runtime/AIModule/Classes/BehaviorTree/BehaviorTreeManager.h
--- a/Engine/Source/Runtime/AIModule/Classes/BehaviorTree/BehaviorTreeManager.h
+++ b/Engine/Source/Runtime/AIModule/Classes/BehaviorTree/BehaviorTreeManager.h
@@ -42,7 +42,7 @@ class AIMODULE_API UBehaviorTreeManager : public UObject
 	static void InitializeMemoryHelper(const TArray<class UBTDecorator*>& Nodes, TArray<uint16>& MemoryOffsets, int32& MemorySize);
 
 	/** cleanup hooks for map loading */
-	virtual void FinishDestroy() OVERRIDE;
+	virtual void FinishDestroy() override;
 
 	void DumpUsageStats() const;
 
diff --git a/Engine/Source/Runtime/AIModule/Classes/BehaviorTree/Blackboard/BlackboardKeyType_Bool.h b/Engine/Source/Runtime/AIModule/Classes/BehaviorTree/Blackboard/BlackboardKeyType_Bool.h
--- a/Engine/Source/Runtime/AIModule/Classes/BehaviorTree/Blackboard/BlackboardKeyType_Bool.h
+++ b/Engine/Source/Runtime/AIModule/Classes/BehaviorTree/Blackboard/BlackboardKeyType_Bool.h
@@ -11,7 +11,7 @@ class UBlackboardKeyType_Bool : public UBlackboardKeyType
 	static bool GetValue(const uint8* RawData);
 	static bool SetValue(uint8* RawData, bool bValue);
 
-	virtual FString DescribeValue(const uint8* RawData) const OVERRIDE;
-	virtual EBlackboardCompare::Type Compare(const uint8* MemoryBlockA, const uint8* MemoryBlockB) const OVERRIDE;
-	virtual bool TestBasicOperation(const uint8* MemoryBlock, EBasicKeyOperation::Type Op) const OVERRIDE;
+	virtual FString DescribeValue(const uint8* RawData) const override;
+	virtual EBlackboardCompare::Type Compare(const uint8* MemoryBlockA, const uint8* MemoryBlockB) const override;
+	virtual bool TestBasicOperation(const uint8* MemoryBlock, EBasicKeyOperation::Type Op) const override;
 };
diff --git a/Engine/Source/Runtime/AIModule/Classes/BehaviorTree/Blackboard/BlackboardKeyType_Class.h b/Engine/Source/Runtime/AIModule/Classes/BehaviorTree/Blackboard/BlackboardKeyType_Class.h
--- a/Engine/Source/Runtime/AIModule/Classes/BehaviorTree/Blackboard/BlackboardKeyType_Class.h
+++ b/Engine/Source/Runtime/AIModule/Classes/BehaviorTree/Blackboard/BlackboardKeyType_Class.h
@@ -14,9 +14,9 @@ class UBlackboardKeyType_Class : public UBlackboardKeyType
 	static UClass* GetValue(const uint8* RawData);
 	static bool SetValue(uint8* RawData, UClass* Value);
 
-	virtual FString DescribeValue(const uint8* RawData) const OVERRIDE;
-	virtual FString DescribeSelf() const OVERRIDE;
-	virtual bool IsAllowedByFilter(UBlackboardKeyType* FilterOb) const OVERRIDE;
-	virtual EBlackboardCompare::Type Compare(const uint8* MemoryBlockA, const uint8* MemoryBlockB) const OVERRIDE;
-	virtual bool TestBasicOperation(const uint8* MemoryBlock, EBasicKeyOperation::Type Op) const OVERRIDE;
+	virtual FString DescribeValue(const uint8* RawData) const override;
+	virtual FString DescribeSelf() const override;
+	virtual bool IsAllowedByFilter(UBlackboardKeyType* FilterOb) const override;
+	virtual EBlackboardCompare::Type Compare(const uint8* MemoryBlockA, const uint8* MemoryBlockB) const override;
+	virtual bool TestBasicOperation(const uint8* MemoryBlock, EBasicKeyOperation::Type Op) const override;
 };
diff --git a/Engine/Source/Runtime/AIModule/Classes/BehaviorTree/Blackboard/BlackboardKeyType_Enum.h b/Engine/Source/Runtime/AIModule/Classes/BehaviorTree/Blackboard/BlackboardKeyType_Enum.h
--- a/Engine/Source/Runtime/AIModule/Classes/BehaviorTree/Blackboard/BlackboardKeyType_Enum.h
+++ b/Engine/Source/Runtime/AIModule/Classes/BehaviorTree/Blackboard/BlackboardKeyType_Enum.h
@@ -14,10 +14,10 @@ class UBlackboardKeyType_Enum : public UBlackboardKeyType
 	static uint8 GetValue(const uint8* RawData);
 	static bool SetValue(uint8* RawData, uint8 Value);
 
-	virtual FString DescribeValue(const uint8* RawData) const OVERRIDE;
-	virtual FString DescribeSelf() const OVERRIDE;
-	virtual bool IsAllowedByFilter(UBlackboardKeyType* FilterOb) const OVERRIDE;
-	virtual EBlackboardCompare::Type Compare(const uint8* MemoryBlockA, const uint8* MemoryBlockB) const OVERRIDE;
-	virtual bool TestArithmeticOperation(const uint8* MemoryBlock, EArithmeticKeyOperation::Type Op, int32 OtherIntValue, float OtherFloatValue) const OVERRIDE;
-	virtual FString DescribeArithmeticParam(int32 IntValue, float FloatValue) const OVERRIDE;
+	virtual FString DescribeValue(const uint8* RawData) const override;
+	virtual FString DescribeSelf() const override;
+	virtual bool IsAllowedByFilter(UBlackboardKeyType* FilterOb) const override;
+	virtual EBlackboardCompare::Type Compare(const uint8* MemoryBlockA, const uint8* MemoryBlockB) const override;
+	virtual bool TestArithmeticOperation(const uint8* MemoryBlock, EArithmeticKeyOperation::Type Op, int32 OtherIntValue, float OtherFloatValue) const override;
+	virtual FString DescribeArithmeticParam(int32 IntValue, float FloatValue) const override;
 };
diff --git a/Engine/Source/Runtime/AIModule/Classes/BehaviorTree/Blackboard/BlackboardKeyType_Float.h b/Engine/Source/Runtime/AIModule/Classes/BehaviorTree/Blackboard/BlackboardKeyType_Float.h
--- a/Engine/Source/Runtime/AIModule/Classes/BehaviorTree/Blackboard/BlackboardKeyType_Float.h
+++ b/Engine/Source/Runtime/AIModule/Classes/BehaviorTree/Blackboard/BlackboardKeyType_Float.h
@@ -11,8 +11,8 @@ class UBlackboardKeyType_Float : public UBlackboardKeyType
 	static float GetValue(const uint8* RawData);
 	static bool SetValue(uint8* RawData, float Value);
 
-	virtual FString DescribeValue(const uint8* RawData) const OVERRIDE;
-	virtual EBlackboardCompare::Type Compare(const uint8* MemoryBlockA, const uint8* MemoryBlockB) const OVERRIDE;
-	virtual bool TestArithmeticOperation(const uint8* MemoryBlock, EArithmeticKeyOperation::Type Op, int32 OtherIntValue, float OtherFloatValue) const OVERRIDE;
-	virtual FString DescribeArithmeticParam(int32 IntValue, float FloatValue) const OVERRIDE;
+	virtual FString DescribeValue(const uint8* RawData) const override;
+	virtual EBlackboardCompare::Type Compare(const uint8* MemoryBlockA, const uint8* MemoryBlockB) const override;
+	virtual bool TestArithmeticOperation(const uint8* MemoryBlock, EArithmeticKeyOperation::Type Op, int32 OtherIntValue, float OtherFloatValue) const override;
+	virtual FString DescribeArithmeticParam(int32 IntValue, float FloatValue) const override;
 };
diff --git a/Engine/Source/Runtime/AIModule/Classes/BehaviorTree/Blackboard/BlackboardKeyType_Int.h b/Engine/Source/Runtime/AIModule/Classes/BehaviorTree/Blackboard/BlackboardKeyType_Int.h
--- a/Engine/Source/Runtime/AIModule/Classes/BehaviorTree/Blackboard/BlackboardKeyType_Int.h
+++ b/Engine/Source/Runtime/AIModule/Classes/BehaviorTree/Blackboard/BlackboardKeyType_Int.h
@@ -11,8 +11,8 @@ class UBlackboardKeyType_Int : public UBlackboardKeyType
 	static int32 GetValue(const uint8* RawData);
 	static bool SetValue(uint8* RawData, int32 Value);
 
-	virtual FString DescribeValue(const uint8* RawData) const OVERRIDE;
-	virtual EBlackboardCompare::Type Compare(const uint8* MemoryBlockA, const uint8* MemoryBlockB) const OVERRIDE;
-	virtual bool TestArithmeticOperation(const uint8* MemoryBlock, EArithmeticKeyOperation::Type Op, int32 OtherIntValue, float OtherFloatValue) const OVERRIDE;
-	virtual FString DescribeArithmeticParam(int32 IntValue, float FloatValue) const OVERRIDE;
+	virtual FString DescribeValue(const uint8* RawData) const override;
+	virtual EBlackboardCompare::Type Compare(const uint8* MemoryBlockA, const uint8* MemoryBlockB) const override;
+	virtual bool TestArithmeticOperation(const uint8* MemoryBlock, EArithmeticKeyOperation::Type Op, int32 OtherIntValue, float OtherFloatValue) const override;
+	virtual FString DescribeArithmeticParam(int32 IntValue, float FloatValue) const override;
 };
diff --git a/Engine/Source/Runtime/AIModule/Classes/BehaviorTree/Blackboard/BlackboardKeyType_Name.h b/Engine/Source/Runtime/AIModule/Classes/BehaviorTree/Blackboard/BlackboardKeyType_Name.h
--- a/Engine/Source/Runtime/AIModule/Classes/BehaviorTree/Blackboard/BlackboardKeyType_Name.h
+++ b/Engine/Source/Runtime/AIModule/Classes/BehaviorTree/Blackboard/BlackboardKeyType_Name.h
@@ -11,7 +11,7 @@ class UBlackboardKeyType_Name : public UBlackboardKeyType
 	static FName GetValue(const uint8* RawData);
 	static bool SetValue(uint8* RawData, const FName& Value);
 
-	virtual FString DescribeValue(const uint8* RawData) const OVERRIDE;
-	virtual EBlackboardCompare::Type Compare(const uint8* MemoryBlockA, const uint8* MemoryBlockB) const OVERRIDE;
-	virtual bool TestTextOperation(const uint8* MemoryBlock, ETextKeyOperation::Type Op, const FString& OtherString) const OVERRIDE;
+	virtual FString DescribeValue(const uint8* RawData) const override;
+	virtual EBlackboardCompare::Type Compare(const uint8* MemoryBlockA, const uint8* MemoryBlockB) const override;
+	virtual bool TestTextOperation(const uint8* MemoryBlock, ETextKeyOperation::Type Op, const FString& OtherString) const override;
 };
diff --git a/Engine/Source/Runtime/AIModule/Classes/BehaviorTree/Blackboard/BlackboardKeyType_NativeEnum.h b/Engine/Source/Runtime/AIModule/Classes/BehaviorTree/Blackboard/BlackboardKeyType_NativeEnum.h
--- a/Engine/Source/Runtime/AIModule/Classes/BehaviorTree/Blackboard/BlackboardKeyType_NativeEnum.h
+++ b/Engine/Source/Runtime/AIModule/Classes/BehaviorTree/Blackboard/BlackboardKeyType_NativeEnum.h
@@ -21,14 +21,14 @@ class UBlackboardKeyType_NativeEnum : public UBlackboardKeyType
 	static uint8 GetValue(const uint8* RawData);
 	static bool SetValue(uint8* RawData, uint8 Value);
 
-	virtual FString DescribeValue(const uint8* RawData) const OVERRIDE;
-	virtual FString DescribeSelf() const OVERRIDE;
-	virtual bool IsAllowedByFilter(UBlackboardKeyType* FilterOb) const OVERRIDE;
-	virtual EBlackboardCompare::Type Compare(const uint8* MemoryBlockA, const uint8* MemoryBlockB) const OVERRIDE;
-	virtual bool TestArithmeticOperation(const uint8* MemoryBlock, EArithmeticKeyOperation::Type Op, int32 OtherIntValue, float OtherFloatValue) const OVERRIDE;
-	virtual FString DescribeArithmeticParam(int32 IntValue, float FloatValue) const OVERRIDE;
+	virtual FString DescribeValue(const uint8* RawData) const override;
+	virtual FString DescribeSelf() const override;
+	virtual bool IsAllowedByFilter(UBlackboardKeyType* FilterOb) const override;
+	virtual EBlackboardCompare::Type Compare(const uint8* MemoryBlockA, const uint8* MemoryBlockB) const override;
+	virtual bool TestArithmeticOperation(const uint8* MemoryBlock, EArithmeticKeyOperation::Type Op, int32 OtherIntValue, float OtherFloatValue) const override;
+	virtual FString DescribeArithmeticParam(int32 IntValue, float FloatValue) const override;
 
 #if WITH_EDITOR
-	virtual void PostEditChangeProperty(struct FPropertyChangedEvent& PropertyChangedEvent) OVERRIDE;
+	virtual void PostEditChangeProperty(struct FPropertyChangedEvent& PropertyChangedEvent) override;
 #endif
 };
diff --git a/Engine/Source/Runtime/AIModule/Classes/BehaviorTree/Blackboard/BlackboardKeyType_Object.h b/Engine/Source/Runtime/AIModule/Classes/BehaviorTree/Blackboard/BlackboardKeyType_Object.h
--- a/Engine/Source/Runtime/AIModule/Classes/BehaviorTree/Blackboard/BlackboardKeyType_Object.h
+++ b/Engine/Source/Runtime/AIModule/Classes/BehaviorTree/Blackboard/BlackboardKeyType_Object.h
@@ -14,11 +14,11 @@ class UBlackboardKeyType_Object : public UBlackboardKeyType
 	static UObject* GetValue(const uint8* RawData);
 	static bool SetValue(uint8* RawData, UObject* Value);
 
-	virtual FString DescribeValue(const uint8* RawData) const OVERRIDE;
-	virtual FString DescribeSelf() const OVERRIDE;
-	virtual bool IsAllowedByFilter(UBlackboardKeyType* FilterOb) const OVERRIDE;
-	virtual bool GetLocation(const uint8* RawData, FVector& Location) const OVERRIDE;
-	virtual bool GetRotation(const uint8* MemoryBlock, FRotator& Rotation) const OVERRIDE;
-	virtual EBlackboardCompare::Type Compare(const uint8* MemoryBlockA, const uint8* MemoryBlockB) const OVERRIDE;
-	virtual bool TestBasicOperation(const uint8* MemoryBlock, EBasicKeyOperation::Type Op) const OVERRIDE;
+	virtual FString DescribeValue(const uint8* RawData) const override;
+	virtual FString DescribeSelf() const override;
+	virtual bool IsAllowedByFilter(UBlackboardKeyType* FilterOb) const override;
+	virtual bool GetLocation(const uint8* RawData, FVector& Location) const override;
+	virtual bool GetRotation(const uint8* MemoryBlock, FRotator& Rotation) const override;
+	virtual EBlackboardCompare::Type Compare(const uint8* MemoryBlockA, const uint8* MemoryBlockB) const override;
+	virtual bool TestBasicOperation(const uint8* MemoryBlock, EBasicKeyOperation::Type Op) const override;
 };
diff --git a/Engine/Source/Runtime/AIModule/Classes/BehaviorTree/Blackboard/BlackboardKeyType_String.h b/Engine/Source/Runtime/AIModule/Classes/BehaviorTree/Blackboard/BlackboardKeyType_String.h
--- a/Engine/Source/Runtime/AIModule/Classes/BehaviorTree/Blackboard/BlackboardKeyType_String.h
+++ b/Engine/Source/Runtime/AIModule/Classes/BehaviorTree/Blackboard/BlackboardKeyType_String.h
@@ -11,7 +11,7 @@ class UBlackboardKeyType_String : public UBlackboardKeyType
 	static FString GetValue(const uint8* RawData);
 	static bool SetValue(uint8* RawData, const FString& Value);
 
-	virtual FString DescribeValue(const uint8* RawData) const OVERRIDE;
-	virtual EBlackboardCompare::Type Compare(const uint8* MemoryBlockA, const uint8* MemoryBlockB) const OVERRIDE;
-	virtual bool TestTextOperation(const uint8* MemoryBlock, ETextKeyOperation::Type Op, const FString& OtherString) const OVERRIDE;
+	virtual FString DescribeValue(const uint8* RawData) const override;
+	virtual EBlackboardCompare::Type Compare(const uint8* MemoryBlockA, const uint8* MemoryBlockB) const override;
+	virtual bool TestTextOperation(const uint8* MemoryBlock, ETextKeyOperation::Type Op, const FString& OtherString) const override;
 };
diff --git a/Engine/Source/Runtime/AIModule/Classes/BehaviorTree/Blackboard/BlackboardKeyType_Vector.h b/Engine/Source/Runtime/AIModule/Classes/BehaviorTree/Blackboard/BlackboardKeyType_Vector.h
--- a/Engine/Source/Runtime/AIModule/Classes/BehaviorTree/Blackboard/BlackboardKeyType_Vector.h
+++ b/Engine/Source/Runtime/AIModule/Classes/BehaviorTree/Blackboard/BlackboardKeyType_Vector.h
@@ -11,9 +11,9 @@ class UBlackboardKeyType_Vector : public UBlackboardKeyType
 	static FVector GetValue(const uint8* RawData);
 	static bool SetValue(uint8* RawData, const FVector& Value);
 
-	virtual void Initialize(uint8* RawData) const OVERRIDE;
-	virtual FString DescribeValue(const uint8* RawData) const OVERRIDE;
-	virtual bool GetLocation(const uint8* RawData, FVector& Location) const OVERRIDE;
-	virtual EBlackboardCompare::Type Compare(const uint8* MemoryBlockA, const uint8* MemoryBlockB) const OVERRIDE;
-	virtual bool TestBasicOperation(const uint8* MemoryBlock, EBasicKeyOperation::Type Op) const OVERRIDE;
+	virtual void Initialize(uint8* RawData) const override;
+	virtual FString DescribeValue(const uint8* RawData) const override;
+	virtual bool GetLocation(const uint8* RawData, FVector& Location) const override;
+	virtual EBlackboardCompare::Type Compare(const uint8* MemoryBlockA, const uint8* MemoryBlockB) const override;
+	virtual bool TestBasicOperation(const uint8* MemoryBlock, EBasicKeyOperation::Type Op) const override;
 };
diff --git a/Engine/Source/Runtime/AIModule/Classes/BehaviorTree/BlackboardComponent.h b/Engine/Source/Runtime/AIModule/Classes/BehaviorTree/BlackboardComponent.h
--- a/Engine/Source/Runtime/AIModule/Classes/BehaviorTree/BlackboardComponent.h
+++ b/Engine/Source/Runtime/AIModule/Classes/BehaviorTree/BlackboardComponent.h
@@ -34,7 +34,7 @@ class AIMODULE_API UBlackboardComponent : public UActorComponent
 	GENERATED_UCLASS_BODY()
 
 	/** BEGIN UActorComponent overrides */
-	virtual void InitializeComponent() OVERRIDE;
+	virtual void InitializeComponent() override;
 	/** END UActorComponent overrides */
 
 	/** @return name of key */
diff --git a/Engine/Source/Runtime/AIModule/Classes/BehaviorTree/BlackboardData.h b/Engine/Source/Runtime/AIModule/Classes/BehaviorTree/BlackboardData.h
--- a/Engine/Source/Runtime/AIModule/Classes/BehaviorTree/BlackboardData.h
+++ b/Engine/Source/Runtime/AIModule/Classes/BehaviorTree/BlackboardData.h
@@ -55,9 +55,9 @@ class AIMODULE_API UBlackboardData : public UDataAsset
 	/** @return key data */
 	const FBlackboardEntry* GetKey(FBlackboard::FKey KeyID) const;
 
-	virtual void PostLoad() OVERRIDE;
+	virtual void PostLoad() override;
 #if WITH_EDITOR
-	virtual void PostEditChangeProperty(struct FPropertyChangedEvent& PropertyChangedEvent) OVERRIDE;
+	virtual void PostEditChangeProperty(struct FPropertyChangedEvent& PropertyChangedEvent) override;
 #endif
 
 	/** @return true if blackboard keys are not conflicting with parent key chain */
diff --git a/Engine/Source/Runtime/AIModule/Classes/BehaviorTree/Composites/BTComposite_Sequence.h b/Engine/Source/Runtime/AIModule/Classes/BehaviorTree/Composites/BTComposite_Sequence.h
--- a/Engine/Source/Runtime/AIModule/Classes/BehaviorTree/Composites/BTComposite_Sequence.h
+++ b/Engine/Source/Runtime/AIModule/Classes/BehaviorTree/Composites/BTComposite_Sequence.h
@@ -12,7 +12,7 @@ class UBTComposite_Sequence : public UBTCompositeNode
 	int32 GetNextChildHandler(struct FBehaviorTreeSearchData& SearchData, int32 PrevChild, EBTNodeResult::Type LastResult) const;
 
 #if WITH_EDITOR
-	virtual bool CanAbortLowerPriority() const OVERRIDE;
+	virtual bool CanAbortLowerPriority() const override;
 	virtual FName GetNodeIconName() const override;
 #endif
 };
diff --git a/Engine/Source/Runtime/AIModule/Classes/BehaviorTree/Composites/BTComposite_SimpleParallel.h b/Engine/Source/Runtime/AIModule/Classes/BehaviorTree/Composites/BTComposite_SimpleParallel.h
--- a/Engine/Source/Runtime/AIModule/Classes/BehaviorTree/Composites/BTComposite_SimpleParallel.h
+++ b/Engine/Source/Runtime/AIModule/Classes/BehaviorTree/Composites/BTComposite_SimpleParallel.h
@@ -49,19 +49,19 @@ class UBTComposite_SimpleParallel : public UBTCompositeNode
 	/** handle child updates */
 	int32 GetNextChildHandler(struct FBehaviorTreeSearchData& SearchData, int32 PrevChild, EBTNodeResult::Type LastResult) const;
 
-	virtual void NotifyChildExecution(class UBehaviorTreeComponent* OwnerComp, uint8* NodeMemory, int32 ChildIdx, EBTNodeResult::Type& NodeResult) const OVERRIDE;
-	virtual void NotifyNodeDeactivation(struct FBehaviorTreeSearchData& SearchData, EBTNodeResult::Type& NodeResult) const OVERRIDE;
-	virtual bool CanPushSubtree(class UBehaviorTreeComponent* OwnerComp, uint8* NodeMemory, int32 ChildIdx) const OVERRIDE;
-	virtual uint16 GetInstanceMemorySize() const OVERRIDE;
-	virtual FString GetStaticDescription() const OVERRIDE;
-	virtual void DescribeRuntimeValues(const class UBehaviorTreeComponent* OwnerComp, uint8* NodeMemory, EBTDescriptionVerbosity::Type Verbosity, TArray<FString>& Values) const OVERRIDE;
+	virtual void NotifyChildExecution(class UBehaviorTreeComponent* OwnerComp, uint8* NodeMemory, int32 ChildIdx, EBTNodeResult::Type& NodeResult) const override;
+	virtual void NotifyNodeDeactivation(struct FBehaviorTreeSearchData& SearchData, EBTNodeResult::Type& NodeResult) const override;
+	virtual bool CanPushSubtree(class UBehaviorTreeComponent* OwnerComp, uint8* NodeMemory, int32 ChildIdx) const override;
+	virtual uint16 GetInstanceMemorySize() const override;
+	virtual FString GetStaticDescription() const override;
+	virtual void DescribeRuntimeValues(const class UBehaviorTreeComponent* OwnerComp, uint8* NodeMemory, EBTDescriptionVerbosity::Type Verbosity, TArray<FString>& Values) const override;
 
 	/** helper for showing values of EBTParallelMode enum */
 	static FString DescribeFinishMode(EBTParallelMode::Type Mode);
 
 #if WITH_EDITOR
-	virtual bool CanAbortLowerPriority() const OVERRIDE;
-	virtual bool CanAbortSelf() const OVERRIDE;
+	virtual bool CanAbortLowerPriority() const override;
+	virtual bool CanAbortSelf() const override;
 	virtual FName GetNodeIconName() const override;
 #endif // WITH_EDITOR
 };
diff --git a/Engine/Source/Runtime/AIModule/Classes/BehaviorTree/Decorators/BTDecorator_Blackboard.h b/Engine/Source/Runtime/AIModule/Classes/BehaviorTree/Decorators/BTDecorator_Blackboard.h
--- a/Engine/Source/Runtime/AIModule/Classes/BehaviorTree/Decorators/BTDecorator_Blackboard.h
+++ b/Engine/Source/Runtime/AIModule/Classes/BehaviorTree/Decorators/BTDecorator_Blackboard.h
@@ -14,10 +14,10 @@ class AIMODULE_API UBTDecorator_Blackboard : public UBTDecorator_BlackboardBase
 {
 	GENERATED_UCLASS_BODY()
 
-	virtual bool CalculateRawConditionValue(class UBehaviorTreeComponent* OwnerComp, uint8* NodeMemory) const OVERRIDE;
-	virtual void OnBlackboardChange(const class UBlackboardComponent* Blackboard, uint8 ChangedKeyID) OVERRIDE;
-	virtual void DescribeRuntimeValues(const class UBehaviorTreeComponent* OwnerComp, uint8* NodeMemory, EBTDescriptionVerbosity::Type Verbosity, TArray<FString>& Values) const OVERRIDE;
-	virtual FString GetStaticDescription() const OVERRIDE;
+	virtual bool CalculateRawConditionValue(class UBehaviorTreeComponent* OwnerComp, uint8* NodeMemory) const override;
+	virtual void OnBlackboardChange(const class UBlackboardComponent* Blackboard, uint8 ChangedKeyID) override;
+	virtual void DescribeRuntimeValues(const class UBehaviorTreeComponent* OwnerComp, uint8* NodeMemory, EBTDescriptionVerbosity::Type Verbosity, TArray<FString>& Values) const override;
+	virtual FString GetStaticDescription() const override;
 
 protected:
 
@@ -59,8 +59,8 @@ protected:
 	/** describe decorator and cache it */
 	virtual void BuildDescription();
 
-	virtual void PostEditChangeProperty(FPropertyChangedEvent& PropertyChangedEvent) OVERRIDE;
-	virtual void InitializeFromAsset(class UBehaviorTree* Asset) OVERRIDE;
+	virtual void PostEditChangeProperty(FPropertyChangedEvent& PropertyChangedEvent) override;
+	virtual void InitializeFromAsset(class UBehaviorTree* Asset) override;
 
 #endif
 
diff --git a/Engine/Source/Runtime/AIModule/Classes/BehaviorTree/Decorators/BTDecorator_BlackboardBase.h b/Engine/Source/Runtime/AIModule/Classes/BehaviorTree/Decorators/BTDecorator_BlackboardBase.h
--- a/Engine/Source/Runtime/AIModule/Classes/BehaviorTree/Decorators/BTDecorator_BlackboardBase.h
+++ b/Engine/Source/Runtime/AIModule/Classes/BehaviorTree/Decorators/BTDecorator_BlackboardBase.h
@@ -9,7 +9,7 @@ class AIMODULE_API UBTDecorator_BlackboardBase : public UBTDecorator
 	GENERATED_UCLASS_BODY()
 
 	/** initialize any asset related data */
-	virtual void InitializeFromAsset(class UBehaviorTree* Asset) OVERRIDE;
+	virtual void InitializeFromAsset(class UBehaviorTree* Asset) override;
 
 	/** notify about change in blackboard keys */
 	virtual void OnBlackboardChange(const class UBlackboardComponent* Blackboard, uint8 ChangedKeyID);
@@ -29,13 +29,13 @@ protected:
 
 	FOnBlackboardChange BBKeyObserver;
 
-	virtual void PostInitProperties() OVERRIDE;
+	virtual void PostInitProperties() override;
 
 	/** called when execution flow controller becomes active */
-	virtual void OnBecomeRelevant(UBehaviorTreeComponent* OwnerComp, uint8* NodeMemory) OVERRIDE;
+	virtual void OnBecomeRelevant(UBehaviorTreeComponent* OwnerComp, uint8* NodeMemory) override;
 
 	/** called when execution flow controller becomes inactive */
-	virtual void OnCeaseRelevant(UBehaviorTreeComponent* OwnerComp, uint8* NodeMemory) OVERRIDE;
+	virtual void OnCeaseRelevant(UBehaviorTreeComponent* OwnerComp, uint8* NodeMemory) override;
 };
 
 //////////////////////////////////////////////////////////////////////////
diff --git a/Engine/Source/Runtime/AIModule/Classes/BehaviorTree/Decorators/BTDecorator_BlueprintBase.h b/Engine/Source/Runtime/AIModule/Classes/BehaviorTree/Decorators/BTDecorator_BlueprintBase.h
--- a/Engine/Source/Runtime/AIModule/Classes/BehaviorTree/Decorators/BTDecorator_BlueprintBase.h
+++ b/Engine/Source/Runtime/AIModule/Classes/BehaviorTree/Decorators/BTDecorator_BlueprintBase.h
@@ -27,14 +27,14 @@ class AIMODULE_API UBTDecorator_BlueprintBase : public UBTDecorator
 	void InitializeProperties();
 
 	/** setup node name */
-	virtual void PostInitProperties() OVERRIDE;
+	virtual void PostInitProperties() override;
 
 	/** notify about changes in blackboard */
 	void OnBlackboardChange(const class UBlackboardComponent* Blackboard, uint8 ChangedKeyID);
 
-	virtual FString GetStaticDescription() const OVERRIDE;
-	virtual void DescribeRuntimeValues(const class UBehaviorTreeComponent* OwnerComp, uint8* NodeMemory, EBTDescriptionVerbosity::Type Verbosity, TArray<FString>& Values) const OVERRIDE;
-	virtual bool CalculateRawConditionValue(class UBehaviorTreeComponent* OwnerComp, uint8* NodeMemory) const OVERRIDE;
+	virtual FString GetStaticDescription() const override;
+	virtual void DescribeRuntimeValues(const class UBehaviorTreeComponent* OwnerComp, uint8* NodeMemory, EBTDescriptionVerbosity::Type Verbosity, TArray<FString>& Values) const override;
+	virtual bool CalculateRawConditionValue(class UBehaviorTreeComponent* OwnerComp, uint8* NodeMemory) const override;
 
 #if WITH_EDITOR
 	virtual FName GetNodeIconName() const override;
@@ -76,11 +76,11 @@ protected:
 
 	FOnBlackboardChange BBKeyObserver;
 
-	virtual void OnBecomeRelevant(UBehaviorTreeComponent* OwnerComp, uint8* NodeMemory) OVERRIDE;
-	virtual void OnCeaseRelevant(UBehaviorTreeComponent* OwnerComp, uint8* NodeMemory) OVERRIDE;
-	virtual void OnNodeActivation(struct FBehaviorTreeSearchData& SearchData) OVERRIDE;
-	virtual void OnNodeDeactivation(struct FBehaviorTreeSearchData& SearchData, EBTNodeResult::Type NodeResult) OVERRIDE;
-	virtual void TickNode(UBehaviorTreeComponent* OwnerComp, uint8* NodeMemory, float DeltaSeconds) OVERRIDE;
+	virtual void OnBecomeRelevant(UBehaviorTreeComponent* OwnerComp, uint8* NodeMemory) override;
+	virtual void OnCeaseRelevant(UBehaviorTreeComponent* OwnerComp, uint8* NodeMemory) override;
+	virtual void OnNodeActivation(struct FBehaviorTreeSearchData& SearchData) override;
+	virtual void OnNodeDeactivation(struct FBehaviorTreeSearchData& SearchData, EBTNodeResult::Type NodeResult) override;
+	virtual void TickNode(UBehaviorTreeComponent* OwnerComp, uint8* NodeMemory, float DeltaSeconds) override;
 
 	/** tick function */
 	UFUNCTION(BlueprintImplementableEvent)
diff --git a/Engine/Source/Runtime/AIModule/Classes/BehaviorTree/Decorators/BTDecorator_CompareBBEntries.h b/Engine/Source/Runtime/AIModule/Classes/BehaviorTree/Decorators/BTDecorator_CompareBBEntries.h
--- a/Engine/Source/Runtime/AIModule/Classes/BehaviorTree/Decorators/BTDecorator_CompareBBEntries.h
+++ b/Engine/Source/Runtime/AIModule/Classes/BehaviorTree/Decorators/BTDecorator_CompareBBEntries.h
@@ -37,16 +37,16 @@ protected:
 
 public:
 
-	virtual void PostInitProperties() OVERRIDE;
+	virtual void PostInitProperties() override;
 
-	virtual bool CalculateRawConditionValue(class UBehaviorTreeComponent* OwnerComp, uint8* NodeMemory) const OVERRIDE;
-	virtual FString GetStaticDescription() const OVERRIDE;
+	virtual bool CalculateRawConditionValue(class UBehaviorTreeComponent* OwnerComp, uint8* NodeMemory) const override;
+	virtual FString GetStaticDescription() const override;
 
-	virtual void InitializeFromAsset(class UBehaviorTree* Asset) OVERRIDE;
+	virtual void InitializeFromAsset(class UBehaviorTree* Asset) override;
 
 	virtual void OnBlackboardChange(const class UBlackboardComponent* Blackboard, uint8 ChangedKeyID);	
-	virtual void OnBecomeRelevant(UBehaviorTreeComponent* OwnerComp, uint8* NodeMemory) OVERRIDE;
-	virtual void OnCeaseRelevant(UBehaviorTreeComponent* OwnerComp, uint8* NodeMemory) OVERRIDE;
+	virtual void OnBecomeRelevant(UBehaviorTreeComponent* OwnerComp, uint8* NodeMemory) override;
+	virtual void OnCeaseRelevant(UBehaviorTreeComponent* OwnerComp, uint8* NodeMemory) override;
 
 #if WITH_EDITOR
 	virtual FName GetNodeIconName() const override;
diff --git a/Engine/Source/Runtime/AIModule/Classes/BehaviorTree/Decorators/BTDecorator_ConeCheck.h b/Engine/Source/Runtime/AIModule/Classes/BehaviorTree/Decorators/BTDecorator_ConeCheck.h
--- a/Engine/Source/Runtime/AIModule/Classes/BehaviorTree/Decorators/BTDecorator_ConeCheck.h
+++ b/Engine/Source/Runtime/AIModule/Classes/BehaviorTree/Decorators/BTDecorator_ConeCheck.h
@@ -34,10 +34,10 @@ class AIMODULE_API UBTDecorator_ConeCheck : public UBTDecorator
 	
 	float ConeHalfAngleDot;
 
-	virtual void InitializeFromAsset(class UBehaviorTree* Asset) OVERRIDE;
-	virtual uint16 GetInstanceMemorySize() const OVERRIDE;
-	virtual void DescribeRuntimeValues(const class UBehaviorTreeComponent* OwnerComp, uint8* NodeMemory, EBTDescriptionVerbosity::Type Verbosity, TArray<FString>& Values) const OVERRIDE;
-	virtual FString GetStaticDescription() const OVERRIDE;
+	virtual void InitializeFromAsset(class UBehaviorTree* Asset) override;
+	virtual uint16 GetInstanceMemorySize() const override;
+	virtual void DescribeRuntimeValues(const class UBehaviorTreeComponent* OwnerComp, uint8* NodeMemory, EBTDescriptionVerbosity::Type Verbosity, TArray<FString>& Values) const override;
+	virtual FString GetStaticDescription() const override;
 
 #if WITH_EDITOR
 	virtual FName GetNodeIconName() const override;
@@ -45,11 +45,11 @@ class AIMODULE_API UBTDecorator_ConeCheck : public UBTDecorator
 
 protected:
 
-	virtual bool CalculateRawConditionValue(class UBehaviorTreeComponent* OwnerComp, uint8* NodeMemory) const OVERRIDE;
-	virtual void OnBecomeRelevant(UBehaviorTreeComponent* OwnerComp, uint8* NodeMemory) OVERRIDE;
+	virtual bool CalculateRawConditionValue(class UBehaviorTreeComponent* OwnerComp, uint8* NodeMemory) const override;
+	virtual void OnBecomeRelevant(UBehaviorTreeComponent* OwnerComp, uint8* NodeMemory) override;
 	void OnBlackboardChange(const class UBlackboardComponent* Blackboard, uint8 ChangedKeyID);
 
-	virtual void TickNode(UBehaviorTreeComponent* OwnerComp, uint8* NodeMemory, float DeltaSeconds) OVERRIDE;
+	virtual void TickNode(UBehaviorTreeComponent* OwnerComp, uint8* NodeMemory, float DeltaSeconds) override;
 	
 	bool CalculateDirection(const UBlackboardComponent* BlackboardComp, const FBlackboardKeySelector& Origin, const FBlackboardKeySelector& End, FVector& Direction) const;
 
diff --git a/Engine/Source/Runtime/AIModule/Classes/BehaviorTree/Decorators/BTDecorator_Cooldown.h b/Engine/Source/Runtime/AIModule/Classes/BehaviorTree/Decorators/BTDecorator_Cooldown.h
--- a/Engine/Source/Runtime/AIModule/Classes/BehaviorTree/Decorators/BTDecorator_Cooldown.h
+++ b/Engine/Source/Runtime/AIModule/Classes/BehaviorTree/Decorators/BTDecorator_Cooldown.h
@@ -20,14 +20,14 @@ class AIMODULE_API UBTDecorator_Cooldown : public UBTDecorator
 	float CoolDownTime;
 
 	// Begin UObject Interface
-	virtual void PostLoad() OVERRIDE;
+	virtual void PostLoad() override;
 	// End UObject Interface
 
-	virtual bool CalculateRawConditionValue(class UBehaviorTreeComponent* OwnerComp, uint8* NodeMemory) const OVERRIDE;
-	virtual void InitializeMemory(class UBehaviorTreeComponent* OwnerComp, uint8* NodeMemory) const OVERRIDE;
-	virtual uint16 GetInstanceMemorySize() const OVERRIDE;
-	virtual void DescribeRuntimeValues(const class UBehaviorTreeComponent* OwnerComp, uint8* NodeMemory, EBTDescriptionVerbosity::Type Verbosity, TArray<FString>& Values) const OVERRIDE;
-	virtual FString GetStaticDescription() const OVERRIDE;
+	virtual bool CalculateRawConditionValue(class UBehaviorTreeComponent* OwnerComp, uint8* NodeMemory) const override;
+	virtual void InitializeMemory(class UBehaviorTreeComponent* OwnerComp, uint8* NodeMemory) const override;
+	virtual uint16 GetInstanceMemorySize() const override;
+	virtual void DescribeRuntimeValues(const class UBehaviorTreeComponent* OwnerComp, uint8* NodeMemory, EBTDescriptionVerbosity::Type Verbosity, TArray<FString>& Values) const override;
+	virtual FString GetStaticDescription() const override;
 
 #if WITH_EDITOR
 	virtual FName GetNodeIconName() const override;
@@ -35,6 +35,6 @@ class AIMODULE_API UBTDecorator_Cooldown : public UBTDecorator
 
 protected:
 
-	virtual void OnNodeDeactivation(struct FBehaviorTreeSearchData& SearchData, EBTNodeResult::Type NodeResult) OVERRIDE;
-	virtual void TickNode(class UBehaviorTreeComponent* OwnerComp, uint8* NodeMemory, float DeltaSeconds) OVERRIDE;
+	virtual void OnNodeDeactivation(struct FBehaviorTreeSearchData& SearchData, EBTNodeResult::Type NodeResult) override;
+	virtual void TickNode(class UBehaviorTreeComponent* OwnerComp, uint8* NodeMemory, float DeltaSeconds) override;
 };
diff --git a/Engine/Source/Runtime/AIModule/Classes/BehaviorTree/Decorators/BTDecorator_DoesPathExist.h b/Engine/Source/Runtime/AIModule/Classes/BehaviorTree/Decorators/BTDecorator_DoesPathExist.h
--- a/Engine/Source/Runtime/AIModule/Classes/BehaviorTree/Decorators/BTDecorator_DoesPathExist.h
+++ b/Engine/Source/Runtime/AIModule/Classes/BehaviorTree/Decorators/BTDecorator_DoesPathExist.h
@@ -43,9 +43,9 @@ public:
 	UPROPERTY(Category=Node, EditAnywhere)
 	TSubclassOf<class UNavigationQueryFilter> FilterClass;
 
-	virtual bool CalculateRawConditionValue(class UBehaviorTreeComponent* OwnerComp, uint8* NodeMemory) const OVERRIDE;
-	virtual FString GetStaticDescription() const OVERRIDE;
-	virtual void InitializeFromAsset(class UBehaviorTree* Asset) OVERRIDE;
+	virtual bool CalculateRawConditionValue(class UBehaviorTreeComponent* OwnerComp, uint8* NodeMemory) const override;
+	virtual FString GetStaticDescription() const override;
+	virtual void InitializeFromAsset(class UBehaviorTree* Asset) override;
 
 #if WITH_EDITOR
 	virtual FName GetNodeIconName() const override;
diff --git a/Engine/Source/Runtime/AIModule/Classes/BehaviorTree/Decorators/BTDecorator_KeepInCone.h b/Engine/Source/Runtime/AIModule/Classes/BehaviorTree/Decorators/BTDecorator_KeepInCone.h
--- a/Engine/Source/Runtime/AIModule/Classes/BehaviorTree/Decorators/BTDecorator_KeepInCone.h
+++ b/Engine/Source/Runtime/AIModule/Classes/BehaviorTree/Decorators/BTDecorator_KeepInCone.h
@@ -38,10 +38,10 @@ class AIMODULE_API UBTDecorator_KeepInCone : public UBTDecorator
 	
 	float ConeHalfAngleDot;
 
-	virtual void InitializeFromAsset(class UBehaviorTree* Asset) OVERRIDE;
-	virtual uint16 GetInstanceMemorySize() const OVERRIDE;
-	virtual void DescribeRuntimeValues(const class UBehaviorTreeComponent* OwnerComp, uint8* NodeMemory, EBTDescriptionVerbosity::Type Verbosity, TArray<FString>& Values) const OVERRIDE;
-	virtual FString GetStaticDescription() const OVERRIDE;
+	virtual void InitializeFromAsset(class UBehaviorTree* Asset) override;
+	virtual uint16 GetInstanceMemorySize() const override;
+	virtual void DescribeRuntimeValues(const class UBehaviorTreeComponent* OwnerComp, uint8* NodeMemory, EBTDescriptionVerbosity::Type Verbosity, TArray<FString>& Values) const override;
+	virtual FString GetStaticDescription() const override;
 
 #if WITH_EDITOR
 	virtual FName GetNodeIconName() const override;
@@ -49,8 +49,8 @@ class AIMODULE_API UBTDecorator_KeepInCone : public UBTDecorator
 
 protected:
 
-	virtual void OnBecomeRelevant(UBehaviorTreeComponent* OwnerComp, uint8* NodeMemory) OVERRIDE;
-	virtual void TickNode(UBehaviorTreeComponent* OwnerComp, uint8* NodeMemory, float DeltaSeconds) OVERRIDE;
+	virtual void OnBecomeRelevant(UBehaviorTreeComponent* OwnerComp, uint8* NodeMemory) override;
+	virtual void TickNode(UBehaviorTreeComponent* OwnerComp, uint8* NodeMemory, float DeltaSeconds) override;
 
 	bool CalculateCurrentDirection(const UBehaviorTreeComponent* OwnerComp, FVector& Direction) const;
 };
diff --git a/Engine/Source/Runtime/AIModule/Classes/BehaviorTree/Decorators/BTDecorator_Loop.h b/Engine/Source/Runtime/AIModule/Classes/BehaviorTree/Decorators/BTDecorator_Loop.h
--- a/Engine/Source/Runtime/AIModule/Classes/BehaviorTree/Decorators/BTDecorator_Loop.h
+++ b/Engine/Source/Runtime/AIModule/Classes/BehaviorTree/Decorators/BTDecorator_Loop.h
@@ -21,9 +21,9 @@ class AIMODULE_API UBTDecorator_Loop : public UBTDecorator
 	UPROPERTY(Category=Decorator, EditAnywhere)
 	bool bInfiniteLoop;
 
-	virtual uint16 GetInstanceMemorySize() const OVERRIDE;
-	virtual void DescribeRuntimeValues(const class UBehaviorTreeComponent* OwnerComp, uint8* NodeMemory, EBTDescriptionVerbosity::Type Verbosity, TArray<FString>& Values) const OVERRIDE;
-	virtual FString GetStaticDescription() const OVERRIDE;
+	virtual uint16 GetInstanceMemorySize() const override;
+	virtual void DescribeRuntimeValues(const class UBehaviorTreeComponent* OwnerComp, uint8* NodeMemory, EBTDescriptionVerbosity::Type Verbosity, TArray<FString>& Values) const override;
+	virtual FString GetStaticDescription() const override;
 
 #if WITH_EDITOR
 	virtual FName GetNodeIconName() const override;
@@ -31,5 +31,5 @@ class AIMODULE_API UBTDecorator_Loop : public UBTDecorator
 
 protected:
 
-	virtual void OnNodeActivation(struct FBehaviorTreeSearchData& SearchData) OVERRIDE;
+	virtual void OnNodeActivation(struct FBehaviorTreeSearchData& SearchData) override;
 };
diff --git a/Engine/Source/Runtime/AIModule/Classes/BehaviorTree/Decorators/BTDecorator_ReachedMoveGoal.h b/Engine/Source/Runtime/AIModule/Classes/BehaviorTree/Decorators/BTDecorator_ReachedMoveGoal.h
--- a/Engine/Source/Runtime/AIModule/Classes/BehaviorTree/Decorators/BTDecorator_ReachedMoveGoal.h
+++ b/Engine/Source/Runtime/AIModule/Classes/BehaviorTree/Decorators/BTDecorator_ReachedMoveGoal.h
@@ -9,7 +9,7 @@ class AIMODULE_API UBTDecorator_ReachedMoveGoal : public UBTDecorator
 {
 	GENERATED_UCLASS_BODY()
 
-	virtual bool CalculateRawConditionValue(class UBehaviorTreeComponent* OwnerComp, uint8* NodeMemory) const OVERRIDE;
+	virtual bool CalculateRawConditionValue(class UBehaviorTreeComponent* OwnerComp, uint8* NodeMemory) const override;
 
 #if WITH_EDITOR
 	virtual FName GetNodeIconName() const override;
diff --git a/Engine/Source/Runtime/AIModule/Classes/BehaviorTree/Decorators/BTDecorator_TimeLimit.h b/Engine/Source/Runtime/AIModule/Classes/BehaviorTree/Decorators/BTDecorator_TimeLimit.h
--- a/Engine/Source/Runtime/AIModule/Classes/BehaviorTree/Decorators/BTDecorator_TimeLimit.h
+++ b/Engine/Source/Runtime/AIModule/Classes/BehaviorTree/Decorators/BTDecorator_TimeLimit.h
@@ -12,8 +12,8 @@ class AIMODULE_API UBTDecorator_TimeLimit : public UBTDecorator
 	UPROPERTY(Category=Decorator, EditAnywhere)
 	float TimeLimit;
 
-	virtual void DescribeRuntimeValues(const class UBehaviorTreeComponent* OwnerComp, uint8* NodeMemory, EBTDescriptionVerbosity::Type Verbosity, TArray<FString>& Values) const OVERRIDE;
-	virtual FString GetStaticDescription() const OVERRIDE;
+	virtual void DescribeRuntimeValues(const class UBehaviorTreeComponent* OwnerComp, uint8* NodeMemory, EBTDescriptionVerbosity::Type Verbosity, TArray<FString>& Values) const override;
+	virtual FString GetStaticDescription() const override;
 
 #if WITH_EDITOR
 	virtual FName GetNodeIconName() const override;
@@ -21,6 +21,6 @@ class AIMODULE_API UBTDecorator_TimeLimit : public UBTDecorator
 
 protected:
 
-	virtual void OnBecomeRelevant(UBehaviorTreeComponent* OwnerComp, uint8* NodeMemory) OVERRIDE;
-	virtual void TickNode(UBehaviorTreeComponent* OwnerComp, uint8* NodeMemory, float DeltaSeconds) OVERRIDE;
+	virtual void OnBecomeRelevant(UBehaviorTreeComponent* OwnerComp, uint8* NodeMemory) override;
+	virtual void TickNode(UBehaviorTreeComponent* OwnerComp, uint8* NodeMemory, float DeltaSeconds) override;
 };
diff --git a/Engine/Source/Runtime/AIModule/Classes/BehaviorTree/Services/BTService_BlackboardBase.h b/Engine/Source/Runtime/AIModule/Classes/BehaviorTree/Services/BTService_BlackboardBase.h
--- a/Engine/Source/Runtime/AIModule/Classes/BehaviorTree/Services/BTService_BlackboardBase.h
+++ b/Engine/Source/Runtime/AIModule/Classes/BehaviorTree/Services/BTService_BlackboardBase.h
@@ -9,7 +9,7 @@ class AIMODULE_API UBTService_BlackboardBase : public UBTService
 	GENERATED_UCLASS_BODY()
 
 	/** initialize any asset related data */
-	virtual void InitializeFromAsset(class UBehaviorTree* Asset) OVERRIDE;
+	virtual void InitializeFromAsset(class UBehaviorTree* Asset) override;
 
 	/** get name of selected blackboard key */
 	FName GetSelectedBlackboardKey() const;
diff --git a/Engine/Source/Runtime/AIModule/Classes/BehaviorTree/Services/BTService_BlueprintBase.h b/Engine/Source/Runtime/AIModule/Classes/BehaviorTree/Services/BTService_BlueprintBase.h
--- a/Engine/Source/Runtime/AIModule/Classes/BehaviorTree/Services/BTService_BlueprintBase.h
+++ b/Engine/Source/Runtime/AIModule/Classes/BehaviorTree/Services/BTService_BlueprintBase.h
@@ -17,10 +17,10 @@ class AIMODULE_API UBTService_BlueprintBase : public UBTService
 {
 	GENERATED_UCLASS_BODY()
 
-	virtual void PostInitProperties() OVERRIDE;
+	virtual void PostInitProperties() override;
 
-	virtual FString GetStaticDescription() const OVERRIDE;
-	virtual void DescribeRuntimeValues(const class UBehaviorTreeComponent* OwnerComp, uint8* NodeMemory, EBTDescriptionVerbosity::Type Verbosity, TArray<FString>& Values) const OVERRIDE;
+	virtual FString GetStaticDescription() const override;
+	virtual void DescribeRuntimeValues(const class UBehaviorTreeComponent* OwnerComp, uint8* NodeMemory, EBTDescriptionVerbosity::Type Verbosity, TArray<FString>& Values) const override;
 
 #if WITH_EDITOR
 	virtual bool UsesBlueprint() const override;
@@ -44,9 +44,9 @@ protected:
 	/** set if ReceiveDeactivation is implemented by blueprint */
 	uint32 bImplementsReceiveDeactivation : 1;
 
-	virtual void OnBecomeRelevant(class UBehaviorTreeComponent* OwnerComp, uint8* NodeMemory) OVERRIDE;
-	virtual void OnCeaseRelevant(class UBehaviorTreeComponent* OwnerComp, uint8* NodeMemory) OVERRIDE;
-	virtual void TickNode(class UBehaviorTreeComponent* OwnerComp, uint8* NodeMemory, float DeltaSeconds) OVERRIDE;
+	virtual void OnBecomeRelevant(class UBehaviorTreeComponent* OwnerComp, uint8* NodeMemory) override;
+	virtual void OnCeaseRelevant(class UBehaviorTreeComponent* OwnerComp, uint8* NodeMemory) override;
+	virtual void TickNode(class UBehaviorTreeComponent* OwnerComp, uint8* NodeMemory, float DeltaSeconds) override;
 
 	/** tick function */
 	UFUNCTION(BlueprintImplementableEvent)
diff --git a/Engine/Source/Runtime/AIModule/Classes/BehaviorTree/Services/BTService_DefaultFocus.h b/Engine/Source/Runtime/AIModule/Classes/BehaviorTree/Services/BTService_DefaultFocus.h
--- a/Engine/Source/Runtime/AIModule/Classes/BehaviorTree/Services/BTService_DefaultFocus.h
+++ b/Engine/Source/Runtime/AIModule/Classes/BehaviorTree/Services/BTService_DefaultFocus.h
@@ -21,12 +21,12 @@ class UBTService_DefaultFocus : public UBTService_BlackboardBase
 {
 	GENERATED_UCLASS_BODY()
 
-	virtual uint16 GetInstanceMemorySize() const OVERRIDE;
-	virtual void OnBecomeRelevant(class UBehaviorTreeComponent* OwnerComp, uint8* NodeMemory) OVERRIDE;
-	virtual void OnCeaseRelevant(class UBehaviorTreeComponent* OwnerComp, uint8* NodeMemory) OVERRIDE;
+	virtual uint16 GetInstanceMemorySize() const override;
+	virtual void OnBecomeRelevant(class UBehaviorTreeComponent* OwnerComp, uint8* NodeMemory) override;
+	virtual void OnCeaseRelevant(class UBehaviorTreeComponent* OwnerComp, uint8* NodeMemory) override;
 
-	virtual FString GetStaticDescription() const OVERRIDE;
-	virtual void DescribeRuntimeValues(const class UBehaviorTreeComponent* OwnerComp, uint8* NodeMemory, EBTDescriptionVerbosity::Type Verbosity, TArray<FString>& Values) const OVERRIDE;
+	virtual FString GetStaticDescription() const override;
+	virtual void DescribeRuntimeValues(const class UBehaviorTreeComponent* OwnerComp, uint8* NodeMemory, EBTDescriptionVerbosity::Type Verbosity, TArray<FString>& Values) const override;
 
 #if WITH_EDITOR
 	virtual FName GetNodeIconName() const override;
diff --git a/Engine/Source/Runtime/AIModule/Classes/BehaviorTree/Tasks/BTTask_BlackboardBase.h b/Engine/Source/Runtime/AIModule/Classes/BehaviorTree/Tasks/BTTask_BlackboardBase.h
--- a/Engine/Source/Runtime/AIModule/Classes/BehaviorTree/Tasks/BTTask_BlackboardBase.h
+++ b/Engine/Source/Runtime/AIModule/Classes/BehaviorTree/Tasks/BTTask_BlackboardBase.h
@@ -11,7 +11,7 @@ class AIMODULE_API UBTTask_BlackboardBase : public UBTTaskNode
 	GENERATED_UCLASS_BODY()
 
 	/** initialize any asset related data */
-	virtual void InitializeFromAsset(class UBehaviorTree* Asset) OVERRIDE;
+	virtual void InitializeFromAsset(class UBehaviorTree* Asset) override;
 
 	/** get name of selected blackboard key */
 	FName GetSelectedBlackboardKey() const;
diff --git a/Engine/Source/Runtime/AIModule/Classes/BehaviorTree/Tasks/BTTask_BlueprintBase.h b/Engine/Source/Runtime/AIModule/Classes/BehaviorTree/Tasks/BTTask_BlueprintBase.h
--- a/Engine/Source/Runtime/AIModule/Classes/BehaviorTree/Tasks/BTTask_BlueprintBase.h
+++ b/Engine/Source/Runtime/AIModule/Classes/BehaviorTree/Tasks/BTTask_BlueprintBase.h
@@ -17,13 +17,13 @@ class AIMODULE_API UBTTask_BlueprintBase : public UBTTaskNode
 	GENERATED_UCLASS_BODY()
 
 	/** setup node name */
-	virtual void PostInitProperties() OVERRIDE;
+	virtual void PostInitProperties() override;
 
-	virtual EBTNodeResult::Type ExecuteTask(class UBehaviorTreeComponent* OwnerComp, uint8* NodeMemory) OVERRIDE;
-	virtual EBTNodeResult::Type AbortTask(class UBehaviorTreeComponent* OwnerComp, uint8* NodeMemory) OVERRIDE;
+	virtual EBTNodeResult::Type ExecuteTask(class UBehaviorTreeComponent* OwnerComp, uint8* NodeMemory) override;
+	virtual EBTNodeResult::Type AbortTask(class UBehaviorTreeComponent* OwnerComp, uint8* NodeMemory) override;
 
-	virtual FString GetStaticDescription() const OVERRIDE;
-	virtual void DescribeRuntimeValues(const class UBehaviorTreeComponent* OwnerComp, uint8* NodeMemory, EBTDescriptionVerbosity::Type Verbosity, TArray<FString>& Values) const OVERRIDE;
+	virtual FString GetStaticDescription() const override;
+	virtual void DescribeRuntimeValues(const class UBehaviorTreeComponent* OwnerComp, uint8* NodeMemory, EBTDescriptionVerbosity::Type Verbosity, TArray<FString>& Values) const override;
 
 #if WITH_EDITOR
 	virtual bool UsesBlueprint() const override;
@@ -87,7 +87,7 @@ protected:
 	bool IsTaskExecuting() const;
 
 	/** ticks this task */
-	virtual void TickTask(class UBehaviorTreeComponent* OwnerComp, uint8* NodeMemory, float DeltaSeconds) OVERRIDE;
+	virtual void TickTask(class UBehaviorTreeComponent* OwnerComp, uint8* NodeMemory, float DeltaSeconds) override;
 
 	friend class FBehaviorBlueprintDetails;
 };
diff --git a/Engine/Source/Runtime/AIModule/Classes/BehaviorTree/Tasks/BTTask_MakeNoise.h b/Engine/Source/Runtime/AIModule/Classes/BehaviorTree/Tasks/BTTask_MakeNoise.h
--- a/Engine/Source/Runtime/AIModule/Classes/BehaviorTree/Tasks/BTTask_MakeNoise.h
+++ b/Engine/Source/Runtime/AIModule/Classes/BehaviorTree/Tasks/BTTask_MakeNoise.h
@@ -12,7 +12,7 @@ class UBTTask_MakeNoise : public UBTTaskNode
 	UPROPERTY(Category=Node, EditAnywhere, meta=(ClampMin = "0", ClampMax = "1", UIMin = "0", UIMax = "1"))
 	float Loudnes;
 
-	virtual EBTNodeResult::Type ExecuteTask(UBehaviorTreeComponent* OwnerComp, uint8* NodeMemory) OVERRIDE;
+	virtual EBTNodeResult::Type ExecuteTask(UBehaviorTreeComponent* OwnerComp, uint8* NodeMemory) override;
 
 #if WITH_EDITOR
 	virtual FName GetNodeIconName() const override;
diff --git a/Engine/Source/Runtime/AIModule/Classes/BehaviorTree/Tasks/BTTask_MoveDirectlyToward.h b/Engine/Source/Runtime/AIModule/Classes/BehaviorTree/Tasks/BTTask_MoveDirectlyToward.h
--- a/Engine/Source/Runtime/AIModule/Classes/BehaviorTree/Tasks/BTTask_MoveDirectlyToward.h
+++ b/Engine/Source/Runtime/AIModule/Classes/BehaviorTree/Tasks/BTTask_MoveDirectlyToward.h
@@ -31,11 +31,11 @@ class AIMODULE_API UBTTask_MoveDirectlyToward : public UBTTask_BlackboardBase
 	UPROPERTY(Category=Node, EditAnywhere)
 	uint32 bAllowStrafe : 1;
 
-	virtual EBTNodeResult::Type ExecuteTask(class UBehaviorTreeComponent* OwnerComp, uint8* NodeMemory) OVERRIDE;
-	virtual EBTNodeResult::Type AbortTask(class UBehaviorTreeComponent* OwnerComp, uint8* NodeMemory) OVERRIDE;
-	virtual uint16 GetInstanceMemorySize() const OVERRIDE;
-	virtual void DescribeRuntimeValues(const class UBehaviorTreeComponent* OwnerComp, uint8* NodeMemory, EBTDescriptionVerbosity::Type Verbosity, TArray<FString>& Values) const OVERRIDE;
-	virtual FString GetStaticDescription() const OVERRIDE;
+	virtual EBTNodeResult::Type ExecuteTask(class UBehaviorTreeComponent* OwnerComp, uint8* NodeMemory) override;
+	virtual EBTNodeResult::Type AbortTask(class UBehaviorTreeComponent* OwnerComp, uint8* NodeMemory) override;
+	virtual uint16 GetInstanceMemorySize() const override;
+	virtual void DescribeRuntimeValues(const class UBehaviorTreeComponent* OwnerComp, uint8* NodeMemory, EBTDescriptionVerbosity::Type Verbosity, TArray<FString>& Values) const override;
+	virtual FString GetStaticDescription() const override;
 
 #if WITH_EDITOR
 	virtual FName GetNodeIconName() const override;
diff --git a/Engine/Source/Runtime/AIModule/Classes/BehaviorTree/Tasks/BTTask_MoveTo.h b/Engine/Source/Runtime/AIModule/Classes/BehaviorTree/Tasks/BTTask_MoveTo.h
--- a/Engine/Source/Runtime/AIModule/Classes/BehaviorTree/Tasks/BTTask_MoveTo.h
+++ b/Engine/Source/Runtime/AIModule/Classes/BehaviorTree/Tasks/BTTask_MoveTo.h
@@ -27,15 +27,15 @@ class AIMODULE_API UBTTask_MoveTo : public UBTTask_BlackboardBase
 	UPROPERTY(Category=Node, EditAnywhere)
 	uint32 bAllowStrafe : 1;
 	
-	virtual EBTNodeResult::Type ExecuteTask(class UBehaviorTreeComponent* OwnerComp, uint8* NodeMemory) OVERRIDE;
-	virtual EBTNodeResult::Type AbortTask(class UBehaviorTreeComponent* OwnerComp, uint8* NodeMemory) OVERRIDE;
-	virtual void TickTask(class UBehaviorTreeComponent* OwnerComp, uint8* NodeMemory, float DeltaSeconds) OVERRIDE;
-	virtual uint16 GetInstanceMemorySize() const OVERRIDE;
+	virtual EBTNodeResult::Type ExecuteTask(class UBehaviorTreeComponent* OwnerComp, uint8* NodeMemory) override;
+	virtual EBTNodeResult::Type AbortTask(class UBehaviorTreeComponent* OwnerComp, uint8* NodeMemory) override;
+	virtual void TickTask(class UBehaviorTreeComponent* OwnerComp, uint8* NodeMemory, float DeltaSeconds) override;
+	virtual uint16 GetInstanceMemorySize() const override;
 
-	virtual void OnMessage(class UBehaviorTreeComponent* OwnerComp, uint8* NodeMemory, FName Message, int32 RequestID, bool bSuccess) OVERRIDE;
+	virtual void OnMessage(class UBehaviorTreeComponent* OwnerComp, uint8* NodeMemory, FName Message, int32 RequestID, bool bSuccess) override;
 
-	virtual void DescribeRuntimeValues(const class UBehaviorTreeComponent* OwnerComp, uint8* NodeMemory, EBTDescriptionVerbosity::Type Verbosity, TArray<FString>& Values) const OVERRIDE;
-	virtual FString GetStaticDescription() const OVERRIDE;
+	virtual void DescribeRuntimeValues(const class UBehaviorTreeComponent* OwnerComp, uint8* NodeMemory, EBTDescriptionVerbosity::Type Verbosity, TArray<FString>& Values) const override;
+	virtual FString GetStaticDescription() const override;
 
 #if WITH_EDITOR
 	virtual FName GetNodeIconName() const override;
diff --git a/Engine/Source/Runtime/AIModule/Classes/BehaviorTree/Tasks/BTTask_PlaySound.h b/Engine/Source/Runtime/AIModule/Classes/BehaviorTree/Tasks/BTTask_PlaySound.h
--- a/Engine/Source/Runtime/AIModule/Classes/BehaviorTree/Tasks/BTTask_PlaySound.h
+++ b/Engine/Source/Runtime/AIModule/Classes/BehaviorTree/Tasks/BTTask_PlaySound.h
@@ -15,8 +15,8 @@ class UBTTask_PlaySound : public UBTTaskNode
 	UPROPERTY(Category=Node, EditAnywhere)
 	USoundCue* SoundToPlay;
 
-	virtual EBTNodeResult::Type ExecuteTask(UBehaviorTreeComponent* OwnerComp, uint8* NodeMemory) OVERRIDE;
-	virtual FString GetStaticDescription() const OVERRIDE;
+	virtual EBTNodeResult::Type ExecuteTask(UBehaviorTreeComponent* OwnerComp, uint8* NodeMemory) override;
+	virtual FString GetStaticDescription() const override;
 
 #if WITH_EDITOR
 	virtual FName GetNodeIconName() const override;
diff --git a/Engine/Source/Runtime/AIModule/Classes/BehaviorTree/Tasks/BTTask_RunBehavior.h b/Engine/Source/Runtime/AIModule/Classes/BehaviorTree/Tasks/BTTask_RunBehavior.h
--- a/Engine/Source/Runtime/AIModule/Classes/BehaviorTree/Tasks/BTTask_RunBehavior.h
+++ b/Engine/Source/Runtime/AIModule/Classes/BehaviorTree/Tasks/BTTask_RunBehavior.h
@@ -23,9 +23,9 @@ class UBTTask_RunBehavior : public UBTTaskNode
 {
 	GENERATED_UCLASS_BODY()
 
-	virtual EBTNodeResult::Type ExecuteTask(UBehaviorTreeComponent* OwnerComp, uint8* NodeMemory) OVERRIDE;
-	virtual uint16 GetInstanceMemorySize() const OVERRIDE;
-	virtual FString GetStaticDescription() const OVERRIDE;
+	virtual EBTNodeResult::Type ExecuteTask(UBehaviorTreeComponent* OwnerComp, uint8* NodeMemory) override;
+	virtual uint16 GetInstanceMemorySize() const override;
+	virtual FString GetStaticDescription() const override;
 
 #if WITH_EDITOR
 	virtual FName GetNodeIconName() const override;
diff --git a/Engine/Source/Runtime/AIModule/Classes/BehaviorTree/Tasks/BTTask_RunEQSQuery.h b/Engine/Source/Runtime/AIModule/Classes/BehaviorTree/Tasks/BTTask_RunEQSQuery.h
--- a/Engine/Source/Runtime/AIModule/Classes/BehaviorTree/Tasks/BTTask_RunEQSQuery.h
+++ b/Engine/Source/Runtime/AIModule/Classes/BehaviorTree/Tasks/BTTask_RunEQSQuery.h
@@ -23,19 +23,19 @@ class UBTTask_RunEQSQuery : public UBTTask_BlackboardBase
 	UPROPERTY(Category=Node, EditAnywhere)
 	TArray<struct FEnvNamedValue> QueryParams;
 
-	virtual EBTNodeResult::Type ExecuteTask(class UBehaviorTreeComponent* OwnerComp, uint8* NodeMemory) OVERRIDE;
-	virtual EBTNodeResult::Type AbortTask(class UBehaviorTreeComponent* OwnerComp, uint8* NodeMemory) OVERRIDE;
+	virtual EBTNodeResult::Type ExecuteTask(class UBehaviorTreeComponent* OwnerComp, uint8* NodeMemory) override;
+	virtual EBTNodeResult::Type AbortTask(class UBehaviorTreeComponent* OwnerComp, uint8* NodeMemory) override;
 
-	virtual void DescribeRuntimeValues(const class UBehaviorTreeComponent* OwnerComp, uint8* NodeMemory, EBTDescriptionVerbosity::Type Verbosity, TArray<FString>& Values) const OVERRIDE;
-	virtual FString GetStaticDescription() const OVERRIDE;
-	virtual uint16 GetInstanceMemorySize() const OVERRIDE;
+	virtual void DescribeRuntimeValues(const class UBehaviorTreeComponent* OwnerComp, uint8* NodeMemory, EBTDescriptionVerbosity::Type Verbosity, TArray<FString>& Values) const override;
+	virtual FString GetStaticDescription() const override;
+	virtual uint16 GetInstanceMemorySize() const override;
 
 	/** finish task */
 	void OnQueryFinished(TSharedPtr<struct FEnvQueryResult> Result);
 
 #if WITH_EDITOR
 	/** prepare query params */
-	virtual void PostEditChangeProperty(struct FPropertyChangedEvent& PropertyChangedEvent) OVERRIDE;
+	virtual void PostEditChangeProperty(struct FPropertyChangedEvent& PropertyChangedEvent) override;
 	virtual FName GetNodeIconName() const override;
 #endif
 
diff --git a/Engine/Source/Runtime/AIModule/Classes/BehaviorTree/Tasks/BTTask_Wait.h b/Engine/Source/Runtime/AIModule/Classes/BehaviorTree/Tasks/BTTask_Wait.h
--- a/Engine/Source/Runtime/AIModule/Classes/BehaviorTree/Tasks/BTTask_Wait.h
+++ b/Engine/Source/Runtime/AIModule/Classes/BehaviorTree/Tasks/BTTask_Wait.h
@@ -18,10 +18,10 @@ class UBTTask_Wait : public UBTTaskNode
 	UPROPERTY(Category=Node, EditAnywhere)
 	float WaitTime;
 
-	virtual EBTNodeResult::Type ExecuteTask(class UBehaviorTreeComponent* OwnerComp, uint8* NodeMemory) OVERRIDE;
-	virtual uint16 GetInstanceMemorySize() const OVERRIDE;
-	virtual void DescribeRuntimeValues(const class UBehaviorTreeComponent* OwnerComp, uint8* NodeMemory, EBTDescriptionVerbosity::Type Verbosity, TArray<FString>& Values) const OVERRIDE;
-	virtual FString GetStaticDescription() const OVERRIDE;
+	virtual EBTNodeResult::Type ExecuteTask(class UBehaviorTreeComponent* OwnerComp, uint8* NodeMemory) override;
+	virtual uint16 GetInstanceMemorySize() const override;
+	virtual void DescribeRuntimeValues(const class UBehaviorTreeComponent* OwnerComp, uint8* NodeMemory, EBTDescriptionVerbosity::Type Verbosity, TArray<FString>& Values) const override;
+	virtual FString GetStaticDescription() const override;
 
 #if WITH_EDITOR
 	virtual FName GetNodeIconName() const override;
@@ -29,5 +29,5 @@ class UBTTask_Wait : public UBTTaskNode
 
 protected:
 
-	virtual void TickTask(class UBehaviorTreeComponent* OwnerComp, uint8* NodeMemory, float DeltaSeconds) OVERRIDE;
+	virtual void TickTask(class UBehaviorTreeComponent* OwnerComp, uint8* NodeMemory, float DeltaSeconds) override;
 };
diff --git a/Engine/Source/Runtime/AIModule/Classes/Blueprint/AIAsyncTaskBlueprintProxy.h b/Engine/Source/Runtime/AIModule/Classes/Blueprint/AIAsyncTaskBlueprintProxy.h
--- a/Engine/Source/Runtime/AIModule/Classes/Blueprint/AIAsyncTaskBlueprintProxy.h
+++ b/Engine/Source/Runtime/AIModule/Classes/Blueprint/AIAsyncTaskBlueprintProxy.h
@@ -25,7 +25,7 @@ public:
 	void OnNoPath();
 
 	// Begin UObject interface
-	virtual void BeginDestroy() OVERRIDE;
+	virtual void BeginDestroy() override;
 	// End UObject interface
 
 	TWeakObjectPtr<class AAIController> AIController;
diff --git a/Engine/Source/Runtime/AIModule/Classes/BrainComponent.h b/Engine/Source/Runtime/AIModule/Classes/BrainComponent.h
--- a/Engine/Source/Runtime/AIModule/Classes/BrainComponent.h
+++ b/Engine/Source/Runtime/AIModule/Classes/BrainComponent.h
@@ -126,16 +126,16 @@ public:
 #endif // ENABLE_VISUAL_LOG
 	
 	// IAIResourceInterface begin
-	virtual void LockResource(EAILockSource::Type LockSource) OVERRIDE;
-	virtual void ClearResourceLock(EAILockSource::Type LockSource) OVERRIDE;
-	virtual void ForceUnlockResource() OVERRIDE;
-	virtual bool IsResourceLocked() const OVERRIDE;
+	virtual void LockResource(EAILockSource::Type LockSource) override;
+	virtual void ClearResourceLock(EAILockSource::Type LockSource) override;
+	virtual void ForceUnlockResource() override;
+	virtual bool IsResourceLocked() const override;
 	// IAIResourceInterface end
 	
 	virtual void HandleMessage(const FAIMessage& Message);
 	
 	/** BEGIN UActorComponent overrides */
-	virtual void InitializeComponent() OVERRIDE;
+	virtual void InitializeComponent() override;
 	virtual void TickComponent(float DeltaTime, enum ELevelTick TickType, FActorComponentTickFunction *ThisTickFunction);
 	/** END UActorComponent overrides */
 
diff --git a/Engine/Source/Runtime/AIModule/Classes/EnvironmentQuery/Contexts/EnvQueryContext_BlueprintBase.h b/Engine/Source/Runtime/AIModule/Classes/EnvironmentQuery/Contexts/EnvQueryContext_BlueprintBase.h
--- a/Engine/Source/Runtime/AIModule/Classes/EnvironmentQuery/Contexts/EnvQueryContext_BlueprintBase.h
+++ b/Engine/Source/Runtime/AIModule/Classes/EnvironmentQuery/Contexts/EnvQueryContext_BlueprintBase.h
@@ -21,7 +21,7 @@ class UEnvQueryContext_BlueprintBase : public UEnvQueryContext
 
 	ECallMode CallMode;
 
-	virtual void ProvideContext(struct FEnvQueryInstance& QueryInstance, struct FEnvQueryContextData& ContextData) const OVERRIDE;
+	virtual void ProvideContext(struct FEnvQueryInstance& QueryInstance, struct FEnvQueryContextData& ContextData) const override;
 
 	UFUNCTION(BlueprintImplementableEvent)
 	virtual void ProvideSingleActor(AActor* QuerierActor, AActor*& ResultingActor) const;
diff --git a/Engine/Source/Runtime/AIModule/Classes/EnvironmentQuery/Contexts/EnvQueryContext_Querier.h b/Engine/Source/Runtime/AIModule/Classes/EnvironmentQuery/Contexts/EnvQueryContext_Querier.h
--- a/Engine/Source/Runtime/AIModule/Classes/EnvironmentQuery/Contexts/EnvQueryContext_Querier.h
+++ b/Engine/Source/Runtime/AIModule/Classes/EnvironmentQuery/Contexts/EnvQueryContext_Querier.h
@@ -9,5 +9,5 @@ class UEnvQueryContext_Querier : public UEnvQueryContext
 {
 	GENERATED_UCLASS_BODY()
  
-	virtual void ProvideContext(struct FEnvQueryInstance& QueryInstance, struct FEnvQueryContextData& ContextData) const OVERRIDE;
+	virtual void ProvideContext(struct FEnvQueryInstance& QueryInstance, struct FEnvQueryContextData& ContextData) const override;
 };
diff --git a/Engine/Source/Runtime/AIModule/Classes/EnvironmentQuery/EQSRenderingComponent.h b/Engine/Source/Runtime/AIModule/Classes/EnvironmentQuery/EQSRenderingComponent.h
--- a/Engine/Source/Runtime/AIModule/Classes/EnvironmentQuery/EQSRenderingComponent.h
+++ b/Engine/Source/Runtime/AIModule/Classes/EnvironmentQuery/EQSRenderingComponent.h
@@ -12,10 +12,10 @@ public:
 	FEQSSceneProxy(const UPrimitiveComponent* InComponent, const FString& ViewFlagName = TEXT("DebugAI"), bool bDrawOnlyWhenSelected = true);
 	FEQSSceneProxy(const UPrimitiveComponent* InComponent, const FString& ViewFlagName, bool bDrawOnlyWhenSelected, const TArray<FSphere>& Spheres, const TArray<FText3d>& Texts);
 
-	virtual void DrawDebugLabels(UCanvas* Canvas, APlayerController*) OVERRIDE;
+	virtual void DrawDebugLabels(UCanvas* Canvas, APlayerController*) override;
 	
 	//virtual void DrawStaticElements(FStaticPrimitiveDrawInterface* PDI) OVERRIDE;
-	virtual FPrimitiveViewRelevance GetViewRelevance(const FSceneView* View) OVERRIDE;
+	virtual FPrimitiveViewRelevance GetViewRelevance(const FSceneView* View) override;
 
 	static void CollectEQSData(const UPrimitiveComponent* InComponent, const class IEQSQueryResultSourceInterface* QueryDataSource, TArray<FSphere>& Spheres, TArray<FText3d>& Texts);
 private:
@@ -38,8 +38,8 @@ class AIMODULE_API UEQSRenderingComponent : public UPrimitiveComponent
 	FString DrawFlagName;
 	bool bDrawOnlyWhenSelected;
 
-	virtual FPrimitiveSceneProxy* CreateSceneProxy() OVERRIDE;
-	virtual FBoxSphereBounds CalcBounds(const FTransform &LocalToWorld) const OVERRIDE;
-	virtual void CreateRenderState_Concurrent() OVERRIDE;
-	virtual void DestroyRenderState_Concurrent() OVERRIDE;
+	virtual FPrimitiveSceneProxy* CreateSceneProxy() override;
+	virtual FBoxSphereBounds CalcBounds(const FTransform &LocalToWorld) const override;
+	virtual void CreateRenderState_Concurrent() override;
+	virtual void DestroyRenderState_Concurrent() override;
 };
diff --git a/Engine/Source/Runtime/AIModule/Classes/EnvironmentQuery/EQSTestingPawn.h b/Engine/Source/Runtime/AIModule/Classes/EnvironmentQuery/EQSTestingPawn.h
--- a/Engine/Source/Runtime/AIModule/Classes/EnvironmentQuery/EQSTestingPawn.h
+++ b/Engine/Source/Runtime/AIModule/Classes/EnvironmentQuery/EQSTestingPawn.h
@@ -53,20 +53,20 @@ protected:
 
 public:
 	/** This pawn class spawns its controller in PostInitProperties to have it available in editor mode*/
-	virtual void TickActor( float DeltaTime, enum ELevelTick TickType, FActorTickFunction& ThisTickFunction ) OVERRIDE;
-	virtual void PostLoad() OVERRIDE;
+	virtual void TickActor( float DeltaTime, enum ELevelTick TickType, FActorTickFunction& ThisTickFunction ) override;
+	virtual void PostLoad() override;
 
 #if WITH_EDITOR
-	virtual void PostEditChangeProperty(struct FPropertyChangedEvent& PropertyChangedEvent) OVERRIDE;
-	virtual void PostEditMove(bool bFinished) OVERRIDE;
+	virtual void PostEditChangeProperty(struct FPropertyChangedEvent& PropertyChangedEvent) override;
+	virtual void PostEditMove(bool bFinished) override;
 #endif // WITH_EDITOR
 
 	// IEQSQueryResultSourceInterface start
-	virtual const struct FEnvQueryResult* GetQueryResult() const OVERRIDE;
-	virtual const struct FEnvQueryInstance* GetQueryInstance() const  OVERRIDE;
+	virtual const struct FEnvQueryResult* GetQueryResult() const override;
+	virtual const struct FEnvQueryInstance* GetQueryInstance() const  override;
 
-	virtual bool GetShouldDebugDrawLabels() const OVERRIDE { return bDrawLabels; }
-	virtual bool GetShouldDrawFailedItems() const OVERRIDE{ return bDrawFailedItems; }
+	virtual bool GetShouldDebugDrawLabels() const override { return bDrawLabels; }
+	virtual bool GetShouldDrawFailedItems() const override{ return bDrawFailedItems; }
 	// IEQSQueryResultSourceInterface end
 
 	void RunEQSQuery();
diff --git a/Engine/Source/Runtime/AIModule/Classes/EnvironmentQuery/EnvQueryManager.h b/Engine/Source/Runtime/AIModule/Classes/EnvironmentQuery/EnvQueryManager.h
--- a/Engine/Source/Runtime/AIModule/Classes/EnvironmentQuery/EnvQueryManager.h
+++ b/Engine/Source/Runtime/AIModule/Classes/EnvironmentQuery/EnvQueryManager.h
@@ -90,13 +90,13 @@ class AIMODULE_API UEnvQueryManager : public UObject, public FTickableGameObject
 	GENERATED_UCLASS_BODY()
 
 	/** [FTickableGameObject] tick function */
-	virtual void Tick(float DeltaTime) OVERRIDE;
+	virtual void Tick(float DeltaTime) override;
 
 	/** [FTickableGameObject] always tick, unless it's the default object */
-	virtual bool IsTickable() const OVERRIDE { return HasAnyFlags(RF_ClassDefaultObject) == false; }
+	virtual bool IsTickable() const override { return HasAnyFlags(RF_ClassDefaultObject) == false; }
 
 	/** [FTickableGameObject] tick stats */
-	virtual TStatId GetStatId() const OVERRIDE;
+	virtual TStatId GetStatId() const override;
 
 	/** execute query */
 	int32 RunQuery(const struct FEnvQueryRequest& Request, EEnvQueryRunMode::Type RunMode, FQueryFinishedSignature const& FinishDelegate);
@@ -118,7 +118,7 @@ class AIMODULE_API UEnvQueryManager : public UObject, public FTickableGameObject
 	virtual void OnPreLoadMap();
 
 	/** cleanup hooks for map loading */
-	virtual void FinishDestroy() OVERRIDE;
+	virtual void FinishDestroy() override;
 
 	/** list of all known item types */
 	static TArray<TSubclassOf<UEnvQueryItemType> > RegisteredItemTypes;
diff --git a/Engine/Source/Runtime/AIModule/Classes/EnvironmentQuery/EnvQueryTest.h b/Engine/Source/Runtime/AIModule/Classes/EnvironmentQuery/EnvQueryTest.h
--- a/Engine/Source/Runtime/AIModule/Classes/EnvironmentQuery/EnvQueryTest.h
+++ b/Engine/Source/Runtime/AIModule/Classes/EnvironmentQuery/EnvQueryTest.h
@@ -134,10 +134,10 @@ public:
 	FText DescribeFloatTestParams() const;
 	FText DescribeBoolTestParams(const FString& ConditionDesc) const;
 
-	virtual void PostLoad() OVERRIDE;
+	virtual void PostLoad() override;
 
 #if WITH_EDITOR
-	virtual void PostEditChangeProperty(struct FPropertyChangedEvent& PropertyChangedEvent) OVERRIDE;
+	virtual void PostEditChangeProperty(struct FPropertyChangedEvent& PropertyChangedEvent) override;
 #endif //WITH_EDITOR
 };
 
diff --git a/Engine/Source/Runtime/AIModule/Classes/EnvironmentQuery/Generators/EnvQueryGenerator_Composite.h b/Engine/Source/Runtime/AIModule/Classes/EnvironmentQuery/Generators/EnvQueryGenerator_Composite.h
--- a/Engine/Source/Runtime/AIModule/Classes/EnvironmentQuery/Generators/EnvQueryGenerator_Composite.h
+++ b/Engine/Source/Runtime/AIModule/Classes/EnvironmentQuery/Generators/EnvQueryGenerator_Composite.h
@@ -18,5 +18,5 @@ class UEnvQueryGenerator_Composite : public UEnvQueryGenerator
 	TArray<class UEnvQueryGenerator*> Generators;
 
 	void GenerateItems(struct FEnvQueryInstance& QueryInstance);
-	virtual FText GetDescriptionTitle() const OVERRIDE;
+	virtual FText GetDescriptionTitle() const override;
 };
diff --git a/Engine/Source/Runtime/AIModule/Classes/EnvironmentQuery/Generators/EnvQueryGenerator_OnCircle.h b/Engine/Source/Runtime/AIModule/Classes/EnvironmentQuery/Generators/EnvQueryGenerator_OnCircle.h
--- a/Engine/Source/Runtime/AIModule/Classes/EnvironmentQuery/Generators/EnvQueryGenerator_OnCircle.h
+++ b/Engine/Source/Runtime/AIModule/Classes/EnvironmentQuery/Generators/EnvQueryGenerator_OnCircle.h
@@ -39,15 +39,15 @@ class UEnvQueryGenerator_OnCircle : public UEnvQueryGenerator_ProjectedPoints
 	UPROPERTY()
 	uint32 bDefineArc:1;
 
-	virtual void PostLoad() OVERRIDE;
+	virtual void PostLoad() override;
 
 	void GenerateItems(struct FEnvQueryInstance& QueryInstance); 
 
-	virtual FText GetDescriptionTitle() const OVERRIDE;
-	virtual FText GetDescriptionDetails() const OVERRIDE;
+	virtual FText GetDescriptionTitle() const override;
+	virtual FText GetDescriptionDetails() const override;
 
 #if WITH_EDITOR
-	virtual void PostEditChangeProperty( struct FPropertyChangedEvent& PropertyChangedEvent) OVERRIDE;
+	virtual void PostEditChangeProperty( struct FPropertyChangedEvent& PropertyChangedEvent) override;
 #endif // WITH_EDITOR
 
 protected:
diff --git a/Engine/Source/Runtime/AIModule/Classes/EnvironmentQuery/Generators/EnvQueryGenerator_PathingGrid.h b/Engine/Source/Runtime/AIModule/Classes/EnvironmentQuery/Generators/EnvQueryGenerator_PathingGrid.h
--- a/Engine/Source/Runtime/AIModule/Classes/EnvironmentQuery/Generators/EnvQueryGenerator_PathingGrid.h
+++ b/Engine/Source/Runtime/AIModule/Classes/EnvironmentQuery/Generators/EnvQueryGenerator_PathingGrid.h
@@ -40,8 +40,8 @@ class UEnvQueryGenerator_PathingGrid : public UEnvQueryGenerator
 
 	void GenerateItems(struct FEnvQueryInstance& QueryInstance); 
 
-	virtual FText GetDescriptionTitle() const OVERRIDE;
-	virtual FText GetDescriptionDetails() const OVERRIDE;
+	virtual FText GetDescriptionTitle() const override;
+	virtual FText GetDescriptionDetails() const override;
 
 protected:
 #if WITH_RECAST
diff --git a/Engine/Source/Runtime/AIModule/Classes/EnvironmentQuery/Generators/EnvQueryGenerator_SimpleGrid.h b/Engine/Source/Runtime/AIModule/Classes/EnvironmentQuery/Generators/EnvQueryGenerator_SimpleGrid.h
--- a/Engine/Source/Runtime/AIModule/Classes/EnvironmentQuery/Generators/EnvQueryGenerator_SimpleGrid.h
+++ b/Engine/Source/Runtime/AIModule/Classes/EnvironmentQuery/Generators/EnvQueryGenerator_SimpleGrid.h
@@ -27,6 +27,6 @@ class UEnvQueryGenerator_SimpleGrid : public UEnvQueryGenerator_ProjectedPoints
 
 	void GenerateItems(struct FEnvQueryInstance& QueryInstance); 
 
-	virtual FText GetDescriptionTitle() const OVERRIDE;
-	virtual FText GetDescriptionDetails() const OVERRIDE;
+	virtual FText GetDescriptionTitle() const override;
+	virtual FText GetDescriptionDetails() const override;
 };
diff --git a/Engine/Source/Runtime/AIModule/Classes/EnvironmentQuery/Items/EnvQueryItemType.h b/Engine/Source/Runtime/AIModule/Classes/EnvironmentQuery/Items/EnvQueryItemType.h
--- a/Engine/Source/Runtime/AIModule/Classes/EnvironmentQuery/Items/EnvQueryItemType.h
+++ b/Engine/Source/Runtime/AIModule/Classes/EnvironmentQuery/Items/EnvQueryItemType.h
@@ -18,7 +18,7 @@ class AIMODULE_API UEnvQueryItemType : public UObject
 	virtual bool StoreInBlackboard(struct FBlackboardKeySelector& KeySelector, class UBlackboardComponent* Blackboard, const uint8* RawData) const;
 
 	/** unregister from known types in EnvQueryManager */
-	virtual void FinishDestroy() OVERRIDE;
+	virtual void FinishDestroy() override;
 
 	/** describe item */
 	virtual FString GetDescription(const uint8* RawData) const;
diff --git a/Engine/Source/Runtime/AIModule/Classes/EnvironmentQuery/Items/EnvQueryItemType_Actor.h b/Engine/Source/Runtime/AIModule/Classes/EnvironmentQuery/Items/EnvQueryItemType_Actor.h
--- a/Engine/Source/Runtime/AIModule/Classes/EnvironmentQuery/Items/EnvQueryItemType_Actor.h
+++ b/Engine/Source/Runtime/AIModule/Classes/EnvironmentQuery/Items/EnvQueryItemType_Actor.h
@@ -17,7 +17,7 @@ class AIMODULE_API UEnvQueryItemType_Actor : public UEnvQueryItemType_ActorBase
 	static void SetContextHelper(struct FEnvQueryContextData& ContextData, const TArray<const AActor*>& MultipleActors);
 	static void SetContextHelper(struct FEnvQueryContextData& ContextData, const TArray<AActor*>& MultipleActors);
 
-	virtual FVector GetLocation(const uint8* RawData) const OVERRIDE;
-	virtual FRotator GetRotation(const uint8* RawData) const OVERRIDE;
-	virtual AActor* GetActor(const uint8* RawData) const OVERRIDE;
+	virtual FVector GetLocation(const uint8* RawData) const override;
+	virtual FRotator GetRotation(const uint8* RawData) const override;
+	virtual AActor* GetActor(const uint8* RawData) const override;
 };
diff --git a/Engine/Source/Runtime/AIModule/Classes/EnvironmentQuery/Items/EnvQueryItemType_ActorBase.h b/Engine/Source/Runtime/AIModule/Classes/EnvironmentQuery/Items/EnvQueryItemType_ActorBase.h
--- a/Engine/Source/Runtime/AIModule/Classes/EnvironmentQuery/Items/EnvQueryItemType_ActorBase.h
+++ b/Engine/Source/Runtime/AIModule/Classes/EnvironmentQuery/Items/EnvQueryItemType_ActorBase.h
@@ -13,7 +13,7 @@ class AIMODULE_API UEnvQueryItemType_ActorBase : public UEnvQueryItemType_Vector
 
 	virtual AActor* GetActor(const uint8* RawData) const;
 
-	virtual void AddBlackboardFilters(struct FBlackboardKeySelector& KeySelector, UObject* FilterOwner) const OVERRIDE;
-	virtual bool StoreInBlackboard(struct FBlackboardKeySelector& KeySelector, class UBlackboardComponent* Blackboard, const uint8* RawData) const OVERRIDE;
-	virtual FString GetDescription(const uint8* RawData) const OVERRIDE;
+	virtual void AddBlackboardFilters(struct FBlackboardKeySelector& KeySelector, UObject* FilterOwner) const override;
+	virtual bool StoreInBlackboard(struct FBlackboardKeySelector& KeySelector, class UBlackboardComponent* Blackboard, const uint8* RawData) const override;
+	virtual FString GetDescription(const uint8* RawData) const override;
 };
diff --git a/Engine/Source/Runtime/AIModule/Classes/EnvironmentQuery/Items/EnvQueryItemType_Direction.h b/Engine/Source/Runtime/AIModule/Classes/EnvironmentQuery/Items/EnvQueryItemType_Direction.h
--- a/Engine/Source/Runtime/AIModule/Classes/EnvironmentQuery/Items/EnvQueryItemType_Direction.h
+++ b/Engine/Source/Runtime/AIModule/Classes/EnvironmentQuery/Items/EnvQueryItemType_Direction.h
@@ -20,6 +20,6 @@ class AIMODULE_API UEnvQueryItemType_Direction : public UEnvQueryItemType_Vector
 	static void SetContextHelper(struct FEnvQueryContextData& ContextData, const TArray<FVector>& MultipleDirections);
 	static void SetContextHelper(struct FEnvQueryContextData& ContextData, const TArray<FRotator>& MultipleRotations);
 
-	virtual FRotator GetRotation(const uint8* RawData) const OVERRIDE;
-	virtual FString GetDescription(const uint8* RawData) const OVERRIDE;
+	virtual FRotator GetRotation(const uint8* RawData) const override;
+	virtual FString GetDescription(const uint8* RawData) const override;
 };
diff --git a/Engine/Source/Runtime/AIModule/Classes/EnvironmentQuery/Items/EnvQueryItemType_Point.h b/Engine/Source/Runtime/AIModule/Classes/EnvironmentQuery/Items/EnvQueryItemType_Point.h
--- a/Engine/Source/Runtime/AIModule/Classes/EnvironmentQuery/Items/EnvQueryItemType_Point.h
+++ b/Engine/Source/Runtime/AIModule/Classes/EnvironmentQuery/Items/EnvQueryItemType_Point.h
@@ -15,5 +15,5 @@ class AIMODULE_API UEnvQueryItemType_Point : public UEnvQueryItemType_VectorBase
 	static void SetContextHelper(struct FEnvQueryContextData& ContextData, const FVector& SinglePoint);
 	static void SetContextHelper(struct FEnvQueryContextData& ContextData, const TArray<FVector>& MultiplePoints);
 
-	virtual FVector GetLocation(const uint8* RawData) const OVERRIDE;
+	virtual FVector GetLocation(const uint8* RawData) const override;
 };
diff --git a/Engine/Source/Runtime/AIModule/Classes/EnvironmentQuery/Items/EnvQueryItemType_VectorBase.h b/Engine/Source/Runtime/AIModule/Classes/EnvironmentQuery/Items/EnvQueryItemType_VectorBase.h
--- a/Engine/Source/Runtime/AIModule/Classes/EnvironmentQuery/Items/EnvQueryItemType_VectorBase.h
+++ b/Engine/Source/Runtime/AIModule/Classes/EnvironmentQuery/Items/EnvQueryItemType_VectorBase.h
@@ -12,7 +12,7 @@ class AIMODULE_API UEnvQueryItemType_VectorBase : public UEnvQueryItemType
 	virtual FVector GetLocation(const uint8* RawData) const;
 	virtual FRotator GetRotation(const uint8* RawData) const;
 
-	virtual void AddBlackboardFilters(struct FBlackboardKeySelector& KeySelector, UObject* FilterOwner) const OVERRIDE;
-	virtual bool StoreInBlackboard(struct FBlackboardKeySelector& KeySelector, class UBlackboardComponent* Blackboard, const uint8* RawData) const OVERRIDE;
-	virtual FString GetDescription(const uint8* RawData) const OVERRIDE;
+	virtual void AddBlackboardFilters(struct FBlackboardKeySelector& KeySelector, UObject* FilterOwner) const override;
+	virtual bool StoreInBlackboard(struct FBlackboardKeySelector& KeySelector, class UBlackboardComponent* Blackboard, const uint8* RawData) const override;
+	virtual FString GetDescription(const uint8* RawData) const override;
 };
diff --git a/Engine/Source/Runtime/AIModule/Classes/EnvironmentQuery/Tests/EnvQueryTest_Distance.h b/Engine/Source/Runtime/AIModule/Classes/EnvironmentQuery/Tests/EnvQueryTest_Distance.h
--- a/Engine/Source/Runtime/AIModule/Classes/EnvironmentQuery/Tests/EnvQueryTest_Distance.h
+++ b/Engine/Source/Runtime/AIModule/Classes/EnvironmentQuery/Tests/EnvQueryTest_Distance.h
@@ -30,6 +30,6 @@ class UEnvQueryTest_Distance : public UEnvQueryTest
 
 	void RunTest(struct FEnvQueryInstance& QueryInstance);
 
-	virtual FString GetDescriptionTitle() const OVERRIDE;
-	virtual FText GetDescriptionDetails() const OVERRIDE;
+	virtual FString GetDescriptionTitle() const override;
+	virtual FText GetDescriptionDetails() const override;
 };
diff --git a/Engine/Source/Runtime/AIModule/Classes/EnvironmentQuery/Tests/EnvQueryTest_Dot.h b/Engine/Source/Runtime/AIModule/Classes/EnvironmentQuery/Tests/EnvQueryTest_Dot.h
--- a/Engine/Source/Runtime/AIModule/Classes/EnvironmentQuery/Tests/EnvQueryTest_Dot.h
+++ b/Engine/Source/Runtime/AIModule/Classes/EnvironmentQuery/Tests/EnvQueryTest_Dot.h
@@ -19,8 +19,8 @@ class UEnvQueryTest_Dot : public UEnvQueryTest
 
 	void RunTest(struct FEnvQueryInstance& QueryInstance);
 
-	virtual FString GetDescriptionTitle() const OVERRIDE;
-	virtual FText GetDescriptionDetails() const OVERRIDE;
+	virtual FString GetDescriptionTitle() const override;
+	virtual FText GetDescriptionDetails() const override;
 
 protected:
 
