Fixed some coding style inconsistencies #UE4

diff --git a/Engine/Source/Runtime/AIModule/Private/AIController.cpp b/Engine/Source/Runtime/AIModule/Private/AIController.cpp
--- a/Engine/Source/Runtime/AIModule/Private/AIController.cpp
+++ b/Engine/Source/Runtime/AIModule/Private/AIController.cpp
@@ -333,26 +333,26 @@ bool AAIController::LineOfSightTo(const AActor* Other, FVector ViewPoint, bool b
 		int32 imax = 0;
 		float currentmin = (Points[0] - ViewPoint).SizeSquared();
 		float currentmax = currentmin;
-		for ( int32 i=1; i<4; i++ )
+		for ( int32 PointIndex=1; PointIndex<4; PointIndex++ )
 		{
-			float nextsize = (Points[i] - ViewPoint).SizeSquared();
+			float nextsize = (Points[PointIndex] - ViewPoint).SizeSquared();
 			if (nextsize > currentmax)
 			{
 				currentmax = nextsize;
-				imax = i;
+				imax = PointIndex;
 			}
 			else if (nextsize < currentmin)
 			{
 				currentmin = nextsize;
-				imin = i;
+				imin = PointIndex;
 			}
 		}
 
-		for ( int32 i=0; i<4; i++ )
+		for ( int32 PointIndex=0; PointIndex<4; PointIndex++ )
 		{
-			if	( (i != imin) && (i != imax) )
+			if	( (PointIndex != imin) && (PointIndex != imax) )
 			{
-				bHit = GetWorld()->LineTraceTest(ViewPoint,  Points[i], ECC_Visibility, CollisionParams);
+				bHit = GetWorld()->LineTraceTest(ViewPoint,  Points[PointIndex], ECC_Visibility, CollisionParams);
 				if ( !bHit )
 				{
 					return true;
diff --git a/Engine/Source/Runtime/AIModule/Private/BehaviorTree/BTCompositeNode.cpp b/Engine/Source/Runtime/AIModule/Private/BehaviorTree/BTCompositeNode.cpp
--- a/Engine/Source/Runtime/AIModule/Private/BehaviorTree/BTCompositeNode.cpp
+++ b/Engine/Source/Runtime/AIModule/Private/BehaviorTree/BTCompositeNode.cpp
@@ -60,12 +60,12 @@ int32 UBTCompositeNode::GetChildIndex(struct FBehaviorTreeSearchData& SearchData
 
 int32 UBTCompositeNode::GetChildIndex(const class UBTNode* ChildNode) const
 {
-	for (int32 i = 0; i < Children.Num(); i++)
+	for (int32 ChildIndex = 0; ChildIndex < Children.Num(); ChildIndex++)
 	{
-		if (Children[i].ChildComposite == ChildNode ||
-			Children[i].ChildTask == ChildNode)
+		if (Children[ChildIndex].ChildComposite == ChildNode ||
+			Children[ChildIndex].ChildTask == ChildNode)
 		{
-			return i;
+			return ChildIndex;
 		}
 	}
 
@@ -129,9 +129,9 @@ void UBTCompositeNode::OnNodeActivation(struct FBehaviorTreeSearchData& SearchDa
 	}
 
 	// add services when execution flow enters this composite
-	for (int32 i = 0; i < Services.Num(); i++)
+	for (int32 ServiceIndex = 0; ServiceIndex < Services.Num(); ServiceIndex++)
 	{
-		SearchData.AddUniqueUpdate(FBehaviorTreeSearchUpdate(Services[i], SearchData.OwnerComp->GetActiveInstanceIdx(), EBTNodeUpdateMode::Add));
+		SearchData.AddUniqueUpdate(FBehaviorTreeSearchUpdate(Services[ServiceIndex], SearchData.OwnerComp->GetActiveInstanceIdx(), EBTNodeUpdateMode::Add));
 	}
 }
 
@@ -149,9 +149,9 @@ void UBTCompositeNode::OnNodeDeactivation(struct FBehaviorTreeSearchData& Search
 	}
 
 	// remove all services if execution flow leaves this composite
-	for (int32 i = 0; i < Services.Num(); i++)
+	for (int32 ServiceIndex = 0; ServiceIndex < Services.Num(); ServiceIndex++)
 	{
-		SearchData.AddUniqueUpdate(FBehaviorTreeSearchUpdate(Services[i], SearchData.OwnerComp->GetActiveInstanceIdx(), EBTNodeUpdateMode::Remove));
+		SearchData.AddUniqueUpdate(FBehaviorTreeSearchUpdate(Services[ServiceIndex], SearchData.OwnerComp->GetActiveInstanceIdx(), EBTNodeUpdateMode::Remove));
 	}
 }
 
@@ -165,9 +165,9 @@ void UBTCompositeNode::OnNodeRestart(struct FBehaviorTreeSearchData& SearchData)
 void UBTCompositeNode::NotifyDecoratorsOnActivation(struct FBehaviorTreeSearchData& SearchData, int32 ChildIdx) const
 {
 	const FBTCompositeChild& ChildInfo = Children[ChildIdx];
-	for (int32 i = 0; i < ChildInfo.Decorators.Num(); i++)
+	for (int32 DecoratorIndex = 0; DecoratorIndex < ChildInfo.Decorators.Num(); DecoratorIndex++)
 	{
-		const UBTDecorator* DecoratorOb = ChildInfo.Decorators[i];
+		const UBTDecorator* DecoratorOb = ChildInfo.Decorators[DecoratorIndex];
 		DecoratorOb->WrappedOnNodeActivation(SearchData);
 
 		switch (DecoratorOb->GetFlowAbortMode())
@@ -190,9 +190,9 @@ void UBTCompositeNode::NotifyDecoratorsOnActivation(struct FBehaviorTreeSearchDa
 void UBTCompositeNode::NotifyDecoratorsOnDeactivation(struct FBehaviorTreeSearchData& SearchData, int32 ChildIdx, EBTNodeResult::Type& NodeResult) const
 {
 	const FBTCompositeChild& ChildInfo = Children[ChildIdx];
-	for (int32 i = 0; i < ChildInfo.Decorators.Num(); i++)
+	for (int32 DecoratorIndex = 0; DecoratorIndex < ChildInfo.Decorators.Num(); DecoratorIndex++)
 	{
-		const UBTDecorator* DecoratorOb = ChildInfo.Decorators[i];
+		const UBTDecorator* DecoratorOb = ChildInfo.Decorators[DecoratorIndex];
 		DecoratorOb->WrappedOnNodeProcessed(SearchData, NodeResult);
 		DecoratorOb->WrappedOnNodeDeactivation(SearchData, NodeResult);
 
@@ -215,9 +215,9 @@ void UBTCompositeNode::NotifyDecoratorsOnFailedActivation(struct FBehaviorTreeSe
 	const FBTCompositeChild& ChildInfo = Children[ChildIdx];
 	const uint16 ActiveInstanceIdx = SearchData.OwnerComp->GetActiveInstanceIdx();
 
-	for (int32 i = 0; i < ChildInfo.Decorators.Num(); i++)
+	for (int32 DecoratorIndex = 0; DecoratorIndex < ChildInfo.Decorators.Num(); DecoratorIndex++)
 	{
-		const UBTDecorator* DecoratorOb = ChildInfo.Decorators[i];
+		const UBTDecorator* DecoratorOb = ChildInfo.Decorators[DecoratorIndex];
 		DecoratorOb->WrappedOnNodeProcessed(SearchData, NodeResult);
 
 		if (DecoratorOb->GetFlowAbortMode() == EBTFlowAbortMode::LowerPriority ||
@@ -246,11 +246,11 @@ void UBTCompositeNode::ConditionalNotifyChildExecution(class UBehaviorTreeCompon
 {
 	if (bUseChildExecutionNotify)
 	{
-		for (int32 i = 0; i < Children.Num(); i++)
+		for (int32 ChildIndex = 0; ChildIndex < Children.Num(); ChildIndex++)
 		{
-			if (Children[i].ChildComposite == ChildNode || Children[i].ChildTask == ChildNode)
+			if (Children[ChildIndex].ChildComposite == ChildNode || Children[ChildIndex].ChildTask == ChildNode)
 			{
-				NotifyChildExecution(OwnerComp, NodeMemory, i, NodeResult);
+				NotifyChildExecution(OwnerComp, NodeMemory, ChildIndex, NodeResult);
 				break;
 			}
 		}
@@ -360,9 +360,9 @@ bool UBTCompositeNode::DoDecoratorsAllowExecution(class UBehaviorTreeComponent*
 	if (ChildInfo.DecoratorOps.Num() == 0)
 	{
 		// simple check: all decorators must agree
-		for (int32 i = 0; i < ChildInfo.Decorators.Num(); i++)
+		for (int32 DecoratorIndex = 0; DecoratorIndex < ChildInfo.Decorators.Num(); DecoratorIndex++)
 		{
-			const UBTDecorator* TestDecorator = ChildInfo.Decorators[i];
+			const UBTDecorator* TestDecorator = ChildInfo.Decorators[DecoratorIndex];
 			const bool bIsAllowed = TestDecorator ? TestDecorator->WrappedCanExecute(OwnerComp, TestDecorator->GetNodeMemory<uint8>(MyInstance)) : false;
 			OwnerComp->StoreDebuggerSearchStep(TestDecorator, InstanceIdx, bIsAllowed);
 
@@ -393,9 +393,9 @@ bool UBTCompositeNode::DoDecoratorsAllowExecution(class UBehaviorTreeComponent*
 		int32 FailedDecoratorIdx = INDEX_NONE;
 		bool bShouldStoreNodeIndex = true;
 
-		for (int32 i = 0; i < ChildInfo.DecoratorOps.Num(); i++)
+		for (int32 OperationIndex = 0; OperationIndex < ChildInfo.DecoratorOps.Num(); OperationIndex++)
 		{
-			const FBTDecoratorLogic& DecoratorOp = ChildInfo.DecoratorOps[i];
+			const FBTDecoratorLogic& DecoratorOp = ChildInfo.DecoratorOps[OperationIndex];
 			if (IsLogicOp(DecoratorOp))
 			{
 				OperationStack.Add(FOperationStackInfo(DecoratorOp));
@@ -453,12 +453,12 @@ int32 UBTCompositeNode::GetMatchingChildIndex(int32 ActiveInstanceIdx, struct FB
 		}
 
 		// find child outside range
-		for (int32 i = 0; i < Children.Num(); i++)
+		for (int32 ChildIndex = 0; ChildIndex < Children.Num(); ChildIndex++)
 		{
-			const UBTNode* ChildNode = GetChildNode(i);
+			const UBTNode* ChildNode = GetChildNode(ChildIndex);
 			if (ChildNode->GetExecutionIndex() > NodeIdx.ExecutionIndex)
 			{
-				return i ? (i - 1) : 0;
+				return ChildIndex ? (ChildIndex - 1) : 0;
 			}
 		}
 
diff --git a/Engine/Source/Runtime/AIModule/Private/BehaviorTree/BTNode.cpp b/Engine/Source/Runtime/AIModule/Private/BehaviorTree/BTNode.cpp
--- a/Engine/Source/Runtime/AIModule/Private/BehaviorTree/BTNode.cpp
+++ b/Engine/Source/Runtime/AIModule/Private/BehaviorTree/BTNode.cpp
@@ -138,10 +138,10 @@ FString UBTNode::GetRuntimeDescription(const class UBehaviorTreeComponent* Owner
 		NodeOb->DescribeRuntimeValues(OwnerComp, NodeMemory, Verbosity, RuntimeValues);
 	}
 
-	for (int32 i = 0; i < RuntimeValues.Num(); i++)
+	for (int32 ValueIndex = 0; ValueIndex < RuntimeValues.Num(); ValueIndex++)
 	{
 		Description += TEXT(", ");
-		Description += RuntimeValues[i];
+		Description += RuntimeValues[ValueIndex];
 	}
 
 	return Description;
diff --git a/Engine/Source/Runtime/AIModule/Private/BehaviorTree/BehaviorTreeComponent.cpp b/Engine/Source/Runtime/AIModule/Private/BehaviorTree/BehaviorTreeComponent.cpp
--- a/Engine/Source/Runtime/AIModule/Private/BehaviorTree/BehaviorTreeComponent.cpp
+++ b/Engine/Source/Runtime/AIModule/Private/BehaviorTree/BehaviorTreeComponent.cpp
@@ -132,9 +132,9 @@ bool UBehaviorTreeComponent::StartTree(class UBehaviorTree* TreeAsset, EBTExecut
 void UBehaviorTreeComponent::StopTree()
 {
 	// clear current state, don't touch debugger data
-	for (int32 i = 0; i < InstanceStack.Num(); i++)
+	for (int32 InstanceIndex = 0; InstanceIndex < InstanceStack.Num(); InstanceIndex++)
 	{
-		InstanceStack[i].Cleanup(this);
+		InstanceStack[InstanceIndex].Cleanup(this);
 	}
 
 	InstanceStack.Reset();
@@ -286,12 +286,12 @@ bool UBehaviorTreeComponent::IsAuxNodeActive(const class UBTAuxiliaryNode* AuxNo
 	}
 
 	const uint16 AuxExecutionIndex = AuxNode->GetExecutionIndex();
-	for (int32 i = 0; i < InstanceStack.Num(); i++)
+	for (int32 InstanceIndex = 0; InstanceIndex < InstanceStack.Num(); InstanceIndex++)
 	{
-		const FBehaviorTreeInstance& InstanceInfo = InstanceStack[i];
-		for (int32 iAux = 0; iAux < InstanceInfo.ActiveAuxNodes.Num(); iAux++)
+		const FBehaviorTreeInstance& InstanceInfo = InstanceStack[InstanceIndex];
+		for (int32 AuxIndex = 0; AuxIndex < InstanceInfo.ActiveAuxNodes.Num(); AuxIndex++)
 		{
-			const UBTAuxiliaryNode* TestAuxNode = InstanceInfo.ActiveAuxNodes[iAux];
+			const UBTAuxiliaryNode* TestAuxNode = InstanceInfo.ActiveAuxNodes[AuxIndex];
 
 			// check template version
 			if (TestAuxNode == AuxNode)
@@ -327,9 +327,9 @@ EBTTaskStatus::Type UBehaviorTreeComponent::GetTaskStatus(const class UBTTaskNod
 		const FBehaviorTreeInstance& InstanceInfo = InstanceStack[InstanceIdx];
 
 		// always check parallel execution first, it takes priority over ActiveNodeType
-		for (int32 i = 0; i < InstanceInfo.ParallelTasks.Num(); i++)
+		for (int32 TaskIndex = 0; TaskIndex < InstanceInfo.ParallelTasks.Num(); TaskIndex++)
 		{
-			const FBehaviorTreeParallelTask& ParallelInfo = InstanceInfo.ParallelTasks[i];
+			const FBehaviorTreeParallelTask& ParallelInfo = InstanceInfo.ParallelTasks[TaskIndex];
 			if (ParallelInfo.TaskNode == TaskNode ||
 				(TaskNode->IsInstanced() && ParallelInfo.TaskNode && ParallelInfo.TaskNode->GetExecutionIndex() == ExecutionIndex))
 			{
@@ -597,9 +597,9 @@ void UBehaviorTreeComponent::RequestExecution(class UBTCompositeNode* RequestedO
 
 void UBehaviorTreeComponent::ApplySearchUpdates(const TArray<struct FBehaviorTreeSearchUpdate>& UpdateList, int32 UpToIdx, int32 NewNodeExecutionIndex, bool bPostUpdate)
 {
-	for (int32 i = 0; i < UpToIdx; i++)
+	for (int32 Index = 0; Index < UpToIdx; Index++)
 	{
-		const FBehaviorTreeSearchUpdate& UpdateInfo = UpdateList[i];
+		const FBehaviorTreeSearchUpdate& UpdateInfo = UpdateList[Index];
 		const UBTNode* UpdateNode = UpdateInfo.AuxNode ? (const UBTNode*)UpdateInfo.AuxNode : (const UBTNode*)UpdateInfo.TaskNode;
 		FBehaviorTreeInstance& UpdateInstance = InstanceStack[UpdateInfo.InstanceIndex];
 		int32 ParallelTaskIdx = INDEX_NONE;
@@ -721,19 +721,19 @@ void UBehaviorTreeComponent::TickComponent(float DeltaTime, enum ELevelTick Tick
 	}
 	
 	// tick active auxiliary nodes and parallel tasks (in execution order, before task)
-	for (int32 i = 0; i < InstanceStack.Num(); i++)
+	for (int32 InstanceIndex = 0; InstanceIndex < InstanceStack.Num(); InstanceIndex++)
 	{
-		FBehaviorTreeInstance& InstanceInfo = InstanceStack[i];		
-		for (int32 iAux = 0; iAux < InstanceInfo.ActiveAuxNodes.Num(); iAux++)
+		FBehaviorTreeInstance& InstanceInfo = InstanceStack[InstanceIndex];		
+		for (int32 AuxIndex = 0; AuxIndex < InstanceInfo.ActiveAuxNodes.Num(); AuxIndex++)
 		{
-			const UBTAuxiliaryNode* AuxNode = InstanceInfo.ActiveAuxNodes[iAux];
+			const UBTAuxiliaryNode* AuxNode = InstanceInfo.ActiveAuxNodes[AuxIndex];
 			uint8* NodeMemory = AuxNode->GetNodeMemory<uint8>(InstanceInfo);
 			AuxNode->WrappedTickNode(this, NodeMemory, DeltaTime);
 		}
 
-		for (int32 iTask = 0; iTask < InstanceInfo.ParallelTasks.Num(); iTask++)
+		for (int32 TaskIndex = 0; TaskIndex < InstanceInfo.ParallelTasks.Num(); TaskIndex++)
 		{
-			const UBTTaskNode* ParallelTask = InstanceInfo.ParallelTasks[iTask].TaskNode;
+			const UBTTaskNode* ParallelTask = InstanceInfo.ParallelTasks[TaskIndex].TaskNode;
 			uint8* NodeMemory = ParallelTask->GetNodeMemory<uint8>(InstanceInfo);
 			ParallelTask->WrappedTickTask(this, NodeMemory, DeltaTime);
 		}
@@ -791,13 +791,13 @@ void UBehaviorTreeComponent::ProcessExecutionRequest()
 		// find pending abort request with highest priority (search from end to optimize)
 		// and apply all changes below, to be sure that tree is not executing
 		// inactive tasks before searching for new one
-		for (int32 i = SearchData.PendingUpdates.Num() - 1; i >= 0; i--)
+		for (int32 UpdateIndex = SearchData.PendingUpdates.Num() - 1; UpdateIndex >= 0; UpdateIndex--)
 		{
-			FBehaviorTreeSearchUpdate& UpdateInfo = SearchData.PendingUpdates[i];
+			FBehaviorTreeSearchUpdate& UpdateInfo = SearchData.PendingUpdates[UpdateIndex];
 			if (UpdateInfo.TaskNode && UpdateInfo.Mode == EBTNodeUpdateMode::Remove)
 			{
 				// NewActiveNode param doesn't matter here, it's only for AddForLowerPri mode
-				ApplySearchData(NULL, i);
+				ApplySearchData(NULL, UpdateIndex);
 				break;
 			}
 		}
@@ -813,9 +813,9 @@ void UBehaviorTreeComponent::ProcessExecutionRequest()
 	// (ApplySearchData doesn't remove them, so parallels can abort latently)
 	if (InstanceStack.Num() > (ActiveInstanceIdx + 1))
 	{
-		for (int32 i = ActiveInstanceIdx + 1; i < InstanceStack.Num(); i++)
+		for (int32 InstanceIndex = ActiveInstanceIdx + 1; InstanceIndex < InstanceStack.Num(); InstanceIndex++)
 		{
-			InstanceStack[i].Cleanup(this);
+			InstanceStack[InstanceIndex].Cleanup(this);
 		}
 
 		InstanceStack.SetNum(ActiveInstanceIdx + 1);
@@ -1000,15 +1000,15 @@ bool UBehaviorTreeComponent::DeactivateUpTo(class UBTCompositeNode* Node, uint16
 
 void UBehaviorTreeComponent::UnregisterAuxNodesUpTo(const FBTNodeIndex& Index)
 {
-	for (int32 i = 0; i < InstanceStack.Num(); i++)
+	for (int32 InstanceIndex = 0; InstanceIndex < InstanceStack.Num(); InstanceIndex++)
 	{
-		FBehaviorTreeInstance& InstanceInfo = InstanceStack[i];
-		for (int32 iAux = 0; iAux < InstanceInfo.ActiveAuxNodes.Num(); iAux++)
+		FBehaviorTreeInstance& InstanceInfo = InstanceStack[InstanceIndex];
+		for (int32 AuxIndex = 0; AuxIndex < InstanceInfo.ActiveAuxNodes.Num(); AuxIndex++)
 		{
-			FBTNodeIndex AuxIdx(i, InstanceInfo.ActiveAuxNodes[iAux]->GetExecutionIndex());
+			FBTNodeIndex AuxIdx(InstanceIndex, InstanceInfo.ActiveAuxNodes[AuxIndex]->GetExecutionIndex());
 			if (Index.TakesPriorityOver(AuxIdx))
 			{
-				SearchData.AddUniqueUpdate(FBehaviorTreeSearchUpdate(InstanceInfo.ActiveAuxNodes[iAux], i, EBTNodeUpdateMode::Remove));
+				SearchData.AddUniqueUpdate(FBehaviorTreeSearchUpdate(InstanceInfo.ActiveAuxNodes[AuxIndex], InstanceIndex, EBTNodeUpdateMode::Remove));
 			}
 		}
 	}
@@ -1147,15 +1147,15 @@ void UBehaviorTreeComponent::UnregisterParallelTask(const class UBTTaskNode* Tas
 	if (InstanceStack.IsValidIndex(InstanceIdx))
 	{
 		FBehaviorTreeInstance& InstanceInfo = InstanceStack[InstanceIdx];
-		for (int32 iTask = InstanceInfo.ParallelTasks.Num() - 1; iTask >= 0; iTask--)
+		for (int32 TaskIndex = InstanceInfo.ParallelTasks.Num() - 1; TaskIndex >= 0; TaskIndex--)
 		{
-			FBehaviorTreeParallelTask& ParallelInfo = InstanceInfo.ParallelTasks[iTask];
+			FBehaviorTreeParallelTask& ParallelInfo = InstanceInfo.ParallelTasks[TaskIndex];
 			if (ParallelInfo.TaskNode == TaskNode)
 			{
 				UE_VLOG(GetOwner(), LogBehaviorTree, Verbose, TEXT("Parallel task: %s removed from active list"),
 					*UBehaviorTreeTypes::DescribeNodeHelper(TaskNode));
 
-				InstanceInfo.ParallelTasks.RemoveAt(iTask);
+				InstanceInfo.ParallelTasks.RemoveAt(TaskIndex);
 				bShouldUpdate = true;
 				break;
 			}
@@ -1174,12 +1174,12 @@ void UBehaviorTreeComponent::UpdateAbortingTasks()
 {
 	bWaitingForAbortingTasks = InstanceStack.Num() ? (InstanceStack.Last().ActiveNodeType == EBTActiveNode::AbortingTask) : false;
 
-	for (int32 i = 0; i < InstanceStack.Num() && !bWaitingForAbortingTasks; i++)
+	for (int32 InstanceIndex = 0; InstanceIndex < InstanceStack.Num() && !bWaitingForAbortingTasks; InstanceIndex++)
 	{
-		FBehaviorTreeInstance& InstanceInfo = InstanceStack[i];
-		for (int32 iTask = InstanceInfo.ParallelTasks.Num() - 1; iTask >= 0; iTask--)
+		FBehaviorTreeInstance& InstanceInfo = InstanceStack[InstanceIndex];
+		for (int32 TaskIndex = InstanceInfo.ParallelTasks.Num() - 1; TaskIndex >= 0; TaskIndex--)
 		{
-			FBehaviorTreeParallelTask& ParallelInfo = InstanceInfo.ParallelTasks[iTask];
+			FBehaviorTreeParallelTask& ParallelInfo = InstanceInfo.ParallelTasks[TaskIndex];
 			if (ParallelInfo.Status == EBTTaskStatus::Aborting)
 			{
 				bWaitingForAbortingTasks = true;
@@ -1252,9 +1252,9 @@ bool UBehaviorTreeComponent::PushInstance(class UBehaviorTree* TreeAsset)
 		ActiveInstanceIdx = InstanceStack.Num() - 1;
 
 		// start root level services now (they won't be removed on looping tree anyway)
-		for (int32 i = 0; i < RootNode->Services.Num(); i++)
+		for (int32 ServiceIndex = 0; ServiceIndex < RootNode->Services.Num(); ServiceIndex++)
 		{
-			UBTService* ServiceNode = RootNode->Services[i];
+			UBTService* ServiceNode = RootNode->Services[ServiceIndex];
 			uint8* NodeMemory = (uint8*)ServiceNode->GetNodeMemory<uint8>(InstanceStack[ActiveInstanceIdx]);
 
 			InstanceStack[ActiveInstanceIdx].ActiveAuxNodes.Add(ServiceNode);
@@ -1282,18 +1282,18 @@ uint8 UBehaviorTreeComponent::UpdateInstanceId(class UBehaviorTree* TreeAsset, c
 		InstanceId.Path.Add(ExecutionIndex);
 	}
 
-	for (int32 i = OriginInstanceIdx - 1; i >= 0; i--)
+	for (int32 InstanceIndex = OriginInstanceIdx - 1; InstanceIndex >= 0; InstanceIndex--)
 	{
-		const uint16 ExecutionIndex = InstanceStack[i].ActiveNode ? InstanceStack[i].ActiveNode->GetExecutionIndex() : MAX_uint16;
+		const uint16 ExecutionIndex = InstanceStack[InstanceIndex].ActiveNode ? InstanceStack[InstanceIndex].ActiveNode->GetExecutionIndex() : MAX_uint16;
 		InstanceId.Path.Add(ExecutionIndex);
 	}
 
 	// try to find matching existing Id
-	for (int32 i = 0; i < KnownInstances.Num(); i++)
+	for (int32 InstanceIndex = 0; InstanceIndex < KnownInstances.Num(); InstanceIndex++)
 	{
-		if (KnownInstances[i] == InstanceId)
+		if (KnownInstances[InstanceIndex] == InstanceId)
 		{
-			return i;
+			return InstanceIndex;
 		}
 	}
 
@@ -1320,11 +1320,11 @@ int32 UBehaviorTreeComponent::FindInstanceContainingNode(const UBTNode* Node) co
 				RootNode = RootNode->GetParentNode();
 			}
 
-			for (int32 i = 0; i < InstanceStack.Num(); i++)
+			for (int32 InstanceIndex = 0; InstanceIndex < InstanceStack.Num(); InstanceIndex++)
 			{
-				if (InstanceStack[i].RootNode == RootNode)
+				if (InstanceStack[InstanceIndex].RootNode == RootNode)
 				{
-					InstanceIdx = i;
+					InstanceIdx = InstanceIndex;
 					break;
 				}
 			}
@@ -1346,29 +1346,29 @@ class UBTNode* UBehaviorTreeComponent::FindTemplateNode(const class UBTNode* Nod
 	}
 
 	UBTCompositeNode* ParentNode = Node->GetParentNode();
-	for (int32 iChild = 0; iChild < ParentNode->Children.Num(); iChild++)
+	for (int32 ChildIndex = 0; ChildIndex < ParentNode->Children.Num(); ChildIndex++)
 	{
-		FBTCompositeChild& ChildInfo = ParentNode->Children[iChild];
+		FBTCompositeChild& ChildInfo = ParentNode->Children[ChildIndex];
 
 		if (ChildInfo.ChildTask && ChildInfo.ChildTask->GetExecutionIndex() == Node->GetExecutionIndex())
 		{
 			return ChildInfo.ChildTask;
 		}
 
-		for (int32 i = 0; i < ChildInfo.Decorators.Num(); i++)
+		for (int32 DecoratorIndex = 0; DecoratorIndex < ChildInfo.Decorators.Num(); DecoratorIndex++)
 		{
-			if (ChildInfo.Decorators[i]->GetExecutionIndex() == Node->GetExecutionIndex())
+			if (ChildInfo.Decorators[DecoratorIndex]->GetExecutionIndex() == Node->GetExecutionIndex())
 			{
-				return ChildInfo.Decorators[i];
+				return ChildInfo.Decorators[DecoratorIndex];
 			}
 		}
 	}
 
-	for (int32 i = 0; i < ParentNode->Services.Num(); i++)
+	for (int32 ServiceIndex = 0; ServiceIndex < ParentNode->Services.Num(); ServiceIndex++)
 	{
-		if (ParentNode->Services[i]->GetExecutionIndex() == Node->GetExecutionIndex())
+		if (ParentNode->Services[ServiceIndex]->GetExecutionIndex() == Node->GetExecutionIndex())
 		{
-			return ParentNode->Services[i];
+			return ParentNode->Services[ServiceIndex];
 		}
 	}
 
@@ -1383,9 +1383,9 @@ uint8* UBehaviorTreeComponent::GetNodeMemory(UBTNode* Node, int32 InstanceIdx) c
 FString UBehaviorTreeComponent::GetDebugInfoString() const 
 { 
 	FString DebugInfo;
-	for (int32 i = 0; i < InstanceStack.Num(); i++)
+	for (int32 InstanceIndex = 0; InstanceIndex < InstanceStack.Num(); InstanceIndex++)
 	{
-		const FBehaviorTreeInstance& InstanceData = InstanceStack[i];
+		const FBehaviorTreeInstance& InstanceData = InstanceStack[InstanceIndex];
 		const FBehaviorTreeInstanceId& InstanceInfo = KnownInstances[InstanceData.InstanceIdIndex];
 		DebugInfo += FString::Printf(TEXT("Behavior tree: %s\n"), *GetNameSafe(InstanceInfo.TreeAsset));
 
@@ -1418,12 +1418,12 @@ FString UBehaviorTreeComponent::DescribeActiveTasks() const
 	}
 
 	FString ParallelTasks;
-	for (int32 i = 0; i < InstanceStack.Num(); i++)
+	for (int32 InstanceIndex = 0; InstanceIndex < InstanceStack.Num(); InstanceIndex++)
 	{
-		const FBehaviorTreeInstance& InstanceInfo = InstanceStack[i];
-		for (int32 iParallel = 0; iParallel < InstanceInfo.ParallelTasks.Num(); iParallel++)
+		const FBehaviorTreeInstance& InstanceInfo = InstanceStack[InstanceIndex];
+		for (int32 ParallelNodeIndex = 0; ParallelNodeIndex < InstanceInfo.ParallelTasks.Num(); ParallelNodeIndex++)
 		{
-			const FBehaviorTreeParallelTask& ParallelInfo = InstanceInfo.ParallelTasks[iParallel];
+			const FBehaviorTreeParallelTask& ParallelInfo = InstanceInfo.ParallelTasks[ParallelNodeIndex];
 			if (ParallelInfo.Status == EBTTaskStatus::Active)
 			{
 				ParallelTasks += UBehaviorTreeTypes::DescribeNodeHelper(ParallelInfo.TaskNode);
@@ -1445,9 +1445,9 @@ FString UBehaviorTreeComponent::DescribeActiveTasks() const
 FString UBehaviorTreeComponent::DescribeActiveTrees() const
 {
 	FString Assets;
-	for (int32 i = 0; i < InstanceStack.Num(); i++)
+	for (int32 InstanceIndex = 0; InstanceIndex < InstanceStack.Num(); InstanceIndex++)
 	{
-		const FBehaviorTreeInstanceId& InstanceInfo = KnownInstances[InstanceStack[i].InstanceIdIndex];
+		const FBehaviorTreeInstanceId& InstanceInfo = KnownInstances[InstanceStack[InstanceIndex].InstanceIdIndex];
 		Assets += InstanceInfo.TreeAsset->GetName();
 		Assets += TEXT(", ");
 	}
@@ -1465,13 +1465,13 @@ void UBehaviorTreeComponent::DescribeSelfToVisLog(struct FVisLogEntry* Snapshot)
 	
 	Super::DescribeSelfToVisLog(Snapshot);
 
-	for (int32 i = 0; i < InstanceStack.Num(); i++)
+	for (int32 InstanceIndex = 0; InstanceIndex < InstanceStack.Num(); InstanceIndex++)
 	{
-		const FBehaviorTreeInstance& InstanceInfo = InstanceStack[i];
+		const FBehaviorTreeInstance& InstanceInfo = InstanceStack[InstanceIndex];
 		const FBehaviorTreeInstanceId& InstanceId = KnownInstances[InstanceInfo.InstanceIdIndex];
 
 		FVisLogEntry::FStatusCategory StatusCategory;
-		StatusCategory.Category = FString::Printf(TEXT("BehaviorTree %d (asset: %s)"), i, *GetNameSafe(InstanceId.TreeAsset));
+		StatusCategory.Category = FString::Printf(TEXT("BehaviorTree %d (asset: %s)"), InstanceIndex, *GetNameSafe(InstanceId.TreeAsset));
 		
 		TArray<FString> Descriptions;
 		UBTNode* Node = InstanceInfo.ActiveNode;
@@ -1483,19 +1483,19 @@ void UBehaviorTreeComponent::DescribeSelfToVisLog(struct FVisLogEntry* Snapshot)
 			Node = Node->GetParentNode();
 		}
 
-		for (int32 iDesc = Descriptions.Num() - 1; iDesc >= 0; iDesc--)
+		for (int32 DescriptionIndex = Descriptions.Num() - 1; DescriptionIndex >= 0; DescriptionIndex--)
 		{
 			int32 SplitIdx = INDEX_NONE;
-			if (Descriptions[iDesc].FindChar(TEXT(','), SplitIdx))
+			if (Descriptions[DescriptionIndex].FindChar(TEXT(','), SplitIdx))
 			{
-				const FString KeyDesc = Descriptions[iDesc].Left(SplitIdx);
-				const FString ValueDesc = Descriptions[iDesc].Mid(SplitIdx + 1).Trim();
+				const FString KeyDesc = Descriptions[DescriptionIndex].Left(SplitIdx);
+				const FString ValueDesc = Descriptions[DescriptionIndex].Mid(SplitIdx + 1).Trim();
 
 				StatusCategory.Add(KeyDesc, ValueDesc);
 			}
 			else
 			{
-				StatusCategory.Add(Descriptions[iDesc], TEXT(""));
+				StatusCategory.Add(Descriptions[DescriptionIndex], TEXT(""));
 			}
 		}
 
@@ -1517,18 +1517,18 @@ void UBehaviorTreeComponent::StoreDebuggerExecutionStep(EBTExecutionSnap::Type S
 	CurrentStep.TimeStamp = GetWorld()->GetTimeSeconds();
 	CurrentStep.BlackboardValues = SearchStartBlackboard;
 
-	for (int32 i = 0; i < InstanceStack.Num(); i++)
+	for (int32 InstanceIndex = 0; InstanceIndex < InstanceStack.Num(); InstanceIndex++)
 	{
-		const FBehaviorTreeInstance& ActiveInstance = InstanceStack[i];
+		const FBehaviorTreeInstance& ActiveInstance = InstanceStack[InstanceIndex];
 		
 		FBehaviorTreeDebuggerInstance StoreInfo;
-		StoreDebuggerInstance(StoreInfo, i, SnapType);
+		StoreDebuggerInstance(StoreInfo, InstanceIndex, SnapType);
 		CurrentStep.InstanceStack.Add(StoreInfo);
 	}
 
-	for (int32 i = RemovedInstances.Num() - 1; i >= 0; i--)
+	for (int32 InstanceIndex = RemovedInstances.Num() - 1; InstanceIndex >= 0; InstanceIndex--)
 	{
-		CurrentStep.InstanceStack.Add(RemovedInstances[i]);
+		CurrentStep.InstanceStack.Add(RemovedInstances[InstanceIndex]);
 	}
 
 	CurrentSearchFlow.Reset();
@@ -1568,15 +1568,15 @@ void UBehaviorTreeComponent::StoreDebuggerInstance(FBehaviorTreeDebuggerInstance
 		}
 
 		// add aux nodes
-		for (int32 i = 0; i < ActiveInstance.ActiveAuxNodes.Num(); i++)
+		for (int32 NodeIndex = 0; NodeIndex < ActiveInstance.ActiveAuxNodes.Num(); NodeIndex++)
 		{
-			InstanceInfo.AdditionalActiveNodes.Add(ActiveInstance.ActiveAuxNodes[i]->GetExecutionIndex());
+			InstanceInfo.AdditionalActiveNodes.Add(ActiveInstance.ActiveAuxNodes[NodeIndex]->GetExecutionIndex());
 		}
 
 		// add active parallels
-		for (int32 i = 0; i < ActiveInstance.ParallelTasks.Num(); i++)
+		for (int32 TaskIndex = 0; TaskIndex < ActiveInstance.ParallelTasks.Num(); TaskIndex++)
 		{
-			const FBehaviorTreeParallelTask& TaskInfo = ActiveInstance.ParallelTasks[i];
+			const FBehaviorTreeParallelTask& TaskInfo = ActiveInstance.ParallelTasks[TaskIndex];
 			InstanceInfo.AdditionalActiveNodes.Add(TaskInfo.TaskNode->GetExecutionIndex());
 		}
 
@@ -1593,11 +1593,11 @@ void UBehaviorTreeComponent::StoreDebuggerInstance(FBehaviorTreeDebuggerInstance
 	// store search flow, but remove nodes on execution path
 	if (CurrentSearchFlow.IsValidIndex(InstanceIdx))
 	{
-		for (int32 i = 0; i < CurrentSearchFlow[InstanceIdx].Num(); i++)
+		for (int32 FlowIndex = 0; FlowIndex < CurrentSearchFlow[InstanceIdx].Num(); FlowIndex++)
 		{
-			if (!InstanceInfo.ActivePath.Contains(CurrentSearchFlow[InstanceIdx][i].ExecutionIndex))
+			if (!InstanceInfo.ActivePath.Contains(CurrentSearchFlow[InstanceIdx][FlowIndex].ExecutionIndex))
 			{
-				InstanceInfo.PathFromPrevious.Add(CurrentSearchFlow[InstanceIdx][i]);
+				InstanceInfo.PathFromPrevious.Add(CurrentSearchFlow[InstanceIdx][FlowIndex]);
 			}
 		}
 	}
@@ -1695,14 +1695,14 @@ void UBehaviorTreeComponent::StoreDebuggerRuntimeValues(TArray<FString>& Runtime
 		Node->DescribeRuntimeValues(this, NodeMemory, EBTDescriptionVerbosity::Basic, RuntimeValues);
 
 		FString ComposedDesc;
-		for (int32 i = 0; i < RuntimeValues.Num(); i++)
+		for (int32 ValueIndex = 0; ValueIndex < RuntimeValues.Num(); ValueIndex++)
 		{
 			if (ComposedDesc.Len())
 			{
 				ComposedDesc.AppendChar(TEXT('\n'));
 			}
 
-			ComposedDesc += RuntimeValues[i];
+			ComposedDesc += RuntimeValues[ValueIndex];
 		}
 
 		RuntimeDescriptions.SetNum(Node->GetExecutionIndex() + 1);
@@ -1723,14 +1723,14 @@ void UBehaviorTreeComponent::UpdateDebuggerAfterExecution(const class UBTTaskNod
 	TaskNode->DescribeRuntimeValues(this, NodeMemory, EBTDescriptionVerbosity::Basic, RuntimeValues);
 
 	FString ComposedDesc;
-	for (int32 i = 0; i < RuntimeValues.Num(); i++)
+	for (int32 ValueIndex = 0; ValueIndex < RuntimeValues.Num(); ValueIndex++)
 	{
 		if (ComposedDesc.Len())
 		{
 			ComposedDesc.AppendChar(TEXT('\n'));
 		}
 
-		ComposedDesc += RuntimeValues[i];
+		ComposedDesc += RuntimeValues[ValueIndex];
 	}
 
 	// accessing RuntimeDesc should never be out of bounds (active task MUST be part of active instance)
@@ -1748,13 +1748,13 @@ void UBehaviorTreeComponent::StoreDebuggerBlackboard(TArray<FString>& Blackboard
 		BlackboardValueDesc.Reset();
 		BlackboardValueDesc.AddZeroed(NumKeys);
 
-		for (int32 i = 0; i < NumKeys; i++)
+		for (int32 KeyIndex = 0; KeyIndex < NumKeys; KeyIndex++)
 		{
-			BlackboardValueDesc[i] = BlackboardComp->DescribeKeyValue(i, EBlackboardDescription::OnlyValue);
+			BlackboardValueDesc[KeyIndex] = BlackboardComp->DescribeKeyValue(KeyIndex, EBlackboardDescription::OnlyValue);
 			
-			if (BlackboardValueDesc[i].Len() == 0)
+			if (BlackboardValueDesc[KeyIndex].Len() == 0)
 			{
-				BlackboardValueDesc[i] = TEXT("n/a");
+				BlackboardValueDesc[KeyIndex] = TEXT("n/a");
 			}
 		}
 	}
diff --git a/Engine/Source/Runtime/AIModule/Private/BehaviorTree/BehaviorTreeManager.cpp b/Engine/Source/Runtime/AIModule/Private/BehaviorTree/BehaviorTreeManager.cpp
--- a/Engine/Source/Runtime/AIModule/Private/BehaviorTree/BehaviorTreeManager.cpp
+++ b/Engine/Source/Runtime/AIModule/Private/BehaviorTree/BehaviorTreeManager.cpp
@@ -109,20 +109,20 @@ static void InitializeNodeHelper(UBTCompositeNode* ParentNode, UBTNode* NodeOb,
 	UBTCompositeNode* CompositeOb = Cast<UBTCompositeNode>(NodeOb);
 	if (CompositeOb)
 	{
-		for (int32 i = 0; i < CompositeOb->Services.Num(); i++)
+		for (int32 ServiceIndex = 0; ServiceIndex < CompositeOb->Services.Num(); ServiceIndex++)
 		{
-			if (CompositeOb->Services[i] == NULL)
+			if (CompositeOb->Services[ServiceIndex] == NULL)
 			{
 				UE_LOG(LogBehaviorTree, Warning, TEXT("%s has missing service node! (parent: %s)"),
 					*GetNameSafe(TreeAsset), *UBehaviorTreeTypes::DescribeNodeHelper(CompositeOb));
 
-				CompositeOb->Services.RemoveAt(i, 1, false);
-				i--;
+				CompositeOb->Services.RemoveAt(ServiceIndex, 1, false);
+				ServiceIndex--;
 				continue;
 			}
 
-			UBTService* Service = Cast<UBTService>(StaticDuplicateObject(CompositeOb->Services[i], NodeOuter, TEXT("None")));;
-			CompositeOb->Services[i] = Service;
+			UBTService* Service = Cast<UBTService>(StaticDuplicateObject(CompositeOb->Services[ServiceIndex], NodeOuter, TEXT("None")));;
+			CompositeOb->Services[ServiceIndex] = Service;
 
 			InitList.Add(FNodeInitializationData(Service, CompositeOb, ExecutionIndex, TreeDepth,
 				Service->GetInstanceMemorySize(), Service->GetSpecialMemorySize()));
@@ -131,29 +131,29 @@ static void InitializeNodeHelper(UBTCompositeNode* ParentNode, UBTNode* NodeOb,
 			ExecutionIndex++;
 		}
 
-		for (int32 i = 0; i < CompositeOb->Children.Num(); i++)
+		for (int32 ChildIndex = 0; ChildIndex < CompositeOb->Children.Num(); ChildIndex++)
 		{
-			FBTCompositeChild& ChildInfo = CompositeOb->Children[i];
-			for (int32 j = 0; j < ChildInfo.Decorators.Num(); j++)
+			FBTCompositeChild& ChildInfo = CompositeOb->Children[ChildIndex];
+			for (int32 DecoratorIndex = 0; DecoratorIndex < ChildInfo.Decorators.Num(); DecoratorIndex++)
 			{
-				if (ChildInfo.Decorators[j] == NULL)
+				if (ChildInfo.Decorators[DecoratorIndex] == NULL)
 				{
 					UE_LOG(LogBehaviorTree, Warning, TEXT("%s has missing decorator node! (parent: %s, branch: %d)"),
-						*GetNameSafe(TreeAsset), *UBehaviorTreeTypes::DescribeNodeHelper(CompositeOb), i);
+						*GetNameSafe(TreeAsset), *UBehaviorTreeTypes::DescribeNodeHelper(CompositeOb), ChildIndex);
 
-					ChildInfo.Decorators.RemoveAt(j, 1, false);
-					j--;
+					ChildInfo.Decorators.RemoveAt(DecoratorIndex, 1, false);
+					DecoratorIndex--;
 					continue;
 				}
 
-				UBTDecorator* Decorator = Cast<UBTDecorator>(StaticDuplicateObject(ChildInfo.Decorators[j], NodeOuter, TEXT("None")));
-				ChildInfo.Decorators[j] = Decorator;
+				UBTDecorator* Decorator = Cast<UBTDecorator>(StaticDuplicateObject(ChildInfo.Decorators[DecoratorIndex], NodeOuter, TEXT("None")));
+				ChildInfo.Decorators[DecoratorIndex] = Decorator;
 
 				InitList.Add(FNodeInitializationData(Decorator, CompositeOb, ExecutionIndex, TreeDepth,
 					Decorator->GetInstanceMemorySize(), Decorator->GetSpecialMemorySize()));
 
 				Decorator->InitializeFromAsset(TreeAsset);
-				Decorator->InitializeDecorator(i);
+				Decorator->InitializeDecorator(ChildIndex);
 				ExecutionIndex++;
 			}
 
@@ -184,9 +184,9 @@ bool UBehaviorTreeManager::LoadTree(class UBehaviorTree* Asset, UBTCompositeNode
 {
 	SCOPE_CYCLE_COUNTER(STAT_AI_BehaviorTree_LoadTime);
 
-	for (int32 i = 0; i < LoadedTemplates.Num(); i++)
+	for (int32 TemplateIndex = 0; TemplateIndex < LoadedTemplates.Num(); TemplateIndex++)
 	{
-		FBehaviorTreeTemplateInfo& TemplateInfo = LoadedTemplates[i];
+		FBehaviorTreeTemplateInfo& TemplateInfo = LoadedTemplates[TemplateIndex];
 		if (TemplateInfo.Asset == Asset)
 		{
 			Root = TemplateInfo.Template;
@@ -207,9 +207,9 @@ bool UBehaviorTreeManager::LoadTree(class UBehaviorTree* Asset, UBTCompositeNode
 
 #if USE_BEHAVIORTREE_DEBUGGER
 		// fill in information about next nodes in execution index, before sorting memory offsets
-		for (int32 i = 0; i < InitList.Num() - 1; i++)
+		for (int32 Index = 0; Index < InitList.Num() - 1; Index++)
 		{
-			InitList[i].Node->InitializeExecutionOrder(InitList[i + 1].Node);
+			InitList[Index].Node->InitializeExecutionOrder(InitList[Index + 1].Node);
 		}
 #endif
 
@@ -218,10 +218,10 @@ bool UBehaviorTreeManager::LoadTree(class UBehaviorTree* Asset, UBTCompositeNode
 		// but since all Engine level nodes are good... 
 		InitList.Sort(FNodeInitializationData::FMemorySort());
 		uint16 MemoryOffset = 0;
-		for (int32 i = 0; i < InitList.Num(); i++)
+		for (int32 Index = 0; Index < InitList.Num(); Index++)
 		{
-			InitList[i].Node->InitializeNode(InitList[i].ParentNode, InitList[i].ExecutionIndex, InitList[i].SpecialDataSize + MemoryOffset, InitList[i].TreeDepth);
-			MemoryOffset += InitList[i].DataSize;
+			InitList[Index].Node->InitializeNode(InitList[Index].ParentNode, InitList[Index].ExecutionIndex, InitList[Index].SpecialDataSize + MemoryOffset, InitList[Index].TreeDepth);
+			MemoryOffset += InitList[Index].DataSize;
 		}
 		
 		TemplateInfo.InstanceMemorySize = MemoryOffset;
@@ -239,9 +239,9 @@ bool UBehaviorTreeManager::LoadTree(class UBehaviorTree* Asset, UBTCompositeNode
 void UBehaviorTreeManager::InitializeMemoryHelper(const TArray<UBTDecorator*>& Nodes, TArray<uint16>& MemoryOffsets, int32& MemorySize)
 {
 	TArray<FNodeInitializationData> InitList;
-	for (int32 i = 0; i < Nodes.Num(); i++)
+	for (int32 NodeIndex = 0; NodeIndex < Nodes.Num(); NodeIndex++)
 	{
-		InitList.Add(FNodeInitializationData(Nodes[i], NULL, 0, 0, Nodes[i]->GetInstanceMemorySize(), Nodes[i]->GetSpecialMemorySize()));
+		InitList.Add(FNodeInitializationData(Nodes[NodeIndex], NULL, 0, 0, Nodes[NodeIndex]->GetInstanceMemorySize(), Nodes[NodeIndex]->GetSpecialMemorySize()));
 	}
 
 	InitList.Sort(FNodeInitializationData::FMemorySort());
@@ -249,10 +249,10 @@ void UBehaviorTreeManager::InitializeMemoryHelper(const TArray<UBTDecorator*>& N
 	uint16 MemoryOffset = 0;
 	MemoryOffsets.AddZeroed(Nodes.Num());
 
-	for (int32 i = 0; i < InitList.Num(); i++)
+	for (int32 Index = 0; Index < InitList.Num(); Index++)
 	{
-		MemoryOffsets[i] = InitList[i].SpecialDataSize + MemoryOffset;
-		MemoryOffset += InitList[i].DataSize;
+		MemoryOffsets[Index] = InitList[Index].SpecialDataSize + MemoryOffset;
+		MemoryOffset += InitList[Index].DataSize;
 	}
 
 	MemorySize = MemoryOffset;
diff --git a/Engine/Source/Runtime/AIModule/Private/BehaviorTree/BehaviorTreeTypes.cpp b/Engine/Source/Runtime/AIModule/Private/BehaviorTree/BehaviorTreeTypes.cpp
--- a/Engine/Source/Runtime/AIModule/Private/BehaviorTree/BehaviorTreeTypes.cpp
+++ b/Engine/Source/Runtime/AIModule/Private/BehaviorTree/BehaviorTreeTypes.cpp
@@ -17,20 +17,20 @@ void FBehaviorTreeInstance::Initialize(class UBehaviorTreeComponent* OwnerComp,
 		return;
 	}
 
-	for (int32 i = 0; i < Node->Services.Num(); i++)
+	for (int32 ServiceIndex = 0; ServiceIndex < Node->Services.Num(); ServiceIndex++)
 	{
-		Node->Services[i]->InitializeForInstance(OwnerComp, Node->Services[i]->GetNodeMemory<uint8>(*this), InstancedIndex);
+		Node->Services[ServiceIndex]->InitializeForInstance(OwnerComp, Node->Services[ServiceIndex]->GetNodeMemory<uint8>(*this), InstancedIndex);
 	}
 
 	Node->InitializeForInstance(OwnerComp, Node->GetNodeMemory<uint8>(*this), InstancedIndex);
 
-	for (int32 iChild = 0; iChild < Node->Children.Num(); iChild++)
+	for (int32 ChildIndex = 0; ChildIndex < Node->Children.Num(); ChildIndex++)
 	{
-		FBTCompositeChild& ChildInfo = Node->Children[iChild];
+		FBTCompositeChild& ChildInfo = Node->Children[ChildIndex];
 
-		for (int32 i = 0; i < ChildInfo.Decorators.Num(); i++)
+		for (int32 DecoratorIndex = 0; DecoratorIndex < ChildInfo.Decorators.Num(); DecoratorIndex++)
 		{
-			ChildInfo.Decorators[i]->InitializeForInstance(OwnerComp, ChildInfo.Decorators[i]->GetNodeMemory<uint8>(*this), InstancedIndex);
+			ChildInfo.Decorators[DecoratorIndex]->InitializeForInstance(OwnerComp, ChildInfo.Decorators[DecoratorIndex]->GetNodeMemory<uint8>(*this), InstancedIndex);
 		}
 
 		if (ChildInfo.ChildComposite)
@@ -51,9 +51,9 @@ void FBehaviorTreeInstance::InjectNodes(class UBehaviorTreeComponent* OwnerComp,
 		return;
 	}
 
-	for (int32 iChild = 0; iChild < Node->Children.Num(); iChild++)
+	for (int32 ChildIndex = 0; ChildIndex < Node->Children.Num(); ChildIndex++)
 	{
-		FBTCompositeChild& ChildInfo = Node->Children[iChild];
+		FBTCompositeChild& ChildInfo = Node->Children[ChildIndex];
 		if (ChildInfo.ChildComposite)
 		{
 			InjectNodes(OwnerComp, ChildInfo.ChildComposite, InstancedIndex);
@@ -112,9 +112,9 @@ void FBehaviorTreeSearchData::AddUniqueUpdate(const FBehaviorTreeSearchUpdate& U
 		*UBehaviorTreeTypes::DescribeNodeHelper(UpdateInfo.AuxNode ? (UBTNode*)UpdateInfo.AuxNode : (UBTNode*)UpdateInfo.TaskNode));
 
 	bool bSkipAdding = false;
-	for (int32 i = 0; i < PendingUpdates.Num(); i++)
+	for (int32 UpdateIndex = 0; UpdateIndex < PendingUpdates.Num(); UpdateIndex++)
 	{
-		const FBehaviorTreeSearchUpdate& Info = PendingUpdates[i];
+		const FBehaviorTreeSearchUpdate& Info = PendingUpdates[UpdateIndex];
 		if (Info.AuxNode == UpdateInfo.AuxNode && Info.TaskNode == UpdateInfo.TaskNode)
 		{
 			// duplicate, skip
@@ -129,7 +129,7 @@ void FBehaviorTreeSearchData::AddUniqueUpdate(const FBehaviorTreeSearchUpdate& U
 			// don't add pairs add-remove
 			bSkipAdding = (Info.Mode == EBTNodeUpdateMode::Remove) || (UpdateInfo.Mode == EBTNodeUpdateMode::Remove);
 
-			PendingUpdates.RemoveAt(i, 1, false);
+			PendingUpdates.RemoveAt(UpdateIndex, 1, false);
 		}
 	}
 
@@ -162,17 +162,17 @@ void FBlackboardKeySelector::InitSelectedKey(class UBlackboardData* BlackboardAs
 	check(BlackboardAsset);
 	for (UBlackboardData* It = BlackboardAsset; It; It = It->Parent)
 	{
-		for (int32 i = 0; i < It->Keys.Num(); i++)
+		for (int32 KeyIndex = 0; KeyIndex < It->Keys.Num(); KeyIndex++)
 		{
-			const FBlackboardEntry& EntryInfo = It->Keys[i];
+			const FBlackboardEntry& EntryInfo = It->Keys[KeyIndex];
 			if (EntryInfo.KeyType)
 			{
 				bool bFilterPassed = false;
 				if (AllowedTypes.Num())
 				{
-					for (int32 iFilter = 0; iFilter < AllowedTypes.Num(); iFilter++)
+					for (int32 FilterIndex = 0; FilterIndex < AllowedTypes.Num(); FilterIndex++)
 					{
-						if (EntryInfo.KeyType->IsAllowedByFilter(AllowedTypes[iFilter]))
+						if (EntryInfo.KeyType->IsAllowedByFilter(AllowedTypes[FilterIndex]))
 						{
 							bFilterPassed = true;
 							break;
diff --git a/Engine/Source/Runtime/AIModule/Private/BehaviorTree/Blackboard/BlackboardKeyType.cpp b/Engine/Source/Runtime/AIModule/Private/BehaviorTree/Blackboard/BlackboardKeyType.cpp
--- a/Engine/Source/Runtime/AIModule/Private/BehaviorTree/Blackboard/BlackboardKeyType.cpp
+++ b/Engine/Source/Runtime/AIModule/Private/BehaviorTree/Blackboard/BlackboardKeyType.cpp
@@ -12,9 +12,9 @@ UBlackboardKeyType::UBlackboardKeyType(const class FPostConstructInitializePrope
 FString UBlackboardKeyType::DescribeValue(const uint8* RawData) const
 {
 	FString DescBytes;
-	for (int32 i = 0; i < ValueSize; i++)
+	for (int32 ValueIndex = 0; ValueIndex < ValueSize; ValueIndex++)
 	{
-		DescBytes += FString::Printf(TEXT("%X"), RawData[i]);
+		DescBytes += FString::Printf(TEXT("%X"), RawData[ValueIndex]);
 	}
 
 	return DescBytes.Len() ? (FString("0x") + DescBytes) : FString("empty");
diff --git a/Engine/Source/Runtime/AIModule/Private/BehaviorTree/BlackboardComponent.cpp b/Engine/Source/Runtime/AIModule/Private/BehaviorTree/BlackboardComponent.cpp
--- a/Engine/Source/Runtime/AIModule/Private/BehaviorTree/BlackboardComponent.cpp
+++ b/Engine/Source/Runtime/AIModule/Private/BehaviorTree/BlackboardComponent.cpp
@@ -78,11 +78,11 @@ void UBlackboardComponent::InitializeBlackboard(class UBlackboardData* NewAsset)
 
 			for (UBlackboardData* It = BlackboardAsset; It; It = It->Parent)
 			{
-				for (int32 i = 0; i < It->Keys.Num(); i++)
+				for (int32 KeyIndex = 0; KeyIndex < It->Keys.Num(); KeyIndex++)
 				{
-					if (It->Keys[i].KeyType)
+					if (It->Keys[KeyIndex].KeyType)
 					{
-						InitList.Add(FBlackboardInitializationData(i + It->GetFirstKeyID(), It->Keys[i].KeyType->GetValueSize()));
+						InitList.Add(FBlackboardInitializationData(KeyIndex + It->GetFirstKeyID(), It->Keys[KeyIndex].KeyType->GetValueSize()));
 					}
 				}
 			}
@@ -92,19 +92,19 @@ void UBlackboardComponent::InitializeBlackboard(class UBlackboardData* NewAsset)
 			// but since all Engine level keys are good... 
 			InitList.Sort(FBlackboardInitializationData::FMemorySort());
 			uint16 MemoryOffset = 0;
-			for (int32 i = 0; i < InitList.Num(); i++)
+			for (int32 Index = 0; Index < InitList.Num(); Index++)
 			{
-				ValueOffsets[InitList[i].KeyID] = MemoryOffset;
-				MemoryOffset += InitList[i].DataSize;
+				ValueOffsets[InitList[Index].KeyID] = MemoryOffset;
+				MemoryOffset += InitList[Index].DataSize;
 			}
 
 			ValueMemory.AddZeroed(MemoryOffset);
 
 			// initialize memory
-			for (int32 i = 0; i < InitList.Num(); i++)
+			for (int32 Index = 0; Index < InitList.Num(); Index++)
 			{
-				const FBlackboardEntry* KeyData = BlackboardAsset->GetKey(InitList[i].KeyID);
-				uint8* RawData = GetKeyRawData(InitList[i].KeyID);
+				const FBlackboardEntry* KeyData = BlackboardAsset->GetKey(InitList[Index].KeyID);
+				uint8* RawData = GetKeyRawData(InitList[Index].KeyID);
 
 				KeyData->KeyType->Initialize(RawData);
 			}
@@ -165,9 +165,9 @@ void UBlackboardComponent::ResumeUpdates()
 {
 	bPausedNotifies = false;
 
-	for (int32 i = 0; i < QueuedUpdates.Num(); i++)
+	for (int32 UpdateIndex = 0; UpdateIndex < QueuedUpdates.Num(); UpdateIndex++)
 	{
-		NotifyObservers(QueuedUpdates[i]);
+		NotifyObservers(QueuedUpdates[UpdateIndex]);
 	}
 
 	QueuedUpdates.Empty();
@@ -216,18 +216,18 @@ FString UBlackboardComponent::GetDebugInfoString(EBlackboardDescription::Type Mo
 	uint8 Offset = 0;
 	for (UBlackboardData* It = BlackboardAsset; It; It = It->Parent)
 	{
-		for (int32 i = 0; i < It->Keys.Num(); i++)
+		for (int32 KeyIndex = 0; KeyIndex < It->Keys.Num(); KeyIndex++)
 		{
-			KeyDesc.Add(DescribeKeyValue(i + Offset, Mode));
+			KeyDesc.Add(DescribeKeyValue(KeyIndex + Offset, Mode));
 		}
 		Offset += It->Keys.Num();
 	}
 	
 	KeyDesc.Sort();
-	for (int32 i = 0; i < KeyDesc.Num(); i++)
+	for (int32 KeyDescIndex = 0; KeyDescIndex < KeyDesc.Num(); KeyDescIndex++)
 	{
 		DebugString += TEXT("  ");
-		DebugString += KeyDesc[i];
+		DebugString += KeyDesc[KeyDescIndex];
 		DebugString += TEXT('\n');
 	}
 
@@ -296,11 +296,11 @@ void UBlackboardComponent::DescribeSelfToVisLog(struct FVisLogEntry* Snapshot) c
 
 	for (UBlackboardData* It = BlackboardAsset; It; It = It->Parent)
 	{
-		for (int32 i = 0; i < It->Keys.Num(); i++)
+		for (int32 KeyIndex = 0; KeyIndex < It->Keys.Num(); KeyIndex++)
 		{
-			const FBlackboardEntry& Key = It->Keys[i];
+			const FBlackboardEntry& Key = It->Keys[KeyIndex];
 
-			const uint8* ValueData = GetKeyRawData(It->GetFirstKeyID() + i);
+			const uint8* ValueData = GetKeyRawData(It->GetFirstKeyID() + KeyIndex);
 			FString ValueDesc = Key.KeyType ? *(Key.KeyType->DescribeValue(ValueData)) : TEXT("empty");
 
 			Category.Add(Key.EntryName.ToString(), ValueDesc);
diff --git a/Engine/Source/Runtime/AIModule/Private/BehaviorTree/BlackboardData.cpp b/Engine/Source/Runtime/AIModule/Private/BehaviorTree/BlackboardData.cpp
--- a/Engine/Source/Runtime/AIModule/Private/BehaviorTree/BlackboardData.cpp
+++ b/Engine/Source/Runtime/AIModule/Private/BehaviorTree/BlackboardData.cpp
@@ -64,11 +64,11 @@ int32 UBlackboardData::GetNumKeys() const
 
 FBlackboard::FKey UBlackboardData::InternalGetKeyID(const FName& KeyName, EKeyLookupMode LookupMode) const
 {
-	for (int32 i = 0; i < Keys.Num(); i++)
+	for (int32 KeyIndex = 0; KeyIndex < Keys.Num(); KeyIndex++)
 	{
-		if (Keys[i].EntryName == KeyName)
+		if (Keys[KeyIndex].EntryName == KeyName)
 		{
-			return i + FirstKeyID;
+			return KeyIndex + FirstKeyID;
 		}
 	}
 
@@ -79,13 +79,13 @@ bool UBlackboardData::IsValid() const
 {
 	if (Parent)
 	{
-		for (int32 i = 0; i < Keys.Num(); i++)
+		for (int32 KeyIndex = 0; KeyIndex < Keys.Num(); KeyIndex++)
 		{
-			const FBlackboard::FKey KeyID = Parent->InternalGetKeyID(Keys[i].EntryName, CheckParentKeys);
+			const FBlackboard::FKey KeyID = Parent->InternalGetKeyID(Keys[KeyIndex].EntryName, CheckParentKeys);
 			if (KeyID != FBlackboard::InvalidKey)
 			{
 				UE_LOG(LogBehaviorTree, Warning, TEXT("Blackboard asset (%s) has duplicated key (%s) in parent chain!"),
-					*GetName(), *Keys[i].EntryName.ToString());
+					*GetName(), *Keys[KeyIndex].EntryName.ToString());
 
 				return false;
 			}
@@ -118,17 +118,17 @@ void UBlackboardData::PostEditChangeProperty(struct FPropertyChangedEvent& Prope
 
 static bool ContainsKeyName(FName KeyName, const TArray<FBlackboardEntry>& Keys, const TArray<FBlackboardEntry>& ParentKeys)
 {
-	for (int32 i = 0; i < Keys.Num(); i++)
+	for (int32 KeyIndex = 0; KeyIndex < Keys.Num(); KeyIndex++)
 	{
-		if (Keys[i].EntryName == KeyName)
+		if (Keys[KeyIndex].EntryName == KeyName)
 		{
 			return true;
 		}
 	}
 
-	for (int32 i = 0; i < ParentKeys.Num(); i++)
+	for (int32 KeyIndex = 0; KeyIndex < ParentKeys.Num(); KeyIndex++)
 	{
-		if (ParentKeys[i].EntryName == KeyName)
+		if (ParentKeys[KeyIndex].EntryName == KeyName)
 		{
 			return true;
 		}
@@ -149,12 +149,12 @@ void UBlackboardData::UpdateParentKeys()
 
 	for (UBlackboardData* It = Parent; It; It = It->Parent)
 	{
-		for (int32 i = 0; i < It->Keys.Num(); i++)
+		for (int32 KeyIndex = 0; KeyIndex < It->Keys.Num(); KeyIndex++)
 		{
-			const bool bAlreadyExist = ContainsKeyName(It->Keys[i].EntryName, Keys, ParentKeys);
+			const bool bAlreadyExist = ContainsKeyName(It->Keys[KeyIndex].EntryName, Keys, ParentKeys);
 			if (!bAlreadyExist)
 			{
-				ParentKeys.Add(It->Keys[i]);
+				ParentKeys.Add(It->Keys[KeyIndex]);
 			}
 		}
 	}
diff --git a/Engine/Source/Runtime/AIModule/Private/BehaviorTree/BlueprintNodeHelpers.cpp b/Engine/Source/Runtime/AIModule/Private/BehaviorTree/BlueprintNodeHelpers.cpp
--- a/Engine/Source/Runtime/AIModule/Private/BehaviorTree/BlueprintNodeHelpers.cpp
+++ b/Engine/Source/Runtime/AIModule/Private/BehaviorTree/BlueprintNodeHelpers.cpp
@@ -8,9 +8,9 @@ namespace BlueprintNodeHelpers
 	uint16 GetPropertiesMemorySize(const TArray<UProperty*>& PropertyData)
 	{
 		int32 TotalMem = 0;
-		for (int32 i = 0; i < PropertyData.Num(); i++)
+		for (int32 PropertyIndex = 0; PropertyIndex < PropertyData.Num(); PropertyIndex++)
 		{
-			TotalMem += PropertyData[i]->GetSize();
+			TotalMem += PropertyData[PropertyIndex]->GetSize();
 		}
 
 		if (TotalMem > MAX_uint16)
@@ -169,9 +169,9 @@ namespace BlueprintNodeHelpers
 
 	void DescribeRuntimeValues(const UObject* Ob, const TArray<UProperty*>& PropertyData, TArray<FString>& RuntimeValues)
 	{
-		for (int32 i = 0; i < PropertyData.Num(); i++)
+		for (int32 PropertyIndex = 0; PropertyIndex < PropertyData.Num(); PropertyIndex++)
 		{
-			UProperty* TestProperty = PropertyData[i];
+			UProperty* TestProperty = PropertyData[PropertyIndex];
 			const uint8* PropAddr = TestProperty->ContainerPtrToValuePtr<uint8>(Ob);
 
 			RuntimeValues.Add(DescribeProperty(TestProperty, PropAddr));
@@ -192,9 +192,9 @@ namespace BlueprintNodeHelpers
 		{
 			TArray<UBehaviorTreeComponent*> Components;
 			OwningActor->GetComponents(Components);
-			for (int32 i = 0; i < Components.Num(); i++)
+			for (int32 ComponentIndex = 0; ComponentIndex < Components.Num(); ComponentIndex++)
 			{
-				UBehaviorTreeComponent* BTComp = Components[i];
+				UBehaviorTreeComponent* BTComp = Components[ComponentIndex];
 				if (BTComp)
 				{
 					const int32 InstanceIdx = BTComp->FindInstanceContainingNode(Node);
diff --git a/Engine/Source/Runtime/AIModule/Private/BehaviorTree/Decorators/BTDecorator_BlueprintBase.cpp b/Engine/Source/Runtime/AIModule/Private/BehaviorTree/Decorators/BTDecorator_BlueprintBase.cpp
--- a/Engine/Source/Runtime/AIModule/Private/BehaviorTree/Decorators/BTDecorator_BlueprintBase.cpp
+++ b/Engine/Source/Runtime/AIModule/Private/BehaviorTree/Decorators/BTDecorator_BlueprintBase.cpp
@@ -71,9 +71,9 @@ void UBTDecorator_BlueprintBase::OnBecomeRelevant(UBehaviorTreeComponent* OwnerC
 	UBTDecorator_BlueprintBase* CDO = (UBTDecorator_BlueprintBase*)(GetClass()->GetDefaultObject());
 	if (BlackboardComp && CDO)
 	{
-		for (int32 i = 0; i < CDO->ObservedKeyNames.Num(); i++)
+		for (int32 NameIndex = 0; NameIndex < CDO->ObservedKeyNames.Num(); NameIndex++)
 		{
-			const FBlackboard::FKey KeyID = BlackboardComp->GetKeyID(CDO->ObservedKeyNames[i]);
+			const FBlackboard::FKey KeyID = BlackboardComp->GetKeyID(CDO->ObservedKeyNames[NameIndex]);
 			if (KeyID != FBlackboard::InvalidKey)
 			{
 				BlackboardComp->RegisterObserver(KeyID, BBKeyObserver);
@@ -88,9 +88,9 @@ void UBTDecorator_BlueprintBase::OnCeaseRelevant(UBehaviorTreeComponent* OwnerCo
 	UBTDecorator_BlueprintBase* CDO = (UBTDecorator_BlueprintBase*)(GetClass()->GetDefaultObject());
 	if (BlackboardComp && CDO)
 	{
-		for (int32 i = 0; i < CDO->ObservedKeyNames.Num(); i++)
+		for (int32 NameIndex = 0; NameIndex < CDO->ObservedKeyNames.Num(); NameIndex++)
 		{
-			const FBlackboard::FKey KeyID = BlackboardComp->GetKeyID(CDO->ObservedKeyNames[i]);
+			const FBlackboard::FKey KeyID = BlackboardComp->GetKeyID(CDO->ObservedKeyNames[NameIndex]);
 			if (KeyID != FBlackboard::InvalidKey)
 			{
 				BlackboardComp->UnregisterObserver(KeyID, BBKeyObserver);
diff --git a/Engine/Source/Runtime/AIModule/Private/BehaviorTree/Tasks/BTTask_RunEQSQuery.cpp b/Engine/Source/Runtime/AIModule/Private/BehaviorTree/Tasks/BTTask_RunEQSQuery.cpp
--- a/Engine/Source/Runtime/AIModule/Private/BehaviorTree/Tasks/BTTask_RunEQSQuery.cpp
+++ b/Engine/Source/Runtime/AIModule/Private/BehaviorTree/Tasks/BTTask_RunEQSQuery.cpp
@@ -120,9 +120,9 @@ void UBTTask_RunEQSQuery::OnQueryFinished(TSharedPtr<struct FEnvQueryResult> Res
 
 void UBTTask_RunEQSQuery::CollectKeyFilters()
 {
-	for (int32 i = 0; i < UEnvQueryManager::RegisteredItemTypes.Num(); i++)
+	for (int32 TypeIndex = 0; TypeIndex < UEnvQueryManager::RegisteredItemTypes.Num(); TypeIndex++)
 	{
-		UEnvQueryItemType* ItemTypeCDO = (UEnvQueryItemType*)UEnvQueryManager::RegisteredItemTypes[i]->GetDefaultObject();
+		UEnvQueryItemType* ItemTypeCDO = (UEnvQueryItemType*)UEnvQueryManager::RegisteredItemTypes[TypeIndex]->GetDefaultObject();
 		ItemTypeCDO->AddBlackboardFilters(BlackboardKey, this);
 	}
 }
diff --git a/Engine/Source/Runtime/AIModule/Private/BrainComponent.cpp b/Engine/Source/Runtime/AIModule/Private/BrainComponent.cpp
--- a/Engine/Source/Runtime/AIModule/Private/BrainComponent.cpp
+++ b/Engine/Source/Runtime/AIModule/Private/BrainComponent.cpp
@@ -258,9 +258,9 @@ void UBrainComponent::TickComponent(float DeltaTime, enum ELevelTick TickType, F
 	{
 		for (auto Message : MessagesToProcess)
 		{
-			for (int32 i = 0; i < MessageObservers.Num(); i++)
+			for (int32 ObserverIndex = 0; ObserverIndex < MessageObservers.Num(); ObserverIndex++)
 			{
-				MessageObservers[i]->OnMessage(Message);
+				MessageObservers[ObserverIndex]->OnMessage(Message);
 			}
 		}
 		MessagesToProcess.Reset();
diff --git a/Engine/Source/Runtime/AIModule/Private/EnvironmentQuery/EnvQuery.cpp b/Engine/Source/Runtime/AIModule/Private/EnvironmentQuery/EnvQuery.cpp
--- a/Engine/Source/Runtime/AIModule/Private/EnvironmentQuery/EnvQuery.cpp
+++ b/Engine/Source/Runtime/AIModule/Private/EnvironmentQuery/EnvQuery.cpp
@@ -12,9 +12,9 @@ UEnvQuery::UEnvQuery(const class FPostConstructInitializeProperties& PCIP) : Sup
 
 static bool HasNamedValue(const FName& ParamName, const TArray<FEnvNamedValue>& NamedValues)
 {
-	for (int32 i = 0; i < NamedValues.Num(); i++)
+	for (int32 ValueIndex = 0; ValueIndex < NamedValues.Num(); ValueIndex++)
 	{
-		if (NamedValues[i].ParamName == ParamName)
+		if (NamedValues[ValueIndex].ParamName == ParamName)
 		{
 			return true;
 		}
@@ -47,7 +47,9 @@ static void AddNamedValue(const FName& ParamName, const EEnvQueryParam::Type& Pa
 static void AddNamedValuesFromObject(const UObject* Ob, TArray<FEnvNamedValue>& NamedValues, TArray<FName>& RequiredParams)
 {
 	if (Ob == NULL)
+	{
 		return;
+	}
 
 	for (UProperty* TestProperty = Ob->GetClass()->PropertyLink; TestProperty; TestProperty = TestProperty->PropertyLinkNext)
 	{
@@ -83,24 +85,24 @@ void UEnvQuery::CollectQueryParams(TArray<FEnvNamedValue>& NamedValues) const
 	TArray<FName> RequiredParams;
 
 	// collect all params
-	for (int32 iOption = 0; iOption < Options.Num(); iOption++)
+	for (int32 OptionIndex = 0; OptionIndex < Options.Num(); OptionIndex++)
 	{
-		const UEnvQueryOption* Option = Options[iOption];
+		const UEnvQueryOption* Option = Options[OptionIndex];
 		AddNamedValuesFromObject(Option->Generator, NamedValues, RequiredParams);
 
-		for (int32 iTest = 0; iTest < Option->Tests.Num(); iTest++)
+		for (int32 TestIndex = 0; TestIndex < Option->Tests.Num(); TestIndex++)
 		{
-			const UEnvQueryTest* TestOb = Option->Tests[iTest];
+			const UEnvQueryTest* TestOb = Option->Tests[TestIndex];
 			AddNamedValuesFromObject(TestOb, NamedValues, RequiredParams);
 		}
 	}
 
 	// remove unnecessary params
-	for (int32 i = NamedValues.Num() - 1; i >= 0; i--)
+	for (int32 ValueIndex = NamedValues.Num() - 1; ValueIndex >= 0; ValueIndex--)
 	{
-		if (!RequiredParams.Contains(NamedValues[i].ParamName))
+		if (!RequiredParams.Contains(NamedValues[ValueIndex].ParamName))
 		{
-			NamedValues.RemoveAt(i);
+			NamedValues.RemoveAt(ValueIndex);
 		}
 	}
 }
diff --git a/Engine/Source/Runtime/AIModule/Private/EnvironmentQuery/EnvQueryInstance.cpp b/Engine/Source/Runtime/AIModule/Private/EnvironmentQuery/EnvQueryInstance.cpp
--- a/Engine/Source/Runtime/AIModule/Private/EnvironmentQuery/EnvQueryInstance.cpp
+++ b/Engine/Source/Runtime/AIModule/Private/EnvironmentQuery/EnvQueryInstance.cpp
@@ -94,10 +94,10 @@ bool FEnvQueryInstance::PrepareContext(UClass* Context, TArray<FEnvQuerySpatialD
 		uint8* RawData = ContextData.RawData.GetTypedData();
 
 		Data.Init(ContextData.NumValues);
-		for (int32 i = 0; i < ContextData.NumValues; i++)
+		for (int32 ValueIndex = 0; ValueIndex < ContextData.NumValues; ValueIndex++)
 		{
-			Data[i].Location = DefTypeOb->GetLocation(RawData);
-			Data[i].Rotation = DefTypeOb->GetRotation(RawData);
+			Data[ValueIndex].Location = DefTypeOb->GetLocation(RawData);
+			Data[ValueIndex].Rotation = DefTypeOb->GetRotation(RawData);
 			RawData += DefTypeValueSize;
 		}
 	}
@@ -122,9 +122,9 @@ bool FEnvQueryInstance::PrepareContext(UClass* Context, TArray<FVector>& Data)
 		uint8* RawData = (uint8*)ContextData.RawData.GetTypedData();
 
 		Data.Init(ContextData.NumValues);
-		for (int32 i = 0; i < ContextData.NumValues; i++)
+		for (int32 ValueIndex = 0; ValueIndex < ContextData.NumValues; ValueIndex++)
 		{
-			Data[i] = DefTypeOb->GetLocation(RawData);
+			Data[ValueIndex] = DefTypeOb->GetLocation(RawData);
 			RawData += DefTypeValueSize;
 		}
 	}
@@ -149,9 +149,9 @@ bool FEnvQueryInstance::PrepareContext(UClass* Context, TArray<FRotator>& Data)
 		uint8* RawData = ContextData.RawData.GetTypedData();
 
 		Data.Init(ContextData.NumValues);
-		for (int32 i = 0; i < ContextData.NumValues; i++)
+		for (int32 ValueIndex = 0; ValueIndex < ContextData.NumValues; ValueIndex++)
 		{
-			Data[i] = DefTypeOb->GetRotation(RawData);
+			Data[ValueIndex] = DefTypeOb->GetRotation(RawData);
 			RawData += DefTypeValueSize;
 		}
 	}
@@ -176,9 +176,9 @@ bool FEnvQueryInstance::PrepareContext(UClass* Context, TArray<AActor*>& Data)
 		uint8* RawData = ContextData.RawData.GetTypedData();
 
 		Data.Init(ContextData.NumValues);
-		for (int32 i = 0; i < ContextData.NumValues; i++)
+		for (int32 ValueIndex = 0; ValueIndex < ContextData.NumValues; ValueIndex++)
 		{
-			Data[i] = DefTypeOb->GetActor(RawData);
+			Data[ValueIndex] = DefTypeOb->GetActor(RawData);
 			RawData += DefTypeValueSize;
 		}
 	}
@@ -367,7 +367,7 @@ void FEnvQueryInstance::SortScores()
 {
 	if (Options[OptionIndex].bShuffleItems)
 	{
-		for (int32 i = 0; i < Items.Num(); i++)
+		for (int32 ItemIndex = 0; ItemIndex < Items.Num(); ItemIndex++)
 		{
 			const int32 Idx1 = FMath::RandHelper(Items.Num());
 			const int32 Idx2 = FMath::RandHelper(Items.Num());
@@ -410,11 +410,11 @@ void FEnvQueryInstance::PickBestItem()
 {
 	// find first valid item with score worse than best one
 	int32 NumBestItems = NumValidItems;
-	for (int32 i = 1; i < NumValidItems; i++)
+	for (int32 ItemIndex = 1; ItemIndex < NumValidItems; ItemIndex++)
 	{
-		if (Items[i].Score < Items[0].Score)
+		if (Items[ItemIndex].Score < Items[0].Score)
 		{
-			NumBestItems = i;
+			NumBestItems = ItemIndex;
 			break;
 		}
 	}
@@ -512,9 +512,9 @@ void FEnvQueryInstance::FinalizeGeneration()
 	if (NumValidItems > 0)
 	{
 		ItemDetails.Reserve(NumValidItems);
-		for (int32 i = 0; i < NumValidItems; i++)
+		for (int32 ItemIndex = 0; ItemIndex < NumValidItems; ItemIndex++)
 		{
-			ItemDetails.Add(FEnvQueryItemDetails(NumTests, i));
+			ItemDetails.Add(FEnvQueryItemDetails(NumTests, ItemIndex));
 		}
 	}
 
@@ -584,7 +584,7 @@ void FEnvQueryInstance::OnFinalCondition()
 		DebugData.Store(this);
 
 		// randomize before sorting, so items with the same score can be chosen randomly
-		for (int32 i = 0; i < Items.Num(); i++)
+		for (int32 ItemIndex = 0; ItemIndex < Items.Num(); ItemIndex++)
 		{
 			const int32 Idx1 = FMath::RandHelper(Items.Num());
 			const int32 Idx2 = FMath::RandHelper(Items.Num());
@@ -605,7 +605,7 @@ void FEnvQueryInstance::OnFinalCondition()
 		ItemDetails.Reset();
 
 		// randomize before sorting, so items with the same score can be chosen randomly
-		for (int32 i = 0; i < Items.Num(); i++)
+		for (int32 ItemIndex = 0; ItemIndex < Items.Num(); ItemIndex++)
 		{
 			Items.Swap(FMath::RandHelper(Items.Num()), FMath::RandHelper(Items.Num()));
 		}
@@ -628,9 +628,9 @@ uint32 FEnvQueryInstance::GetAllocatedSize() const
 	MemSize += ItemDetails.GetAllocatedSize();
 	MemSize += Options.GetAllocatedSize();
 
-	for (int32 i = 0; i < Options.Num(); i++)
+	for (int32 OptionIndex = 0; OptionIndex < Options.Num(); OptionIndex++)
 	{
-		MemSize += Options[i].GetAllocatedSize();
+		MemSize += Options[OptionIndex].GetAllocatedSize();
 	}
 
 	return MemSize;
diff --git a/Engine/Source/Runtime/AIModule/Private/EnvironmentQuery/EnvQueryTest.cpp b/Engine/Source/Runtime/AIModule/Private/EnvironmentQuery/EnvQueryTest.cpp
--- a/Engine/Source/Runtime/AIModule/Private/EnvironmentQuery/EnvQueryTest.cpp
+++ b/Engine/Source/Runtime/AIModule/Private/EnvironmentQuery/EnvQueryTest.cpp
@@ -85,9 +85,9 @@ void UEnvQueryTest::NormalizeItemScores(struct FEnvQueryInstance& QueryInstance)
 		(ClampMaxType == EEnvQueryTestClamping::None)
 	   )
 	{
-		for (int32 iItem = 0; iItem < QueryInstance.Items.Num(); iItem++, DetailInfo++)
+		for (int32 ItemIndex = 0; ItemIndex < QueryInstance.Items.Num(); ItemIndex++, DetailInfo++)
 		{
-			if (!QueryInstance.Items[iItem].IsValid())
+			if (!QueryInstance.Items[ItemIndex].IsValid())
 			{
 				continue;
 			}
diff --git a/Engine/Source/Runtime/AIModule/Private/EnvironmentQuery/Generators/EnvQueryGenerator_PathingGrid.cpp b/Engine/Source/Runtime/AIModule/Private/EnvironmentQuery/Generators/EnvQueryGenerator_PathingGrid.cpp
--- a/Engine/Source/Runtime/AIModule/Private/EnvironmentQuery/Generators/EnvQueryGenerator_PathingGrid.cpp
+++ b/Engine/Source/Runtime/AIModule/Private/EnvironmentQuery/Generators/EnvQueryGenerator_PathingGrid.cpp
@@ -118,11 +118,11 @@ void UEnvQueryGenerator_PathingGrid::FindNodeRefsInPathDistance(const class ARec
 	const bool bUseBacktracking = !bPathFromContext;
 	NavMesh->GetClustersWithinPathingDistance(ContextLocation, InMaxPathDistance, NodeRefs, bUseBacktracking);
 
-	for (int32 i = 0; i < NodeRefs.Num(); i++)
+	for (int32 RefIndex = 0; RefIndex < NodeRefs.Num(); RefIndex++)
 	{
 		FBox ClusterBounds;
 		
-		const bool bSuccess = NavMesh->GetClusterBounds(NodeRefs[i], ClusterBounds);
+		const bool bSuccess = NavMesh->GetClusterBounds(NodeRefs[RefIndex], ClusterBounds);
 		if (bSuccess)
 		{
 			MyBounds += ClusterBounds;
@@ -137,11 +137,11 @@ void UEnvQueryGenerator_PathingGrid::FindNodeRefsInPathDistance(const class ARec
 	NavMesh->GetPolysWithinPathingDistance(ContextLocation, InMaxPathDistance, NodeRefs, NavFilterInstance);
 
 	TArray<FVector> PolyVerts;
-	for (int32 i = 0; i < NodeRefs.Num(); i++)
+	for (int32 RefIndex = 0; RefIndex < NodeRefs.Num(); RefIndex++)
 	{
 		PolyVerts.Reset();
 		
-		const bool bSuccess = NavMesh->GetPolyVerts(NodeRefs[i], PolyVerts);
+		const bool bSuccess = NavMesh->GetPolyVerts(NodeRefs[RefIndex], PolyVerts);
 		if (bSuccess)
 		{
 			MyBounds += FBox(PolyVerts);
diff --git a/Engine/Source/Runtime/AIModule/Private/EnvironmentQuery/Generators/EnvQueryGenerator_ProjectedPoints.cpp b/Engine/Source/Runtime/AIModule/Private/EnvironmentQuery/Generators/EnvQueryGenerator_ProjectedPoints.cpp
--- a/Engine/Source/Runtime/AIModule/Private/EnvironmentQuery/Generators/EnvQueryGenerator_ProjectedPoints.cpp
+++ b/Engine/Source/Runtime/AIModule/Private/EnvironmentQuery/Generators/EnvQueryGenerator_ProjectedPoints.cpp
@@ -25,13 +25,13 @@ void UEnvQueryGenerator_ProjectedPoints::ProjectAndFilterNavPoints(TArray<FVecto
 
 			if (NavMesh)
 			{
-				for (int32 i = Points.Num() -1; i >= 0; i--)
+				for (int32 PointIndex = Points.Num() -1; PointIndex >= 0; PointIndex--)
 				{
 #if WITH_RECAST
-					const bool bProjected = ProjectNavPointInRange(Points[i], ProjectionData.ExtentX, ProjectionData.ProjectDown, ProjectionData.ProjectUp, NavMesh, NavigationFilter, TempPoints);
+					const bool bProjected = ProjectNavPointInRange(Points[PointIndex], ProjectionData.ExtentX, ProjectionData.ProjectDown, ProjectionData.ProjectUp, NavMesh, NavigationFilter, TempPoints);
 					if (!bProjected)
 					{
-						Points.RemoveAt(i);
+						Points.RemoveAt(PointIndex);
 					}
 #endif
 				}
@@ -41,12 +41,12 @@ void UEnvQueryGenerator_ProjectedPoints::ProjectAndFilterNavPoints(TArray<FVecto
 		{
 			const FVector ProjectionExtent(ProjectionData.ExtentX, ProjectionData.ExtentX, ProjectionData.ProjectDown);
 
-			for (int32 i = Points.Num() -1; i >= 0; i--)
+			for (int32 PointIndex = Points.Num() -1; PointIndex >= 0; PointIndex--)
 			{
-				const bool bProjected = ProjectNavPointSimple(Points[i], ProjectionExtent, NavData, NavigationFilter);
+				const bool bProjected = ProjectNavPointSimple(Points[PointIndex], ProjectionExtent, NavData, NavigationFilter);
 				if (!bProjected)
 				{
-					Points.RemoveAt(i);
+					Points.RemoveAt(PointIndex);
 				}
 			}
 		}
diff --git a/Engine/Source/Runtime/AIModule/Private/EnvironmentQuery/Generators/EnvQueryGenerator_SimpleGrid.cpp b/Engine/Source/Runtime/AIModule/Private/EnvironmentQuery/Generators/EnvQueryGenerator_SimpleGrid.cpp
--- a/Engine/Source/Runtime/AIModule/Private/EnvironmentQuery/Generators/EnvQueryGenerator_SimpleGrid.cpp
+++ b/Engine/Source/Runtime/AIModule/Private/EnvironmentQuery/Generators/EnvQueryGenerator_SimpleGrid.cpp
@@ -76,9 +76,9 @@ void UEnvQueryGenerator_SimpleGrid::GenerateItems(struct FEnvQueryInstance& Quer
 #endif // WITH_RECAST
 
 	QueryInstance.ReserveItemData(GridPoints.Num());
-	for (int32 i = 0; i < GridPoints.Num(); i++)
+	for (int32 PointIndex = 0; PointIndex < GridPoints.Num(); PointIndex++)
 	{
-		QueryInstance.AddItemData<UEnvQueryItemType_Point>(GridPoints[i]);
+		QueryInstance.AddItemData<UEnvQueryItemType_Point>(GridPoints[PointIndex]);
 	}
 }
 
diff --git a/Engine/Source/Runtime/AIModule/Private/EnvironmentQuery/Items/EnvQueryItemType_Direction.cpp b/Engine/Source/Runtime/AIModule/Private/EnvironmentQuery/Items/EnvQueryItemType_Direction.cpp
--- a/Engine/Source/Runtime/AIModule/Private/EnvironmentQuery/Items/EnvQueryItemType_Direction.cpp
+++ b/Engine/Source/Runtime/AIModule/Private/EnvironmentQuery/Items/EnvQueryItemType_Direction.cpp
@@ -64,9 +64,9 @@ void UEnvQueryItemType_Direction::SetContextHelper(struct FEnvQueryContextData&
 	ContextData.RawData.Init(sizeof(FVector) * MultipleDirections.Num());
 
 	uint8* RawData = (uint8*)ContextData.RawData.GetTypedData();
-	for (int32 i = 0; i < MultipleDirections.Num(); i++)
+	for (int32 DirectionIndex = 0; DirectionIndex < MultipleDirections.Num(); DirectionIndex++)
 	{
-		UEnvQueryItemType_Direction::SetValue(RawData, MultipleDirections[i]);
+		UEnvQueryItemType_Direction::SetValue(RawData, MultipleDirections[DirectionIndex]);
 		RawData += sizeof(FVector);
 	}
 }
@@ -78,9 +78,9 @@ void UEnvQueryItemType_Direction::SetContextHelper(struct FEnvQueryContextData&
 	ContextData.RawData.Init(sizeof(FVector) * MultipleRotations.Num());
 
 	uint8* RawData = (uint8*)ContextData.RawData.GetTypedData();
-	for (int32 i = 0; i < MultipleRotations.Num(); i++)
+	for (int32 RotationIndex = 0; RotationIndex < MultipleRotations.Num(); RotationIndex++)
 	{
-		UEnvQueryItemType_Direction::SetValueRot(RawData, MultipleRotations[i]);
+		UEnvQueryItemType_Direction::SetValueRot(RawData, MultipleRotations[RotationIndex]);
 		RawData += sizeof(FVector);
 	}
 }
diff --git a/Engine/Source/Runtime/AIModule/Private/EnvironmentQuery/Items/EnvQueryItemType_Point.cpp b/Engine/Source/Runtime/AIModule/Private/EnvironmentQuery/Items/EnvQueryItemType_Point.cpp
--- a/Engine/Source/Runtime/AIModule/Private/EnvironmentQuery/Items/EnvQueryItemType_Point.cpp
+++ b/Engine/Source/Runtime/AIModule/Private/EnvironmentQuery/Items/EnvQueryItemType_Point.cpp
@@ -39,9 +39,9 @@ void UEnvQueryItemType_Point::SetContextHelper(struct FEnvQueryContextData& Cont
 	ContextData.RawData.Init(sizeof(FVector) * MultiplePoints.Num());
 
 	uint8* RawData = (uint8*)ContextData.RawData.GetTypedData();
-	for (int32 i = 0; i < MultiplePoints.Num(); i++)
+	for (int32 PointIndex = 0; PointIndex < MultiplePoints.Num(); PointIndex++)
 	{
-		UEnvQueryItemType_Point::SetValue(RawData, MultiplePoints[i]);
+		UEnvQueryItemType_Point::SetValue(RawData, MultiplePoints[PointIndex]);
 		RawData += sizeof(FVector);
 	}
 }
diff --git a/Engine/Source/Runtime/AIModule/Private/EnvironmentQuery/Tests/EnvQueryTest_Distance.cpp b/Engine/Source/Runtime/AIModule/Private/EnvironmentQuery/Tests/EnvQueryTest_Distance.cpp
--- a/Engine/Source/Runtime/AIModule/Private/EnvironmentQuery/Tests/EnvQueryTest_Distance.cpp
+++ b/Engine/Source/Runtime/AIModule/Private/EnvironmentQuery/Tests/EnvQueryTest_Distance.cpp
@@ -65,9 +65,9 @@ void UEnvQueryTest_Distance::RunTest(struct FEnvQueryInstance& QueryInstance)
 		for (FEnvQueryInstance::ItemIterator It(this, QueryInstance); It; ++It)
 		{
 			const FVector ItemLocation = GetItemLocation(QueryInstance, *It);
-			for (int32 iContext = 0; iContext < ContextLocations.Num(); iContext++)
+			for (int32 ContextIndex = 0; ContextIndex < ContextLocations.Num(); ContextIndex++)
 			{
-				const float Distance = CalcDistance3D(ItemLocation, ContextLocations[iContext]);
+				const float Distance = CalcDistance3D(ItemLocation, ContextLocations[ContextIndex]);
 				It.SetScore(TestPurpose, FilterType, Distance, MinThresholdValue, MaxThresholdValue);
 			}
 		}
@@ -76,9 +76,9 @@ void UEnvQueryTest_Distance::RunTest(struct FEnvQueryInstance& QueryInstance)
 		for (FEnvQueryInstance::ItemIterator It(this, QueryInstance); It; ++It)
 		{
 			const FVector ItemLocation = GetItemLocation(QueryInstance, *It);
-			for (int32 iContext = 0; iContext < ContextLocations.Num(); iContext++)
+			for (int32 ContextIndex = 0; ContextIndex < ContextLocations.Num(); ContextIndex++)
 			{
-				const float Distance = CalcDistance2D(ItemLocation, ContextLocations[iContext]);
+				const float Distance = CalcDistance2D(ItemLocation, ContextLocations[ContextIndex]);
 				It.SetScore(TestPurpose, FilterType, Distance, MinThresholdValue, MaxThresholdValue);
 			}
 		}
@@ -87,9 +87,9 @@ void UEnvQueryTest_Distance::RunTest(struct FEnvQueryInstance& QueryInstance)
 		for (FEnvQueryInstance::ItemIterator It(this, QueryInstance); It; ++It)
 		{
 			const FVector ItemLocation = GetItemLocation(QueryInstance, *It);
-			for (int32 iContext = 0; iContext < ContextLocations.Num(); iContext++)
+			for (int32 ContextIndex = 0; ContextIndex < ContextLocations.Num(); ContextIndex++)
 			{
-				const float Distance = CalcDistanceZ(ItemLocation, ContextLocations[iContext]);
+				const float Distance = CalcDistanceZ(ItemLocation, ContextLocations[ContextIndex]);
 				It.SetScore(TestPurpose, FilterType, Distance, MinThresholdValue, MaxThresholdValue);
 			}
 		}
diff --git a/Engine/Source/Runtime/AIModule/Private/EnvironmentQuery/Tests/EnvQueryTest_Dot.cpp b/Engine/Source/Runtime/AIModule/Private/EnvironmentQuery/Tests/EnvQueryTest_Dot.cpp
--- a/Engine/Source/Runtime/AIModule/Private/EnvironmentQuery/Tests/EnvQueryTest_Dot.cpp
+++ b/Engine/Source/Runtime/AIModule/Private/EnvironmentQuery/Tests/EnvQueryTest_Dot.cpp
@@ -85,11 +85,11 @@ void UEnvQueryTest_Dot::RunTest(struct FEnvQueryInstance& QueryInstance)
 		}
 
 		// perform test for each line pair
-		for (int32 iLineA = 0; iLineA < LineADirs.Num(); iLineA++)
+		for (int32 LineAIndex = 0; LineAIndex < LineADirs.Num(); LineAIndex++)
 		{
-			for (int32 iLineB = 0; iLineB < LineBDirs.Num(); iLineB++)
+			for (int32 LineBIndex = 0; LineBIndex < LineBDirs.Num(); LineBIndex++)
 			{
-				const float DotValue = FVector::DotProduct(LineADirs[iLineA], LineBDirs[iLineB]);
+				const float DotValue = FVector::DotProduct(LineADirs[LineAIndex], LineBDirs[LineBIndex]);
 				It.SetScore(TestPurpose, FilterType, DotValue, MinThresholdValue, MaxThresholdValue);
 			}
 		}
@@ -109,7 +109,7 @@ void UEnvQueryTest_Dot::GatherLineDirections(TArray<FVector>& Directions, struct
 		QueryInstance.PrepareContext(LineFrom, ContextLocationFrom);
 	}
 	
-	for (int32 i = 0; i < ContextLocationFrom.Num(); i++)
+	for (int32 FromIndex = 0; FromIndex < ContextLocationFrom.Num(); FromIndex++)
 	{
 		TArray<FVector> ContextLocationTo;
 		if (IsContextPerItem(LineTo))
@@ -121,9 +121,9 @@ void UEnvQueryTest_Dot::GatherLineDirections(TArray<FVector>& Directions, struct
 			QueryInstance.PrepareContext(LineTo, ContextLocationTo);
 		}
 		
-		for (int32 i2 = 0; i2 < ContextLocationTo.Num(); i2++)
+		for (int32 ToIndex = 0; ToIndex < ContextLocationTo.Num(); ToIndex++)
 		{
-			const FVector Dir = (ContextLocationTo[i2] - ContextLocationFrom[i]).SafeNormal();
+			const FVector Dir = (ContextLocationTo[ToIndex] - ContextLocationFrom[FromIndex]).SafeNormal();
 			Directions.Add(Dir);
 		}
 	}
@@ -142,9 +142,9 @@ void UEnvQueryTest_Dot::GatherLineDirections(TArray<FVector>& Directions, struct
 		QueryInstance.PrepareContext(LineDirection, ContextRotations);
 	}
 
-	for (int32 i = 0; i < ContextRotations.Num(); i++)
+	for (int32 RotationIndex = 0; RotationIndex < ContextRotations.Num(); RotationIndex++)
 	{
-		const FVector Dir = ContextRotations[i].Vector();
+		const FVector Dir = ContextRotations[RotationIndex].Vector();
 		Directions.Add(Dir);
 	}
 }
diff --git a/Engine/Source/Runtime/AIModule/Private/EnvironmentQuery/Tests/EnvQueryTest_Pathfinding.cpp b/Engine/Source/Runtime/AIModule/Private/EnvironmentQuery/Tests/EnvQueryTest_Pathfinding.cpp
--- a/Engine/Source/Runtime/AIModule/Private/EnvironmentQuery/Tests/EnvQueryTest_Pathfinding.cpp
+++ b/Engine/Source/Runtime/AIModule/Private/EnvironmentQuery/Tests/EnvQueryTest_Pathfinding.cpp
@@ -71,9 +71,9 @@ void UEnvQueryTest_Pathfinding::RunTest(struct FEnvQueryInstance& QueryInstance)
 		for (FEnvQueryInstance::ItemIterator It(this, QueryInstance); It; ++It)
 		{
 			const FVector ItemLocation = GetItemLocation(QueryInstance, *It);
-			for (int32 iContext = 0; iContext < ContextLocations.Num(); iContext++)
+			for (int32 ContextIndex = 0; ContextIndex < ContextLocations.Num(); ContextIndex++)
 			{
-				const float PathValue = FindPathFunc.Execute(ItemLocation, ContextLocations[iContext], PFMode, NavData, NavSys, QueryInstance.Owner.Get());
+				const float PathValue = FindPathFunc.Execute(ItemLocation, ContextLocations[ContextIndex], PFMode, NavData, NavSys, QueryInstance.Owner.Get());
 				It.SetScore(TestPurpose, FilterType, PathValue, MinThresholdValue, MaxThresholdValue);
 
 				if (bDiscardFailed && PathValue >= BIG_NUMBER)
@@ -93,9 +93,9 @@ void UEnvQueryTest_Pathfinding::RunTest(struct FEnvQueryInstance& QueryInstance)
 		for (FEnvQueryInstance::ItemIterator It(this, QueryInstance); It; ++It)
 		{
 			const FVector ItemLocation = GetItemLocation(QueryInstance, *It);
-			for (int32 iContext = 0; iContext < ContextLocations.Num(); iContext++)
+			for (int32 ContextIndex = 0; ContextIndex < ContextLocations.Num(); ContextIndex++)
 			{
-				const bool bFoundPath = TestPathFunc.Execute(ItemLocation, ContextLocations[iContext], PFMode, NavData, NavSys, QueryInstance.Owner.Get());
+				const bool bFoundPath = TestPathFunc.Execute(ItemLocation, ContextLocations[ContextIndex], PFMode, NavData, NavSys, QueryInstance.Owner.Get());
 				It.SetScore(TestPurpose, FilterType, bFoundPath, bWantsPath);
 			}
 		}
diff --git a/Engine/Source/Runtime/AIModule/Private/EnvironmentQuery/Tests/EnvQueryTest_Trace.cpp b/Engine/Source/Runtime/AIModule/Private/EnvironmentQuery/Tests/EnvQueryTest_Trace.cpp
--- a/Engine/Source/Runtime/AIModule/Private/EnvironmentQuery/Tests/EnvQueryTest_Trace.cpp
+++ b/Engine/Source/Runtime/AIModule/Private/EnvironmentQuery/Tests/EnvQueryTest_Trace.cpp
@@ -79,9 +79,9 @@ void UEnvQueryTest_Trace::RunTest(struct FEnvQueryInstance& QueryInstance)
 		return;
 	}
 
-	for (int32 iContext = 0; iContext < ContextLocations.Num(); iContext++)
+	for (int32 ContextIndex = 0; ContextIndex < ContextLocations.Num(); ContextIndex++)
 	{
-		ContextLocations[iContext].Z += ContextZ;
+		ContextLocations[ContextIndex].Z += ContextZ;
 	}
 
 	for (FEnvQueryInstance::ItemIterator It(this, QueryInstance); It; ++It)
@@ -89,9 +89,9 @@ void UEnvQueryTest_Trace::RunTest(struct FEnvQueryInstance& QueryInstance)
 		const FVector ItemLocation = GetItemLocation(QueryInstance, *It) + FVector(0,0,ItemZ);
 		AActor* ItemActor = GetItemActor(QueryInstance, *It);
 
-		for (int32 iContext = 0; iContext < ContextLocations.Num(); iContext++)
+		for (int32 ContextIndex = 0; ContextIndex < ContextLocations.Num(); ContextIndex++)
 		{
-			const bool bHit = TraceFunc.Execute(ItemLocation, ContextLocations[iContext], ItemActor, QueryInstance.World, TraceCollisionChannel, TraceParams, TraceExtent);
+			const bool bHit = TraceFunc.Execute(ItemLocation, ContextLocations[ContextIndex], ItemActor, QueryInstance.World, TraceCollisionChannel, TraceParams, TraceExtent);
 			It.SetScore(TestPurpose, FilterType, bHit, bWantsHit);
 		}
 	}
diff --git a/Engine/Source/Runtime/AIModule/Private/Navigation/PathFollowingComponent.cpp b/Engine/Source/Runtime/AIModule/Private/Navigation/PathFollowingComponent.cpp
--- a/Engine/Source/Runtime/AIModule/Private/Navigation/PathFollowingComponent.cpp
+++ b/Engine/Source/Runtime/AIModule/Private/Navigation/PathFollowingComponent.cpp
@@ -1118,17 +1118,17 @@ bool UPathFollowingComponent::IsBlocked() const
 	if (LocationSamples.Num() == BlockDetectionSampleCount && BlockDetectionSampleCount > 0)
 	{
 		FVector Center = FVector::ZeroVector;
-		for (int32 i = 0; i < LocationSamples.Num(); i++)
+		for (int32 SampleIndex = 0; SampleIndex < LocationSamples.Num(); SampleIndex++)
 		{
-			Center += *LocationSamples[i];
+			Center += *LocationSamples[SampleIndex];
 		}
 
 		Center /= LocationSamples.Num();
 		bBlocked = true;
 
-		for (int32 i = 0; i < LocationSamples.Num(); i++)
+		for (int32 SampleIndex = 0; SampleIndex < LocationSamples.Num(); SampleIndex++)
 		{
-			const float TestDistance = FVector::DistSquared(*LocationSamples[i], Center);
+			const float TestDistance = FVector::DistSquared(*LocationSamples[SampleIndex], Center);
 			if (TestDistance > BlockDetectionDistance)
 			{
 				bBlocked = false;
@@ -1409,10 +1409,10 @@ FString UPathFollowingComponent::GetDebugString() const
 	GetDebugStringTokens(Tokens, Flags);
 
 	FString Desc;
-	for (int32 i = 0; i < Tokens.Num(); i++)
+	for (int32 TokenIndex = 0; TokenIndex < Tokens.Num(); TokenIndex++)
 	{
-		Desc += Tokens[i];
-		Desc += (Flags.IsValidIndex(i) && Flags[i] == EPathFollowingDebugTokens::ParamName) ? TEXT(": ") : TEXT(", ");
+		Desc += Tokens[TokenIndex];
+		Desc += (Flags.IsValidIndex(TokenIndex) && Flags[TokenIndex] == EPathFollowingDebugTokens::ParamName) ? TEXT(": ") : TEXT(", ");
 	}
 
 	return Desc;
