Some missing checkAtCompileTime -> static_assert change.

diff --git a/Engine/Source/Runtime/AIModule/Classes/BehaviorTree/BehaviorTreeTypes.h b/Engine/Source/Runtime/AIModule/Classes/BehaviorTree/BehaviorTreeTypes.h
--- a/Engine/Source/Runtime/AIModule/Classes/BehaviorTree/BehaviorTreeTypes.h
+++ b/Engine/Source/Runtime/AIModule/Classes/BehaviorTree/BehaviorTreeTypes.h
@@ -370,7 +370,7 @@ protected:
 	uint8 SelectedKeyID;
 	// SelectedKeyId type should be FBlackboard::FKey, but typedefs are not supported by UHT
 	// Please also check key id in FOnBlackboardChange delegate
-	checkAtCompileTime(sizeof(FBlackboard::FKey) == 1, FBlackboardKeySelector_SelectedKeyId_should_be_of_FBlackboard_FKey_type);
+	static_assert(sizeof(FBlackboard::FKey) == 1, "FBlackboardKeySelector::SelectedKeyId should be of FBlackboard::FKey type.");
 
 	UPROPERTY(transient, EditDefaultsOnly, Category=Blackboard)
 	uint32 bNoneIsAllowedValue:1;
