More refactoring of uint8 into FBlackboard::FKey #UE4

diff --git a/Engine/Source/Runtime/AIModule/Classes/BehaviorTree/BehaviorTreeTypes.h b/Engine/Source/Runtime/AIModule/Classes/BehaviorTree/BehaviorTreeTypes.h
--- a/Engine/Source/Runtime/AIModule/Classes/BehaviorTree/BehaviorTreeTypes.h
+++ b/Engine/Source/Runtime/AIModule/Classes/BehaviorTree/BehaviorTreeTypes.h
@@ -21,8 +21,17 @@ DECLARE_DWORD_COUNTER_STAT_EXTERN(TEXT("Num Templates"),STAT_AI_BehaviorTree_Num
 DECLARE_DWORD_COUNTER_STAT_EXTERN(TEXT("Num Instances"),STAT_AI_BehaviorTree_NumInstances,STATGROUP_AIBehaviorTree, );
 DECLARE_MEMORY_STAT_EXTERN(TEXT("Instance memory"),STAT_AI_BehaviorTree_InstanceMemory,STATGROUP_AIBehaviorTree, AIMODULE_API);
 
+namespace FBlackboard
+{
+	const FName KeySelf = TEXT("SelfActor");
+
+	typedef uint8 FKey;
+
+	const FKey InvalidKey = FKey(-1);
+}
+
 // delegate defines
-DECLARE_DELEGATE_TwoParams(FOnBlackboardChange, const class UBlackboardComponent*, uint8 /*key ID*/);
+DECLARE_DELEGATE_TwoParams(FOnBlackboardChange, const class UBlackboardComponent*, FBlackboard::FKey /*key ID*/);
 
 namespace BTSpecialChild
 {
@@ -134,15 +143,6 @@ namespace EBTDescriptionVerbosity
 	};
 }
 
-namespace FBlackboard
-{
-	const FName KeySelf = TEXT("SelfActor");
-
-	typedef uint8 FKey;
-
-	const FKey InvalidKey = FKey(-1);
-}
-
 /** debugger data about subtree instance */
 struct FBehaviorTreeDebuggerInstance
 {
@@ -369,8 +369,7 @@ protected:
 	UPROPERTY(transient, EditInstanceOnly, Category=Blackboard)
 	uint8 SelectedKeyID;
 	// SelectedKeyId type should be FBlackboard::FKey, but typedefs are not supported by UHT
-	// Please also check key id in FOnBlackboardChange delegate
-	static_assert(sizeof(FBlackboard::FKey) == 1, "FBlackboardKeySelector::SelectedKeyId should be of FBlackboard::FKey type.");
+	static_assert(sizeof(uint8) == sizeof(FBlackboard::FKey), "FBlackboardKeySelector::SelectedKeyId should be of FBlackboard::FKey-compatible type.");
 
 	UPROPERTY(transient, EditDefaultsOnly, Category=Blackboard)
 	uint32 bNoneIsAllowedValue:1;
diff --git a/Engine/Source/Runtime/AIModule/Classes/BehaviorTree/Decorators/BTDecorator_Blackboard.h b/Engine/Source/Runtime/AIModule/Classes/BehaviorTree/Decorators/BTDecorator_Blackboard.h
--- a/Engine/Source/Runtime/AIModule/Classes/BehaviorTree/Decorators/BTDecorator_Blackboard.h
+++ b/Engine/Source/Runtime/AIModule/Classes/BehaviorTree/Decorators/BTDecorator_Blackboard.h
@@ -15,7 +15,7 @@ class AIMODULE_API UBTDecorator_Blackboard : public UBTDecorator_BlackboardBase
 	GENERATED_UCLASS_BODY()
 
 	virtual bool CalculateRawConditionValue(class UBehaviorTreeComponent* OwnerComp, uint8* NodeMemory) const override;
-	virtual void OnBlackboardChange(const class UBlackboardComponent* Blackboard, uint8 ChangedKeyID) override;
+	virtual void OnBlackboardChange(const class UBlackboardComponent* Blackboard, FBlackboard::FKey ChangedKeyID) override;
 	virtual void DescribeRuntimeValues(const class UBehaviorTreeComponent* OwnerComp, uint8* NodeMemory, EBTDescriptionVerbosity::Type Verbosity, TArray<FString>& Values) const override;
 	virtual FString GetStaticDescription() const override;
 
diff --git a/Engine/Source/Runtime/AIModule/Classes/BehaviorTree/Decorators/BTDecorator_BlackboardBase.h b/Engine/Source/Runtime/AIModule/Classes/BehaviorTree/Decorators/BTDecorator_BlackboardBase.h
--- a/Engine/Source/Runtime/AIModule/Classes/BehaviorTree/Decorators/BTDecorator_BlackboardBase.h
+++ b/Engine/Source/Runtime/AIModule/Classes/BehaviorTree/Decorators/BTDecorator_BlackboardBase.h
@@ -12,7 +12,7 @@ class AIMODULE_API UBTDecorator_BlackboardBase : public UBTDecorator
 	virtual void InitializeFromAsset(class UBehaviorTree* Asset) override;
 
 	/** notify about change in blackboard keys */
-	virtual void OnBlackboardChange(const class UBlackboardComponent* Blackboard, uint8 ChangedKeyID);
+	virtual void OnBlackboardChange(const class UBlackboardComponent* Blackboard, FBlackboard::FKey ChangedKeyID);
 
 #if WITH_EDITOR
 	virtual FName GetNodeIconName() const override;
diff --git a/Engine/Source/Runtime/AIModule/Classes/BehaviorTree/Decorators/BTDecorator_BlueprintBase.h b/Engine/Source/Runtime/AIModule/Classes/BehaviorTree/Decorators/BTDecorator_BlueprintBase.h
--- a/Engine/Source/Runtime/AIModule/Classes/BehaviorTree/Decorators/BTDecorator_BlueprintBase.h
+++ b/Engine/Source/Runtime/AIModule/Classes/BehaviorTree/Decorators/BTDecorator_BlueprintBase.h
@@ -30,7 +30,7 @@ class AIMODULE_API UBTDecorator_BlueprintBase : public UBTDecorator
 	virtual void PostInitProperties() override;
 
 	/** notify about changes in blackboard */
-	void OnBlackboardChange(const class UBlackboardComponent* Blackboard, uint8 ChangedKeyID);
+	void OnBlackboardChange(const class UBlackboardComponent* Blackboard, FBlackboard::FKey ChangedKeyID);
 
 	virtual FString GetStaticDescription() const override;
 	virtual void DescribeRuntimeValues(const class UBehaviorTreeComponent* OwnerComp, uint8* NodeMemory, EBTDescriptionVerbosity::Type Verbosity, TArray<FString>& Values) const override;
diff --git a/Engine/Source/Runtime/AIModule/Classes/BehaviorTree/Decorators/BTDecorator_CompareBBEntries.h b/Engine/Source/Runtime/AIModule/Classes/BehaviorTree/Decorators/BTDecorator_CompareBBEntries.h
--- a/Engine/Source/Runtime/AIModule/Classes/BehaviorTree/Decorators/BTDecorator_CompareBBEntries.h
+++ b/Engine/Source/Runtime/AIModule/Classes/BehaviorTree/Decorators/BTDecorator_CompareBBEntries.h
@@ -44,7 +44,7 @@ public:
 
 	virtual void InitializeFromAsset(class UBehaviorTree* Asset) override;
 
-	virtual void OnBlackboardChange(const class UBlackboardComponent* Blackboard, uint8 ChangedKeyID);	
+	virtual void OnBlackboardChange(const class UBlackboardComponent* Blackboard, FBlackboard::FKey ChangedKeyID);
 	virtual void OnBecomeRelevant(UBehaviorTreeComponent* OwnerComp, uint8* NodeMemory) override;
 	virtual void OnCeaseRelevant(UBehaviorTreeComponent* OwnerComp, uint8* NodeMemory) override;
 
diff --git a/Engine/Source/Runtime/AIModule/Classes/BehaviorTree/Decorators/BTDecorator_ConeCheck.h b/Engine/Source/Runtime/AIModule/Classes/BehaviorTree/Decorators/BTDecorator_ConeCheck.h
--- a/Engine/Source/Runtime/AIModule/Classes/BehaviorTree/Decorators/BTDecorator_ConeCheck.h
+++ b/Engine/Source/Runtime/AIModule/Classes/BehaviorTree/Decorators/BTDecorator_ConeCheck.h
@@ -47,7 +47,7 @@ protected:
 
 	virtual bool CalculateRawConditionValue(class UBehaviorTreeComponent* OwnerComp, uint8* NodeMemory) const override;
 	virtual void OnBecomeRelevant(UBehaviorTreeComponent* OwnerComp, uint8* NodeMemory) override;
-	void OnBlackboardChange(const class UBlackboardComponent* Blackboard, uint8 ChangedKeyID);
+	void OnBlackboardChange(const class UBlackboardComponent* Blackboard, FBlackboard::FKey ChangedKeyID);
 
 	virtual void TickNode(UBehaviorTreeComponent* OwnerComp, uint8* NodeMemory, float DeltaSeconds) override;
 	
diff --git a/Engine/Source/Runtime/AIModule/Private/BehaviorTree/Decorators/BTDecorator_Blackboard.cpp b/Engine/Source/Runtime/AIModule/Private/BehaviorTree/Decorators/BTDecorator_Blackboard.cpp
--- a/Engine/Source/Runtime/AIModule/Private/BehaviorTree/Decorators/BTDecorator_Blackboard.cpp
+++ b/Engine/Source/Runtime/AIModule/Private/BehaviorTree/Decorators/BTDecorator_Blackboard.cpp
@@ -38,7 +38,7 @@ bool UBTDecorator_Blackboard::EvaluateOnBlackboard(const UBlackboardComponent* B
 	return bResult;
 }
 
-void UBTDecorator_Blackboard::OnBlackboardChange(const class UBlackboardComponent* Blackboard, uint8 ChangedKeyID)
+void UBTDecorator_Blackboard::OnBlackboardChange(const class UBlackboardComponent* Blackboard, FBlackboard::FKey ChangedKeyID)
 {
 	UBehaviorTreeComponent* BehaviorComp = Blackboard ? (UBehaviorTreeComponent*)Blackboard->GetBrainComponent() : NULL;
 	if (BlackboardKey.GetSelectedKeyID() == ChangedKeyID && BehaviorComp)
diff --git a/Engine/Source/Runtime/AIModule/Private/BehaviorTree/Decorators/BTDecorator_BlackboardBase.cpp b/Engine/Source/Runtime/AIModule/Private/BehaviorTree/Decorators/BTDecorator_BlackboardBase.cpp
--- a/Engine/Source/Runtime/AIModule/Private/BehaviorTree/Decorators/BTDecorator_BlackboardBase.cpp
+++ b/Engine/Source/Runtime/AIModule/Private/BehaviorTree/Decorators/BTDecorator_BlackboardBase.cpp
@@ -44,7 +44,7 @@ void UBTDecorator_BlackboardBase::OnCeaseRelevant(UBehaviorTreeComponent* OwnerC
 	}
 }
 
-void UBTDecorator_BlackboardBase::OnBlackboardChange(const UBlackboardComponent* Blackboard, uint8 ChangedKeyID)
+void UBTDecorator_BlackboardBase::OnBlackboardChange(const UBlackboardComponent* Blackboard, FBlackboard::FKey ChangedKeyID)
 {
 	UBehaviorTreeComponent* BehaviorComp = Blackboard ? (UBehaviorTreeComponent*)Blackboard->GetBrainComponent() : NULL;
 	if (BlackboardKey.GetSelectedKeyID() == ChangedKeyID && BehaviorComp)
diff --git a/Engine/Source/Runtime/AIModule/Private/BehaviorTree/Decorators/BTDecorator_BlueprintBase.cpp b/Engine/Source/Runtime/AIModule/Private/BehaviorTree/Decorators/BTDecorator_BlueprintBase.cpp
--- a/Engine/Source/Runtime/AIModule/Private/BehaviorTree/Decorators/BTDecorator_BlueprintBase.cpp
+++ b/Engine/Source/Runtime/AIModule/Private/BehaviorTree/Decorators/BTDecorator_BlueprintBase.cpp
@@ -185,7 +185,7 @@ void UBTDecorator_BlueprintBase::DescribeRuntimeValues(const class UBehaviorTree
 	}
 }
 
-void UBTDecorator_BlueprintBase::OnBlackboardChange(const UBlackboardComponent* Blackboard, uint8 ChangedKeyID)
+void UBTDecorator_BlueprintBase::OnBlackboardChange(const UBlackboardComponent* Blackboard, FBlackboard::FKey ChangedKeyID)
 {
 	UBehaviorTreeComponent* BehaviorComp = Blackboard ? (UBehaviorTreeComponent*)Blackboard->GetBrainComponent() : NULL;
 	if (BehaviorComp)
diff --git a/Engine/Source/Runtime/AIModule/Private/BehaviorTree/Decorators/BTDecorator_CompareBBEntries.cpp b/Engine/Source/Runtime/AIModule/Private/BehaviorTree/Decorators/BTDecorator_CompareBBEntries.cpp
--- a/Engine/Source/Runtime/AIModule/Private/BehaviorTree/Decorators/BTDecorator_CompareBBEntries.cpp
+++ b/Engine/Source/Runtime/AIModule/Private/BehaviorTree/Decorators/BTDecorator_CompareBBEntries.cpp
@@ -85,7 +85,7 @@ void UBTDecorator_CompareBBEntries::OnCeaseRelevant(UBehaviorTreeComponent* Owne
 	}
 }
 
-void UBTDecorator_CompareBBEntries::OnBlackboardChange(const UBlackboardComponent* Blackboard, uint8 ChangedKeyID)
+void UBTDecorator_CompareBBEntries::OnBlackboardChange(const UBlackboardComponent* Blackboard, FBlackboard::FKey ChangedKeyID)
 {
 	UBehaviorTreeComponent* BehaviorComp = Blackboard ? (UBehaviorTreeComponent*)Blackboard->GetBrainComponent() : NULL;
 	if (BehaviorComp && (BlackboardKeyA.GetSelectedKeyID() == ChangedKeyID || BlackboardKeyB.GetSelectedKeyID() == ChangedKeyID))
diff --git a/Engine/Source/Runtime/AIModule/Private/BehaviorTree/Decorators/BTDecorator_ConeCheck.cpp b/Engine/Source/Runtime/AIModule/Private/BehaviorTree/Decorators/BTDecorator_ConeCheck.cpp
--- a/Engine/Source/Runtime/AIModule/Private/BehaviorTree/Decorators/BTDecorator_ConeCheck.cpp
+++ b/Engine/Source/Runtime/AIModule/Private/BehaviorTree/Decorators/BTDecorator_ConeCheck.cpp
@@ -87,7 +87,7 @@ void UBTDecorator_ConeCheck::OnBecomeRelevant(UBehaviorTreeComponent* OwnerComp,
 	DecoratorMemory->bLastRawResult = CalcConditionImpl(OwnerComp, NodeMemory);
 }
 
-void UBTDecorator_ConeCheck::OnBlackboardChange(const class UBlackboardComponent* Blackboard, uint8 ChangedKeyID)
+void UBTDecorator_ConeCheck::OnBlackboardChange(const class UBlackboardComponent* Blackboard, FBlackboard::FKey ChangedKeyID)
 {
 	check(false);
 }
