Fixed post processing breaking whenever object is selected in editor. Commented out broken code in AIModule.

diff --git a/Engine/Source/Runtime/AIModule/Private/EnvironmentQuery/EQSTestingPawn.cpp b/Engine/Source/Runtime/AIModule/Private/EnvironmentQuery/EQSTestingPawn.cpp
--- a/Engine/Source/Runtime/AIModule/Private/EnvironmentQuery/EQSTestingPawn.cpp
+++ b/Engine/Source/Runtime/AIModule/Private/EnvironmentQuery/EQSTestingPawn.cpp
@@ -83,6 +83,8 @@ AEQSTestingPawn::AEQSTestingPawn(const class FPostConstructInitializeProperties&
 
 void AEQSTestingPawn::OnEditorSelectionChanged(UObject* NewSelection)
 {
+	/* This is totally busted and should not depend on gameplay debugger.
+
 	bool bEQSPawnSelected = Cast<AEQSTestingPawn>(NewSelection) != NULL;
 	if (bEQSPawnSelected == false)
 	{
@@ -101,6 +103,7 @@ void AEQSTestingPawn::OnEditorSelectionChanged(UObject* NewSelection)
 		GCurrentLevelEditingViewportClient->EngineShowFlags.SetSingleFlag(FGameplayDebuggerSettings::ShowFlagIndex, bEQSPawnSelected);
 	}
 #endif // WITH_EDITOR
+	*/
 }
 
 void AEQSTestingPawn::TickActor(float DeltaTime, enum ELevelTick TickType, FActorTickFunction& ThisTickFunction)
