Fixed PathFollowing Agent's radius taken into consideration in "has reached" checks even if it was undesired #UE4

diff --git a/Engine/Source/Runtime/AIModule/Private/Navigation/PathFollowingComponent.cpp b/Engine/Source/Runtime/AIModule/Private/Navigation/PathFollowingComponent.cpp
--- a/Engine/Source/Runtime/AIModule/Private/Navigation/PathFollowingComponent.cpp
+++ b/Engine/Source/Runtime/AIModule/Private/Navigation/PathFollowingComponent.cpp
@@ -940,7 +940,7 @@ bool UPathFollowingComponent::HasReachedInternal(const FVector& Goal, float Goal
 	const FVector ToGoal = Goal - AgentLocation;
 
 	const float Dist2DSq = ToGoal.SizeSquared2D();
-	const float UseRadius = GoalRadius + RadiusThreshold + (bUseAgentRadius ? AgentRadius : 0.0f) + (AgentRadius * MinAgentRadiusPct);
+	const float UseRadius = GoalRadius + RadiusThreshold + (bUseAgentRadius ? AgentRadius * (MinAgentRadiusPct + 1.f) : 0.0f);
 	if (Dist2DSq > FMath::Square(UseRadius))
 	{
 		return false;
@@ -1008,7 +1008,7 @@ void UPathFollowingComponent::DebugReachTest(float& CurrentDot, float& CurrentDi
 	MovingAgent->GetSimpleCollisionCylinder(AgentRadius, AgentHalfHeight);
 
 	CurrentDistance = ToGoal.Size2D();
-	const float UseRadius = GoalRadius + RadiusThreshold + (bUseAgentRadius ? AgentRadius : 0.0f) + (AgentRadius * MinAgentRadiusPct);
+	const float UseRadius = GoalRadius + RadiusThreshold + (bUseAgentRadius ? AgentRadius * (MinAgentRadiusPct + 1.f) : 0.0f);
 	bDistanceFailed = (CurrentDistance > UseRadius) ? 1 : 0;
 
 	CurrentHeight = FMath::Abs(ToGoal.Z);
