Added FRotator as new BT blackboard type

Added explicit constructor for FRotator which takes a single float
Added FAISystem::IsValidRotation
Added BlackboardKeyType_Rotator
Allow to filter by FRotator in FBlackboardKeySelector
Added Get/Set rotator functions in BlackboardComponent
Fixed typo
Updated BlackBoardComponent to check type of the key in the functions which take the keyid and not the FName

GitHub pull request #227
https://github.com/EpicGames/UnrealEngine/pull/227

diff --git a/Engine/Source/Runtime/AIModule/Classes/AITypes.h b/Engine/Source/Runtime/AIModule/Classes/AITypes.h
--- a/Engine/Source/Runtime/AIModule/Classes/AITypes.h
+++ b/Engine/Source/Runtime/AIModule/Classes/AITypes.h
@@ -7,11 +7,17 @@
 namespace FAISystem
 {
 	static const FVector InvalidLocation(FLT_MAX);
+	static const FRotator InvalidRotation(FLT_MAX);
 
 	FORCEINLINE bool IsValidLocation(const FVector& TestLocation)
 	{
 		return TestLocation != InvalidLocation;
 	}
+
+	FORCEINLINE bool IsValidRotation(const FRotator& TestRotation)
+	{
+		return TestRotation != InvalidRotation;
+	}
 }
 
 UENUM()
diff --git a/Engine/Source/Runtime/AIModule/Classes/BehaviorTree/BehaviorTreeTypes.h b/Engine/Source/Runtime/AIModule/Classes/BehaviorTree/BehaviorTreeTypes.h
--- a/Engine/Source/Runtime/AIModule/Classes/BehaviorTree/BehaviorTreeTypes.h
+++ b/Engine/Source/Runtime/AIModule/Classes/BehaviorTree/BehaviorTreeTypes.h
@@ -394,6 +394,7 @@ public:
 	void AddFloatFilter(UObject* Owner);
 	void AddBoolFilter(UObject* Owner);
 	void AddVectorFilter(UObject* Owner);
+	void AddRotatorFilter(UObject* Owner);
 	void AddStringFilter(UObject* Owner);
 	void AddNameFilter(UObject* Owner);
 
diff --git a/Engine/Source/Runtime/AIModule/Classes/BehaviorTree/Blackboard/BlackboardKeyAllTypes.h b/Engine/Source/Runtime/AIModule/Classes/BehaviorTree/Blackboard/BlackboardKeyAllTypes.h
--- a/Engine/Source/Runtime/AIModule/Classes/BehaviorTree/Blackboard/BlackboardKeyAllTypes.h
+++ b/Engine/Source/Runtime/AIModule/Classes/BehaviorTree/Blackboard/BlackboardKeyAllTypes.h
@@ -13,4 +13,5 @@
 #include "BehaviorTree/Blackboard/BlackboardKeyType_Object.h"
 #include "BehaviorTree/Blackboard/BlackboardKeyType_String.h"
 #include "BehaviorTree/Blackboard/BlackboardKeyType_Vector.h"
+#include "BehaviorTree/Blackboard/BlackboardKeyType_Rotator.h"
 #include "BehaviorTree/BlackboardComponent.h"
\ No newline at end of file
diff --git a/Engine/Source/Runtime/AIModule/Classes/BehaviorTree/Blackboard/BlackboardKeyType_Rotator.h b/Engine/Source/Runtime/AIModule/Classes/BehaviorTree/Blackboard/BlackboardKeyType_Rotator.h
new file mode 100644
--- /dev/null
+++ b/Engine/Source/Runtime/AIModule/Classes/BehaviorTree/Blackboard/BlackboardKeyType_Rotator.h
@@ -0,0 +1,18 @@
+// Copyright 1998-2014 Epic Games, Inc. All Rights Reserved.
+#pragma once
+#include "BlackboardKeyType_Rotator.generated.h"
+
+UCLASS(EditInlineNew, MinimalAPI)
+class UBlackboardKeyType_Rotator : public UBlackboardKeyType
+{
+	GENERATED_UCLASS_BODY()
+
+	static FRotator GetValue(const uint8* RawData);
+	static bool SetValue(uint8* RawData, const FRotator& Value);
+
+	virtual void Initialize(uint8* RawData) const override;
+	virtual FString DescribeValue(const uint8* RawData) const override;
+	virtual bool GetRotation(const uint8* RawData, FRotator& Rotation) const override;
+	virtual EBlackboardCompare::Type Compare(const uint8* MemoryBlockA, const uint8* MemoryBlockB) const override;
+	virtual bool TestBasicOperation(const uint8* MemoryBlock, EBasicKeyOperation::Type Op) const override;
+};
diff --git a/Engine/Source/Runtime/AIModule/Classes/BehaviorTree/BlackboardComponent.h b/Engine/Source/Runtime/AIModule/Classes/BehaviorTree/BlackboardComponent.h
--- a/Engine/Source/Runtime/AIModule/Classes/BehaviorTree/BlackboardComponent.h
+++ b/Engine/Source/Runtime/AIModule/Classes/BehaviorTree/BlackboardComponent.h
@@ -77,7 +77,7 @@ class AIMODULE_API UBlackboardComponent : public UActorComponent
 	void InitializeBlackboard(class UBlackboardData* NewAsset);
 	
 	/** @return true if component can be used with specified blackboard asset */
-	bool IsCompatibileWith(class UBlackboardData* TestAsset) const;
+	bool IsCompatibleWith(class UBlackboardData* TestAsset) const;
 
 	UFUNCTION(BlueprintCallable, Category="AI|Components|Blackboard")
 	UObject* GetValueAsObject(const FName& KeyName) const;
@@ -115,6 +115,10 @@ class AIMODULE_API UBlackboardComponent : public UActorComponent
 	FVector GetValueAsVector(const FName& KeyName) const;
 	FVector GetValueAsVector(FBlackboard::FKey KeyID) const;
 
+	UFUNCTION(BlueprintCallable, Category="AI|Components|Blackboard")
+	FRotator GetValueAsRotator(const FName& KeyName) const;
+	FRotator GetValueAsRotator(uint8 KeyID) const;
+
 	UFUNCTION(BlueprintCallable, Category="AI|Components|Blackboard")
 	void SetValueAsObject(const FName& KeyName, UObject* ObjectValue);
 	void SetValueAsObject(FBlackboard::FKey KeyID, UObject* ObjectValue);
@@ -155,6 +159,14 @@ class AIMODULE_API UBlackboardComponent : public UActorComponent
 	void ClearValueAsVector(const FName& KeyName);
 	void ClearValueAsVector(FBlackboard::FKey KeyID);
 
+	UFUNCTION(BlueprintCallable, Category="AI|Components|Blackboard")
+	void SetValueAsRotator(const FName& KeyName, FRotator VectorValue);
+	void SetValueAsRotator(uint8 KeyID, const FRotator& VectorValue);
+
+	UFUNCTION(BlueprintCallable, Category="Blackboard")
+	void ClearValueAsRotator(const FName& KeyName);
+	void ClearValueAsRotator(uint8 KeyID);
+
 	/** return false if call failed (most probably no such entry in BB) */
 	UFUNCTION(BlueprintCallable, Category="AI|Components|Blackboard")
 	bool GetLocationFromEntry(const FName& KeyName, FVector& ResultLocation) const;
diff --git a/Engine/Source/Runtime/AIModule/Private/BehaviorTree/BehaviorTreeComponent.cpp b/Engine/Source/Runtime/AIModule/Private/BehaviorTree/BehaviorTreeComponent.cpp
--- a/Engine/Source/Runtime/AIModule/Private/BehaviorTree/BehaviorTreeComponent.cpp
+++ b/Engine/Source/Runtime/AIModule/Private/BehaviorTree/BehaviorTreeComponent.cpp
@@ -1197,7 +1197,7 @@ bool UBehaviorTreeComponent::PushInstance(class UBehaviorTree* TreeAsset)
 	}
 
 	// check if blackboard class match
-	if (TreeAsset->BlackboardAsset && BlackboardComp && !BlackboardComp->IsCompatibileWith(TreeAsset->BlackboardAsset))
+	if (TreeAsset->BlackboardAsset && BlackboardComp && !BlackboardComp->IsCompatibleWith(TreeAsset->BlackboardAsset))
 	{
 		UE_VLOG(GetOwner(), LogBehaviorTree, Warning, TEXT("Failed to execute tree %s: blackboard %s is not compatibile with current: %s!"),
 			*GetNameSafe(TreeAsset), *GetNameSafe(TreeAsset->BlackboardAsset), *GetNameSafe(BlackboardComp->GetBlackboardAsset()));
diff --git a/Engine/Source/Runtime/AIModule/Private/BehaviorTree/BehaviorTreeTypes.cpp b/Engine/Source/Runtime/AIModule/Private/BehaviorTree/BehaviorTreeTypes.cpp
--- a/Engine/Source/Runtime/AIModule/Private/BehaviorTree/BehaviorTreeTypes.cpp
+++ b/Engine/Source/Runtime/AIModule/Private/BehaviorTree/BehaviorTreeTypes.cpp
@@ -239,6 +239,11 @@ void FBlackboardKeySelector::AddVectorFilter(UObject* Owner)
 	AllowedTypes.Add(NewNamedObject<UBlackboardKeyType_Vector>(Owner, TEXT("BlackboardKeyType_Vector")));
 }
 
+void FBlackboardKeySelector::AddRotatorFilter(UObject* Owner)
+{
+	AllowedTypes.Add(NewObject<UBlackboardKeyType_Rotator>(Owner));
+}
+
 void FBlackboardKeySelector::AddStringFilter(UObject* Owner)
 {
 	AllowedTypes.Add(NewNamedObject<UBlackboardKeyType_String>(Owner, TEXT("BlackboardKeyType_String")));
diff --git a/Engine/Source/Runtime/AIModule/Private/BehaviorTree/Blackboard/BlackboardKeyType_Rotator.cpp b/Engine/Source/Runtime/AIModule/Private/BehaviorTree/Blackboard/BlackboardKeyType_Rotator.cpp
new file mode 100644
--- /dev/null
+++ b/Engine/Source/Runtime/AIModule/Private/BehaviorTree/Blackboard/BlackboardKeyType_Rotator.cpp
@@ -0,0 +1,49 @@
+// Copyright 1998-2014 Epic Games, Inc. All Rights Reserved.
+
+#include "AIModulePrivate.h"
+#include "BehaviorTree/Blackboard/BlackboardKeyType_Rotator.h"
+
+UBlackboardKeyType_Rotator::UBlackboardKeyType_Rotator(const class FPostConstructInitializeProperties& PCIP) : Super(PCIP)
+{
+	ValueSize = sizeof(FRotator);
+	SupportedOp = EBlackboardKeyOperation::Basic;
+}
+
+FRotator UBlackboardKeyType_Rotator::GetValue(const uint8* RawData)
+{
+	return GetValueFromMemory<FRotator>(RawData);
+}
+
+bool UBlackboardKeyType_Rotator::SetValue(uint8* RawData, const FRotator& Value)
+{
+	return SetValueInMemory<FRotator>(RawData, Value);
+}
+
+FString UBlackboardKeyType_Rotator::DescribeValue(const uint8* RawData) const
+{
+	const FRotator Rotation = GetValue(RawData);
+	return FAISystem::IsValidRotation(Rotation) ? Rotation.ToString() : TEXT("(invalid)");
+}
+
+bool UBlackboardKeyType_Rotator::GetRotation(const uint8* RawData, FRotator& Rotation) const
+{
+	Rotation = GetValue(RawData);
+	return FAISystem::IsValidRotation(Rotation);
+}
+
+void UBlackboardKeyType_Rotator::Initialize(uint8* RawData) const
+{
+	SetValue(RawData, FAISystem::InvalidRotation);
+}
+
+EBlackboardCompare::Type UBlackboardKeyType_Rotator::Compare(const uint8* MemoryBlockA, const uint8* MemoryBlockB) const
+{
+	return GetValueFromMemory<FRotator>(MemoryBlockA).Equals(GetValueFromMemory<FRotator>(MemoryBlockB))
+		? EBlackboardCompare::Equal : EBlackboardCompare::NotEqual;
+}
+
+bool UBlackboardKeyType_Rotator::TestBasicOperation(const uint8* MemoryBlock, EBasicKeyOperation::Type Op) const
+{
+	const FRotator Rotation = GetValue(MemoryBlock);
+	return (Op == EBasicKeyOperation::Set) ? FAISystem::IsValidRotation(Rotation) : !FAISystem::IsValidRotation(Rotation);
+}
diff --git a/Engine/Source/Runtime/AIModule/Private/BehaviorTree/BlackboardComponent.cpp b/Engine/Source/Runtime/AIModule/Private/BehaviorTree/BlackboardComponent.cpp
--- a/Engine/Source/Runtime/AIModule/Private/BehaviorTree/BlackboardComponent.cpp
+++ b/Engine/Source/Runtime/AIModule/Private/BehaviorTree/BlackboardComponent.cpp
@@ -189,7 +189,7 @@ void UBlackboardComponent::NotifyObservers(FBlackboard::FKey KeyID) const
 	}
 }
 
-bool UBlackboardComponent::IsCompatibileWith(class UBlackboardData* TestAsset) const
+bool UBlackboardComponent::IsCompatibleWith(class UBlackboardData* TestAsset) const
 {
 	for (UBlackboardData* It = BlackboardAsset; It; It = It->Parent)
 	{
@@ -467,17 +467,37 @@ FVector UBlackboardComponent::GetValueAsVector(FBlackboard::FKey KeyID) const
 	return RawData ? UBlackboardKeyType_Vector::GetValue(RawData) : FVector::ZeroVector;
 }
 
-void UBlackboardComponent::SetValueAsObject(const FName& KeyName, UObject* ObjectValue)
+
+FRotator UBlackboardComponent::GetValueAsRotator(const FName& KeyName) const
 {
-	const FBlackboard::FKey KeyID = GetKeyID(KeyName);
-	if (GetKeyType(KeyID) == UBlackboardKeyType_Object::StaticClass())
+	uint8 KeyID = GetKeyID(KeyName);
+	if (GetKeyType(KeyID) != UBlackboardKeyType_Rotator::StaticClass())
 	{
-		SetValueAsObject(KeyID, ObjectValue);
+		KeyID = INDEX_NONE;
 	}
+
+	return GetValueAsRotator(KeyID);
+}
+
+FRotator UBlackboardComponent::GetValueAsRotator(uint8 KeyID) const
+{
+	const uint8* RawData = GetKeyRawData(KeyID);
+	return RawData ? UBlackboardKeyType_Rotator::GetValue(RawData) : FRotator::ZeroRotator;
+}
+
+void UBlackboardComponent::SetValueAsObject(const FName& KeyName, UObject* ObjectValue)
+{
+	const FBlackboard::FKey KeyID = GetKeyID(KeyName);
+	SetValueAsObject(KeyID, ObjectValue);
 }
 
 void UBlackboardComponent::SetValueAsObject(FBlackboard::FKey KeyID, UObject* ObjectValue)
 {
+	if ( GetKeyType(KeyID) != UBlackboardKeyType_Object::StaticClass() )
+	{
+		return;
+	}
+
 	uint8* RawData = GetKeyRawData(KeyID);
 	if (RawData)
 	{
@@ -492,14 +512,16 @@ void UBlackboardComponent::SetValueAsObject(FBlackboard::FKey KeyID, UObject* Ob
 void UBlackboardComponent::SetValueAsClass(const FName& KeyName, UClass* ClassValue)
 {
 	const FBlackboard::FKey KeyID = GetKeyID(KeyName);
-	if (GetKeyType(KeyID) == UBlackboardKeyType_Class::StaticClass())
-	{
-		SetValueAsClass(KeyID, ClassValue);
-	}
+	SetValueAsClass(KeyID, ClassValue);
 }
 
 void UBlackboardComponent::SetValueAsClass(FBlackboard::FKey KeyID, UClass* ClassValue)
 {
+	if (GetKeyType(KeyID) != UBlackboardKeyType_Class::StaticClass())
+	{
+		return;
+	}
+
 	uint8* RawData = GetKeyRawData(KeyID);
 	if (RawData)
 	{
@@ -514,15 +536,17 @@ void UBlackboardComponent::SetValueAsClass(FBlackboard::FKey KeyID, UClass* Clas
 void UBlackboardComponent::SetValueAsEnum(const FName& KeyName, uint8 EnumValue)
 {
 	const FBlackboard::FKey KeyID = GetKeyID(KeyName);
-	if (GetKeyType(KeyID) == UBlackboardKeyType_Enum::StaticClass() ||
-		GetKeyType(KeyID) == UBlackboardKeyType_NativeEnum::StaticClass())
-	{
-		SetValueAsEnum(KeyID, EnumValue);
-	}
+	SetValueAsEnum(KeyID, EnumValue);
 }
 
 void UBlackboardComponent::SetValueAsEnum(FBlackboard::FKey KeyID, uint8 EnumValue)
 {
+	if (GetKeyType(KeyID) != UBlackboardKeyType_Enum::StaticClass() &&
+		GetKeyType(KeyID) != UBlackboardKeyType_NativeEnum::StaticClass())
+	{
+		return;
+	}
+
 	uint8* RawData = GetKeyRawData(KeyID);
 	if (RawData)
 	{
@@ -537,14 +561,16 @@ void UBlackboardComponent::SetValueAsEnum(FBlackboard::FKey KeyID, uint8 EnumVal
 void UBlackboardComponent::SetValueAsInt(const FName& KeyName, int32 IntValue)
 {
 	const FBlackboard::FKey KeyID = GetKeyID(KeyName);
-	if (GetKeyType(KeyID) == UBlackboardKeyType_Int::StaticClass())
-	{
-		SetValueAsInt(KeyID, IntValue);
-	}
+	SetValueAsInt(KeyID, IntValue);
 }
 
 void UBlackboardComponent::SetValueAsInt(FBlackboard::FKey KeyID, int32 IntValue)
 {
+	if (GetKeyType(KeyID) != UBlackboardKeyType_Int::StaticClass())
+	{
+		return;
+	}
+
 	uint8* RawData = GetKeyRawData(KeyID);
 	if (RawData)
 	{
@@ -559,14 +585,16 @@ void UBlackboardComponent::SetValueAsInt(FBlackboard::FKey KeyID, int32 IntValue
 void UBlackboardComponent::SetValueAsFloat(const FName& KeyName, float FloatValue)
 {
 	const FBlackboard::FKey KeyID = GetKeyID(KeyName);
-	if (GetKeyType(KeyID) == UBlackboardKeyType_Float::StaticClass())
-	{
-		SetValueAsFloat(KeyID, FloatValue);
-	}
+	SetValueAsFloat(KeyID, FloatValue);
 }
 
 void UBlackboardComponent::SetValueAsFloat(FBlackboard::FKey KeyID, float FloatValue)
 {
+	if (GetKeyType(KeyID) != UBlackboardKeyType_Float::StaticClass())
+	{
+		return;
+	}
+
 	uint8* RawData = GetKeyRawData(KeyID);
 	if (RawData)
 	{
@@ -581,14 +609,16 @@ void UBlackboardComponent::SetValueAsFloat(FBlackboard::FKey KeyID, float FloatV
 void UBlackboardComponent::SetValueAsBool(const FName& KeyName, bool BoolValue)
 {
 	const FBlackboard::FKey KeyID = GetKeyID(KeyName);
-	if (GetKeyType(KeyID) == UBlackboardKeyType_Bool::StaticClass())
-	{
-		SetValueAsBool(KeyID, BoolValue);
-	}
+	SetValueAsBool(KeyID, BoolValue);
 }
 
 void UBlackboardComponent::SetValueAsBool(FBlackboard::FKey KeyID, bool BoolValue)
 {
+	if (GetKeyType(KeyID) != UBlackboardKeyType_Bool::StaticClass())
+	{
+		return;
+	}
+
 	uint8* RawData = GetKeyRawData(KeyID);
 	if (RawData)
 	{
@@ -603,14 +633,16 @@ void UBlackboardComponent::SetValueAsBool(FBlackboard::FKey KeyID, bool BoolValu
 void UBlackboardComponent::SetValueAsString(const FName& KeyName, FString StringValue)
 {
 	const FBlackboard::FKey KeyID = GetKeyID(KeyName);
-	if (GetKeyType(KeyID) == UBlackboardKeyType_String::StaticClass())
-	{
-		SetValueAsString(KeyID, StringValue);
-	}
+	SetValueAsString(KeyID, StringValue);
 }
 
 void UBlackboardComponent::SetValueAsString(FBlackboard::FKey KeyID, const FString& StringValue)
 {
+	if (GetKeyType(KeyID) != UBlackboardKeyType_String::StaticClass())
+	{
+		return;
+	}
+
 	uint8* RawData = GetKeyRawData(KeyID);
 	if (RawData)
 	{
@@ -625,14 +657,16 @@ void UBlackboardComponent::SetValueAsString(FBlackboard::FKey KeyID, const FStri
 void UBlackboardComponent::SetValueAsName(const FName& KeyName, FName NameValue)
 {
 	const FBlackboard::FKey KeyID = GetKeyID(KeyName);
-	if (GetKeyType(KeyID) == UBlackboardKeyType_Name::StaticClass())
-	{
-		SetValueAsName(KeyID, NameValue);
-	}
+	SetValueAsName(KeyID, NameValue);
 }
 
 void UBlackboardComponent::SetValueAsName(FBlackboard::FKey KeyID, const FName& NameValue)
 {
+	if (GetKeyType(KeyID) != UBlackboardKeyType_Name::StaticClass())
+	{
+		return;
+	}
+
 	uint8* RawData = GetKeyRawData(KeyID);
 	if (RawData)
 	{
@@ -647,14 +681,16 @@ void UBlackboardComponent::SetValueAsName(FBlackboard::FKey KeyID, const FName&
 void UBlackboardComponent::SetValueAsVector(const FName& KeyName, FVector VectorValue)
 {
 	const FBlackboard::FKey KeyID = GetKeyID(KeyName);
-	if (GetKeyType(KeyID) == UBlackboardKeyType_Vector::StaticClass())
-	{
-		SetValueAsVector(KeyID, VectorValue);
-	}
+	SetValueAsVector(KeyID, VectorValue);
 }
 
 void UBlackboardComponent::SetValueAsVector(FBlackboard::FKey KeyID, const FVector& VectorValue)
 {
+	if (GetKeyType(KeyID) != UBlackboardKeyType_Vector::StaticClass())
+	{
+		return;
+	}
+
 	uint8* RawData = GetKeyRawData(KeyID);
 	if (RawData)
 	{
@@ -666,17 +702,44 @@ void UBlackboardComponent::SetValueAsVector(FBlackboard::FKey KeyID, const FVect
 	}
 }
 
-void UBlackboardComponent::ClearValueAsVector(const FName& KeyName)
+
+void UBlackboardComponent::SetValueAsRotator(const FName& KeyName, FRotator RotatorValue)
 {
-	const FBlackboard::FKey KeyID = GetKeyID(KeyName);
-	if (GetKeyType(KeyID) == UBlackboardKeyType_Vector::StaticClass())
+	const uint8 KeyID = GetKeyID(KeyName);
+	SetValueAsRotator(KeyID, RotatorValue);
+}
+
+void UBlackboardComponent::SetValueAsRotator(uint8 KeyID, const FRotator& RotatorValue)
+{
+	if (GetKeyType(KeyID) != UBlackboardKeyType_Rotator::StaticClass())
 	{
-		ClearValueAsVector(KeyID);
+		return;
+	}
+
+	uint8* RawData = GetKeyRawData(KeyID);
+	if (RawData)
+	{
+		const bool bChanged = UBlackboardKeyType_Rotator::SetValue(RawData, RotatorValue);
+		if (bChanged)
+		{
+			NotifyObservers(KeyID);
+		}
 	}
 }
 
+void UBlackboardComponent::ClearValueAsVector(const FName& KeyName)
+{
+	const FBlackboard::FKey KeyID = GetKeyID(KeyName);
+	ClearValueAsVector(KeyID);
+}
+
 void UBlackboardComponent::ClearValueAsVector(FBlackboard::FKey KeyID)
 {
+	if (GetKeyType(KeyID) != UBlackboardKeyType_Vector::StaticClass())
+	{
+		return;
+	}
+
 	uint8* RawData = GetKeyRawData(KeyID);
 	if (RawData)
 	{
@@ -688,6 +751,31 @@ void UBlackboardComponent::ClearValueAsVector(FBlackboard::FKey KeyID)
 	}
 }
 
+void UBlackboardComponent::ClearValueAsRotator(const FName& KeyName)
+{
+	const uint8 KeyID = GetKeyID(KeyName);
+	ClearValueAsRotator(KeyID);
+}
+
+void UBlackboardComponent::ClearValueAsRotator(uint8 KeyID)
+{
+	if (GetKeyType(KeyID) != UBlackboardKeyType_Rotator::StaticClass())
+	{
+		return;
+	}
+
+	uint8* RawData = GetKeyRawData(KeyID);
+	if (RawData)
+	{
+		const bool bChanged = UBlackboardKeyType_Rotator::SetValue(RawData, FAISystem::InvalidRotation);
+		if (bChanged)
+		{
+			NotifyObservers(KeyID);
+		}
+	}
+}
+
+
 bool UBlackboardComponent::GetLocationFromEntry(const FName& KeyName, FVector& ResultLocation) const
 {
 	const FBlackboard::FKey KeyID = GetKeyID(KeyName);
diff --git a/Engine/Source/Runtime/Core/Public/Math/Rotator.h b/Engine/Source/Runtime/Core/Public/Math/Rotator.h
--- a/Engine/Source/Runtime/Core/Public/Math/Rotator.h
+++ b/Engine/Source/Runtime/Core/Public/Math/Rotator.h
@@ -37,6 +37,13 @@ public:
 	 */
 	FORCEINLINE FRotator( ) { };
 
+	/**
+	 * Constructor
+	 *
+	 * @param InF Value to set all components to.
+	 */
+	explicit FORCEINLINE FRotator(float InF);
+
 	/**
 	 * Constructor.
 	 *
@@ -403,7 +410,12 @@ FORCEINLINE FRotator operator*( float Scale, const FRotator &R )
 #endif
 
 
-	FORCEINLINE FRotator::FRotator( float InPitch, float InYaw, float InRoll )
+FORCEINLINE FRotator::FRotator( float InF ) 
+	:	Pitch(InF), Yaw(InF), Roll(InF) 
+{}
+
+
+FORCEINLINE FRotator::FRotator( float InPitch, float InYaw, float InRoll )
 	:	Pitch(InPitch), Yaw(InYaw), Roll(InRoll) 
 {}
 
