Fixing PCH warnings in GameplayDebugger and AIModule

diff --git a/Engine/Source/Developer/GameplayDebugger/Private/GameplayDebuggerExec.cpp b/Engine/Source/Developer/GameplayDebugger/Private/GameplayDebuggerExec.cpp
--- a/Engine/Source/Developer/GameplayDebugger/Private/GameplayDebuggerExec.cpp
+++ b/Engine/Source/Developer/GameplayDebugger/Private/GameplayDebuggerExec.cpp
@@ -5,12 +5,6 @@
 #if !(UE_BUILD_SHIPPING || UE_BUILD_TEST)
 #include "Misc/CoreMisc.h"
 
-#if WITH_EDITOR
-#include "Editor.h"
-#include "EditorViewportClient.h"
-#include "LevelEditorViewport.h"
-#endif //WITH_EDITOR
-
 struct FGameplayDebuggerExec : public FSelfRegisteringExec
 {
 	FGameplayDebuggerExec()
diff --git a/Engine/Source/Developer/GameplayDebugger/Private/GameplayDebuggerPrivate.h b/Engine/Source/Developer/GameplayDebugger/Private/GameplayDebuggerPrivate.h
--- a/Engine/Source/Developer/GameplayDebugger/Private/GameplayDebuggerPrivate.h
+++ b/Engine/Source/Developer/GameplayDebugger/Private/GameplayDebuggerPrivate.h
@@ -36,4 +36,11 @@
 #include "GameplayDebuggingControllerComponent.h"
 #include "GameplayDebuggingComponent.h"
 
+#if WITH_EDITOR
+#include "Editor.h"
+#include "EditorViewportClient.h"
+#include "LevelEditorViewport.h"
+#include "UnrealEd.h"
+#endif //WITH_EDITOR
+
 DECLARE_LOG_CATEGORY_EXTERN(LogGDT, Warning, All);
diff --git a/Engine/Source/Developer/GameplayDebugger/Private/GameplayDebuggingControllerComponent.cpp b/Engine/Source/Developer/GameplayDebugger/Private/GameplayDebuggingControllerComponent.cpp
--- a/Engine/Source/Developer/GameplayDebugger/Private/GameplayDebuggingControllerComponent.cpp
+++ b/Engine/Source/Developer/GameplayDebugger/Private/GameplayDebuggingControllerComponent.cpp
@@ -7,10 +7,6 @@
 #include "TimerManager.h"
 #include "VisualLog.h"
 
-#if WITH_EDITOR
-#include "UnrealEd.h"
-#endif
-
 DEFINE_LOG_CATEGORY_STATIC(LogGameplayDebugging, Log, All);
 
 #define BUGIT_VIEWS (1<<EAIDebugDrawDataView::Basic) | (1 << EAIDebugDrawDataView::OverHead)
diff --git a/Engine/Source/Runtime/AIModule/Private/AIModulePrivate.h b/Engine/Source/Runtime/AIModule/Private/AIModulePrivate.h
--- a/Engine/Source/Runtime/AIModule/Private/AIModulePrivate.h
+++ b/Engine/Source/Runtime/AIModule/Private/AIModulePrivate.h
@@ -45,6 +45,10 @@
 #include "BehaviorTree/Decorators/BTDecorator_BlackboardBase.h"
 #include "BehaviorTree/Services/BTService_BlackboardBase.h"
 
+#if WITH_EDITOR
+#include "UnrealEd.h"
+#endif
+
 DECLARE_LOG_CATEGORY_EXTERN(LogAINavigation, Warning, All);
 DECLARE_LOG_CATEGORY_EXTERN(LogBehaviorTree, Display, All);
 DECLARE_LOG_CATEGORY_EXTERN(LogCrowdFollowing, Warning, All);
\ No newline at end of file
diff --git a/Engine/Source/Runtime/AIModule/Private/EnvironmentQuery/EQSTestingPawn.cpp b/Engine/Source/Runtime/AIModule/Private/EnvironmentQuery/EQSTestingPawn.cpp
--- a/Engine/Source/Runtime/AIModule/Private/EnvironmentQuery/EQSTestingPawn.cpp
+++ b/Engine/Source/Runtime/AIModule/Private/EnvironmentQuery/EQSTestingPawn.cpp
@@ -11,10 +11,6 @@
 #include "Engine/Selection.h"
 #include "Components/BillboardComponent.h"
 
-#if WITH_EDITOR
-#include "UnrealEd.h"
-#endif
-
 //----------------------------------------------------------------------//
 // AEQSTestingPawn
 //----------------------------------------------------------------------//
