Made AIModule not include UnrealEd.h at all #UE4

diff --git a/Engine/Source/Runtime/AIModule/Private/AIModulePrivate.h b/Engine/Source/Runtime/AIModule/Private/AIModulePrivate.h
--- a/Engine/Source/Runtime/AIModule/Private/AIModulePrivate.h
+++ b/Engine/Source/Runtime/AIModule/Private/AIModulePrivate.h
@@ -45,10 +45,6 @@
 #include "BehaviorTree/Decorators/BTDecorator_BlackboardBase.h"
 #include "BehaviorTree/Services/BTService_BlackboardBase.h"
 
-#if WITH_EDITOR
-#include "UnrealEd.h"
-#endif
-
 DECLARE_LOG_CATEGORY_EXTERN(LogAINavigation, Warning, All);
 DECLARE_LOG_CATEGORY_EXTERN(LogBehaviorTree, Display, All);
 DECLARE_LOG_CATEGORY_EXTERN(LogCrowdFollowing, Warning, All);
\ No newline at end of file
diff --git a/Engine/Source/Runtime/AIModule/Private/EnvironmentQuery/EQSTestingPawn.cpp b/Engine/Source/Runtime/AIModule/Private/EnvironmentQuery/EQSTestingPawn.cpp
--- a/Engine/Source/Runtime/AIModule/Private/EnvironmentQuery/EQSTestingPawn.cpp
+++ b/Engine/Source/Runtime/AIModule/Private/EnvironmentQuery/EQSTestingPawn.cpp
@@ -8,6 +8,16 @@
 #include "EnvironmentQuery/EnvQueryManager.h"
 #include "EnvironmentQuery/EQSRenderingComponent.h"
 
+#if WITH_EDITORONLY_DATA
+#include "Components/ArrowComponent.h"
+#endif // WITH_EDITORONLY_DATA
+
+#if WITH_EDITOR
+#include "Engine/Brush.h"
+#include "Editor/EditorEngine.h"
+extern UNREALED_API class UEditorEngine* GEditor;
+#endif // WITH_EDITOR
+
 #include "Engine/Selection.h"
 #include "Components/BillboardComponent.h"
 
diff --git a/Engine/Source/Runtime/AIModule/Private/EnvironmentQuery/EnvQueryManager.cpp b/Engine/Source/Runtime/AIModule/Private/EnvironmentQuery/EnvQueryManager.cpp
--- a/Engine/Source/Runtime/AIModule/Private/EnvironmentQuery/EnvQueryManager.cpp
+++ b/Engine/Source/Runtime/AIModule/Private/EnvironmentQuery/EnvQueryManager.cpp
@@ -9,8 +9,12 @@
 #include "EnvironmentQuery/EQSTestingPawn.h"
 
 #if WITH_EDITOR
-#include "UnrealEd.h"
-#endif
+#include "Engine/Brush.h"
+#include "Editor/EditorEngine.h"
+#include "EngineUtils.h"
+
+extern UNREALED_API class UEditorEngine* GEditor;
+#endif // WITH_EDITOR
 
 DEFINE_LOG_CATEGORY(LogEQS);
 
