Temporary fix for a fall off bug after EQS refactor #UE4

TTP#339295, Errors logged while loading QAGame's EQS assets were breaking the cook. I've applied a temporary fix, but it requires a proper handling.
The easiest repro is to run QAGame editor, look for EQS query assets and load them all.

diff --git a/Engine/Source/Runtime/AIModule/Private/EnvironmentQuery/EnvQueryTest.cpp b/Engine/Source/Runtime/AIModule/Private/EnvironmentQuery/EnvQueryTest.cpp
--- a/Engine/Source/Runtime/AIModule/Private/EnvironmentQuery/EnvQueryTest.cpp
+++ b/Engine/Source/Runtime/AIModule/Private/EnvironmentQuery/EnvQueryTest.cpp
@@ -252,12 +252,10 @@ void UEnvQueryTest::PostLoad()
 		WeightModifier = EEnvTestWeight::Skip;
 	}
 
-	if (WeightModifier != EEnvTestWeight::Skip)
+	if (WeightModifier != EEnvTestWeight::Skip && !bWorkOnFloatValues)
 	{
-		if (!bWorkOnFloatValues)
-		{	// Working on booleans, so MUST use constant value!
-			WeightModifier = EEnvTestWeight::Constant;
-		}
+		// Working on booleans, so MUST use constant value!
+		WeightModifier = EEnvTestWeight::Constant;
 	}
 
 	switch (WeightModifier)
@@ -279,19 +277,21 @@ void UEnvQueryTest::PostLoad()
 			break;
 
 		case EEnvTestWeight::Absolute:
-			UE_LOG(LogEQS, Error, TEXT("Absolute weight is no longer supported!  Sticking with default weighting."));
+			UE_LOG(LogEQS, Warning, TEXT("Absolute weight is no longer supported!  Sticking with default weighting."));
+			// falling back to Linear
+			ScoringEquation = EEnvTestScoreEquation::Linear;
 			break;
 
 		case EEnvTestWeight::Skip:
 			TestPurpose = EEnvTestPurpose::Filter;
 			if (Condition == EEnvTestCondition::NoCondition)
 			{
-				UE_LOG(LogEQS, Error, TEXT("Test was set to neither filter nor score!  Now setting to filter!"));	
+				UE_LOG(LogEQS, Warning, TEXT("Test was set to neither filter nor score!  Now setting to filter!"));	
 			}
 			break;
 
 		default:
-			UE_LOG(LogEQS, Error, TEXT("Invalid Weight Modifier type in UEnvQueryTest::PostLoad!"));
+			UE_LOG(LogEQS, Warning, TEXT("Invalid Weight Modifier type in UEnvQueryTest::PostLoad!"));
 			break;
 	}
 
