Removed ugly commented out header inclusion from AIModule.cpp #UE4

diff --git a/Engine/Source/Runtime/AIModule/Private/AIModule.cpp b/Engine/Source/Runtime/AIModule/Private/AIModule.cpp
--- a/Engine/Source/Runtime/AIModule/Private/AIModule.cpp
+++ b/Engine/Source/Runtime/AIModule/Private/AIModule.cpp
@@ -2,7 +2,6 @@
 
 #include "AIModulePrivate.h"
 #include "AISystem.h"
-//#include "Developer/GameplayDebugger/Public/GameplayDebugger.h"
 
 DEFINE_LOG_CATEGORY_STATIC(LogAIModule, Log, All);
 
