fixed crash in simple parallel node when it doesn't have valid main task

diff --git a/Engine/Source/Runtime/AIModule/Private/BehaviorTree/Composites/BTComposite_SimpleParallel.cpp b/Engine/Source/Runtime/AIModule/Private/BehaviorTree/Composites/BTComposite_SimpleParallel.cpp
--- a/Engine/Source/Runtime/AIModule/Private/BehaviorTree/Composites/BTComposite_SimpleParallel.cpp
+++ b/Engine/Source/Runtime/AIModule/Private/BehaviorTree/Composites/BTComposite_SimpleParallel.cpp
@@ -93,7 +93,10 @@ void UBTComposite_SimpleParallel::NotifyNodeDeactivation(struct FBehaviorTreeSea
 	}
 
 	// remove main task
-	SearchData.AddUniqueUpdate(FBehaviorTreeSearchUpdate(Children[EBTParallelChild::MainTask].ChildTask, ActiveInstanceIdx, EBTNodeUpdateMode::Remove));
+	if (Children.IsValidIndex(EBTParallelChild::MainTask))
+	{
+		SearchData.AddUniqueUpdate(FBehaviorTreeSearchUpdate(Children[EBTParallelChild::MainTask].ChildTask, ActiveInstanceIdx, EBTNodeUpdateMode::Remove));
+	}
 
 	// remove all active nodes from background tree
 	const FBTNodeIndex FirstBackgroundIndex(ActiveInstanceIdx, GetChildExecutionIndex(EBTParallelChild::BackgroundTree, EBTChildIndex::FirstNode));
