fixed crash in behavior tree overwritting itself as a result executing/aborting task

diff --git a/Engine/Source/Runtime/AIModule/Private/BehaviorTree/BehaviorTreeComponent.cpp b/Engine/Source/Runtime/AIModule/Private/BehaviorTree/BehaviorTreeComponent.cpp
--- a/Engine/Source/Runtime/AIModule/Private/BehaviorTree/BehaviorTreeComponent.cpp
+++ b/Engine/Source/Runtime/AIModule/Private/BehaviorTree/BehaviorTreeComponent.cpp
@@ -1033,10 +1033,14 @@ void UBehaviorTreeComponent::ExecuteTask(UBTTaskNode* TaskNode)
 	uint8* NodeMemory = (uint8*)(TaskNode->GetNodeMemory<uint8>(ActiveInstance));
 	EBTNodeResult::Type TaskResult = TaskNode->WrappedExecuteTask(this, NodeMemory);
 
-	// update task's runtime values after it had a chance to initialize memory
-	UpdateDebuggerAfterExecution(TaskNode, InstanceIdx);
+	const UBTNode* ActiveNodeAfterExecution = GetActiveNode();
+	if (ActiveNodeAfterExecution == TaskNode)
+	{
+		// update task's runtime values after it had a chance to initialize memory
+		UpdateDebuggerAfterExecution(TaskNode, InstanceIdx);
 
-	OnTaskFinished(TaskNode, TaskResult);
+		OnTaskFinished(TaskNode, TaskResult);
+	}
 }
 
 void UBehaviorTreeComponent::AbortCurrentTask()
@@ -1053,9 +1057,12 @@ void UBehaviorTreeComponent::AbortCurrentTask()
 	uint8* NodeMemory = (uint8*)(ActiveTask->GetNodeMemory<uint8>(ActiveInstance));
 	EBTNodeResult::Type TaskResult = ActiveTask->WrappedAbortTask(this, NodeMemory);
 
-	ActiveInstance.ActiveNodeType = EBTActiveNode::AbortingTask;
-
-	OnTaskFinished(ActiveTask, TaskResult);
+	const UBTNode* ActiveNodeAfterAbort = GetActiveNode();
+	if (ActiveNodeAfterAbort == ActiveTask)
+	{
+		ActiveInstance.ActiveNodeType = EBTActiveNode::AbortingTask;
+		OnTaskFinished(ActiveTask, TaskResult);
+	}
 }
 
 void UBehaviorTreeComponent::RegisterMessageObserver(const class UBTTaskNode* TaskNode, FName MessageType)
