fixed crash on executing parallel composite BT node without background subtree

diff --git a/Engine/Source/Runtime/AIModule/Private/BehaviorTree/BTCompositeNode.cpp b/Engine/Source/Runtime/AIModule/Private/BehaviorTree/BTCompositeNode.cpp
--- a/Engine/Source/Runtime/AIModule/Private/BehaviorTree/BTCompositeNode.cpp
+++ b/Engine/Source/Runtime/AIModule/Private/BehaviorTree/BTCompositeNode.cpp
@@ -26,7 +26,7 @@ int32 UBTCompositeNode::FindChildToExecute(struct FBehaviorTreeSearchData& Searc
 	if (Children.Num())
 	{
 		for (int32 ChildIdx = GetNextChild(SearchData, NodeMemory->CurrentChild, LastResult);
-			ChildIdx >= 0;
+			Children.IsValidIndex(ChildIdx);
 			ChildIdx = GetNextChild(SearchData, ChildIdx, LastResult))
 		{
 			// check decorators
