PawnActions moving to AIModule merged over from dev branch #UE4

diff --git a/Engine/Source/Runtime/AIModule/Classes/Actions/PawnAction.h b/Engine/Source/Runtime/AIModule/Classes/Actions/PawnAction.h
new file mode 100644
--- /dev/null
+++ b/Engine/Source/Runtime/AIModule/Classes/Actions/PawnAction.h
@@ -0,0 +1,245 @@
+// Copyright 1998-2014 Epic Games, Inc. All Rights Reserved.
+
+#pragma once
+
+#include "BrainComponent.h"
+#include "PawnAction.generated.h"
+
+class UPawnAction;
+class UPawnActionsComponent;
+
+DECLARE_DELEGATE_TwoParams(FPawnActionEventDelegate, UPawnAction*, EPawnActionEventType::Type);
+
+UENUM()
+namespace EPawnActionFailHandling
+{
+	enum Type
+	{
+		RequireSuccess,
+		IgnoreFailure
+	};
+}
+
+/** 
+ *	Things to remember:
+ *	* Actions are created paused
+ */
+UCLASS(abstract)
+class AIMODULE_API UPawnAction : public UObject
+{
+	GENERATED_UCLASS_BODY()
+
+	friend class UPawnActionsComponent;
+	friend struct FPawnActionStack;
+
+private:
+	/** Current child node executing on top of this Action */
+	UPROPERTY(Transient)
+	UPawnAction* ChildAction;
+
+	UPROPERTY(Transient)
+	UPawnAction* ParentAction;
+
+	/** Extra reference to the component this action is being governed by */
+	UPROPERTY(Transient)
+	UPawnActionsComponent* OwnerComponent;
+	
+	/** indicates an object that caused this action. Used for mass removal of actions 
+	 *	by specific object */
+	UPROPERTY(Transient)
+	UObject* Instigator;
+
+protected:
+	/** @Note: THIS IS HERE _ONLY_ BECAUSE OF THE WAY AI MESSAGING IS CURRENTLY IMPLEMENTED. WILL GO AWAY! */
+	UPROPERTY(Transient)
+	class UBrainComponent* BrainComp;
+
+private:
+	/** stores registered message observers */
+	TArray<FAIMessageObserverHandle> MessageHandlers;
+
+	EAIRequestPriority::Type ExecutionPriority;
+
+	FPawnActionEventDelegate ActionObserver;
+
+protected:
+
+	FAIRequestID RequestID;
+
+	/** if this is FALSE and we're trying to push a new instance of a given class,
+	 *	but the top of the stack is already an instance of that class ignore the attempted push */
+	UPROPERTY(Category = PawnAction, EditAnywhere, BlueprintReadOnly)
+	uint32 bAllowNewSameClassInstance : 1;
+
+	/** if this is TRUE, when we try to push a new instance of an action who has the
+	 *	same class as the action on the top of the stack, pop the one on the stack, and push the new one
+	 *	NOTE: This trumps bAllowNewClassInstance (e.g. if this is true and bAllowNewClassInstance
+	 *	is false the active instance will still be replaced) */
+	UPROPERTY(Category = AIAction, EditAnywhere, BlueprintReadWrite)
+	uint32 bReplaceActiveSameClassInstance : 1;
+
+private:
+	/** indicates whether action is in the process of abortion, and if so on what state */
+	EPawnActionAbortState::Type AbortState;
+
+	EPawnActionResult::Type FinishResult;
+
+	/** Used exclusively for action events sorting */
+	int32 IndexOnStack;
+	
+	/** Indicates the action has been paused */
+	uint32 bPaused : 1;
+
+	uint32 bHasBeenStarted : 1;
+
+	/** set to true when action fails the initial Start call */
+	uint32 bFailedToStart : 1;
+
+protected:
+	/** TickAction will get called only if this flag is set. To be set in derived action's
+	 *	constructor. 
+	 *	@NOTE Toggling at runtime is not supported */
+	uint32 bWantsTick : 1;
+
+public:
+
+	// Begin UObject
+	virtual UWorld* GetWorld() const override;
+	// End UObject
+
+	FORCEINLINE const UPawnAction* GetParentAction() const { return ParentAction; }
+	FORCEINLINE const UPawnAction* GetChildAction() const { return ChildAction; }
+	FORCEINLINE UPawnAction* GetChildAction() { return ChildAction; }
+	FORCEINLINE bool IsPaused() const { return !!bPaused; }
+	FORCEINLINE bool IsActive() const { return FinishResult == EPawnActionResult::InProgress && IsPaused() == false && AbortState == EPawnActionAbortState::NotBeingAborted; }
+	FORCEINLINE bool IsBeingAborted() const { return AbortState != EPawnActionAbortState::NotBeingAborted; }
+	FORCEINLINE bool WantsTick() const { return bWantsTick; }
+
+	FORCEINLINE void TickAction(float DeltaTime)
+	{ 
+		// tick ChildAction 
+		if (ChildAction != NULL)
+		{
+			ChildAction->Tick(DeltaTime);
+		}
+		// or self if not paused
+		else if (!!bWantsTick && IsPaused() == false)
+		{
+			Tick(DeltaTime);
+		}
+	}
+
+	///void HandleChildActionFinished(UPawnAction* Action);
+
+	/** triggers aborting of an Action
+	 *	@param bForce
+	 *	@return current state of task abort
+	 *	@NOTE do not make this virtual! Contains some essential logic. */
+	EPawnActionAbortState::Type Abort(EAIForceParam::Type ShouldForce = EAIForceParam::DoNotForce);
+	
+	/** called when */
+	//virtual void OnActionFinished();
+
+	FORCEINLINE EAIRequestPriority::Type GetPriority() const { return ExecutionPriority; }
+	FORCEINLINE EPawnActionResult::Type GetResult() const { return FinishResult; }
+	FORCEINLINE EPawnActionAbortState::Type GetAbortState() const { return AbortState; }
+	FORCEINLINE class UPawnActionsComponent* GetOwnerComponent() const { return OwnerComponent; }
+	FORCEINLINE UObject* GetInstigator() const { return Instigator; }
+	class APawn* GetPawn();
+	class AController* GetController();
+
+	template<class TActionClass>
+	static TActionClass* CreateActionInstance(class UWorld* World)
+	{
+		check(World);
+		TSubclassOf<class UPawnAction> ActionClass = TActionClass::StaticClass();
+		return ConstructObject<TActionClass>(ActionClass, World);
+	}
+
+	//----------------------------------------------------------------------//
+	// messaging
+	//----------------------------------------------------------------------//
+	void WaitForMessage(FName MessageType, FAIRequestID RequestID = FAIRequestID::AnyRequest);
+	virtual void HandleAIMessage(UBrainComponent*, const struct FAIMessage&){};
+
+	void SetActionObserver(const FPawnActionEventDelegate& ActionObserver) { this->ActionObserver = ActionObserver; }
+
+	//----------------------------------------------------------------------//
+	// Blueprint interface
+	//----------------------------------------------------------------------//
+	UFUNCTION(BlueprintPure, Category = "AI|PawnActions")
+	TEnumAsByte<EAIRequestPriority::Type> GetActionPriority();
+
+	UFUNCTION(BlueprintCallable, Category = "AI|PawnActions", meta = (HidePin = "WorldContextObject", DefaultToSelf = "WorldContextObject"))
+	static class UPawnAction* CreateActionInstance(class UObject* WorldContextObject, TSubclassOf<class UPawnAction> ActionClass);
+
+	//----------------------------------------------------------------------//
+	// debug
+	//----------------------------------------------------------------------//
+	FString GetStateDescription() const;
+	FString GetPriorityName() const;
+
+protected:
+
+	/** starts or resumes action, depending on internal state */
+	bool Activate();
+	void OnPopped();
+
+	UFUNCTION(BlueprintCallable, Category = "AI|PawnActions")
+	void Finish(TEnumAsByte<EPawnActionResult::Type> WithResult);
+
+	void SendEvent(EPawnActionEventType::Type Event);
+
+	void StopWaitingForMessages();
+
+	void SetOwnerComponent(class UPawnActionsComponent* Component);
+
+	void SetInstigator(UObject* const InInstigator);
+
+	virtual void Tick(float DeltaTime);
+
+	/** called to start off the Action
+	 *	@return 'true' if actions successfully started. 
+	 *	@NOTE if action fails to start no finishing or aborting mechanics will be triggered */
+	virtual bool Start();
+	/** called to pause action when higher priority or child action kicks in */
+	virtual bool Pause();
+	/** called to resume action after being paused */
+	virtual bool Resume();
+	/** called when this action is being removed from action stacks */
+	virtual void OnFinished(EPawnActionResult::Type WithResult);
+	/** called to give Action chance to react to child action finishing.
+	 *	@NOTE gets called _AFTER_ child's OnFinished to give child action chance 
+	 *		to prepare "finishing data" for parent to read. 
+	 *	@NOTE clears parent-child binding */
+	virtual void OnChildFinished(UPawnAction* Action, EPawnActionResult::Type WithResult);
+
+	/** apart from doing regular push request copies additional values from Parent, like Priority and Instigator */
+	bool PushChildAction(UPawnAction* Action);
+	
+	/** performs actual work on aborting Action. Should be called exclusively by Abort function
+	 *	@return only valid return values here are LatendAbortInProgress and AbortDone */
+	virtual EPawnActionAbortState::Type PerformAbort(EAIForceParam::Type ShouldForce) { return EPawnActionAbortState::AbortDone; }
+
+	bool HasBeenStarted() const { return bHasBeenStarted;  }
+
+private:
+	/** called when this action is put on a stack. Does not indicate action will be started soon
+	 *	(it depends on other actions on other action stacks. Called before Start() call */
+	void OnPushed();
+
+	/** Sets final result for this Action. To be called only once upon Action's finish */
+	void SetFinishResult(EPawnActionResult::Type Result);
+
+	// do not un-private. Internal logic only!
+	void SetAbortState(EPawnActionAbortState::Type NewAbortState);
+};
+
+
+//----------------------------------------------------------------------//
+// Blueprint inlines
+//----------------------------------------------------------------------//
+FORCEINLINE TEnumAsByte<EAIRequestPriority::Type> UPawnAction::GetActionPriority()
+{
+	return ExecutionPriority;
+}
diff --git a/Engine/Source/Runtime/AIModule/Classes/Actions/PawnAction_BlueprintBase.h b/Engine/Source/Runtime/AIModule/Classes/Actions/PawnAction_BlueprintBase.h
new file mode 100644
--- /dev/null
+++ b/Engine/Source/Runtime/AIModule/Classes/Actions/PawnAction_BlueprintBase.h
@@ -0,0 +1,35 @@
+// Copyright 1998-2014 Epic Games, Inc. All Rights Reserved.
+
+#pragma once
+
+#include "PawnAction_BlueprintBase.generated.h"
+
+UCLASS(abstract, Blueprintable)
+class AIMODULE_API UPawnAction_BlueprintBase : public UPawnAction
+{
+	GENERATED_UCLASS_BODY()
+
+public:
+
+	//----------------------------------------------------------------------//
+	// Blueprint interface
+	//----------------------------------------------------------------------//
+
+	UFUNCTION(BlueprintImplementableEvent, Category = "AI|PawnActions")
+	void ActionStart(APawn* ControlledPawn);
+	UFUNCTION(BlueprintImplementableEvent, Category = "AI|PawnActions")
+	void ActionTick(APawn* ControlledPawn, float DeltaSeconds);
+	UFUNCTION(BlueprintImplementableEvent, Category = "AI|PawnActions")
+	void ActionPause(APawn* ControlledPawn);
+	UFUNCTION(BlueprintImplementableEvent, Category = "AI|PawnActions")
+	void ActionResume(APawn* ControlledPawn);
+	UFUNCTION(BlueprintImplementableEvent, Category = "AI|PawnActions")
+	void ActionFinished(APawn* ControlledPawn, EPawnActionResult::Type WithResult);
+
+protected:
+	virtual void Tick(float DeltaTime) override;
+	virtual bool Start();
+	virtual bool Pause() override;
+	virtual bool Resume() override;
+	virtual void OnFinished(EPawnActionResult::Type WithResult) override;
+};
diff --git a/Engine/Source/Runtime/AIModule/Classes/Actions/PawnAction_Move.h b/Engine/Source/Runtime/AIModule/Classes/Actions/PawnAction_Move.h
new file mode 100644
--- /dev/null
+++ b/Engine/Source/Runtime/AIModule/Classes/Actions/PawnAction_Move.h
@@ -0,0 +1,76 @@
+// Copyright 1998-2014 Epic Games, Inc. All Rights Reserved.
+
+#pragma once
+
+#include "PawnAction_Move.generated.h"
+
+namespace EPawnActionMoveMode
+{
+	enum Type
+	{
+		UsePathfinding,
+		StraightLine,
+	};
+}
+
+UCLASS()
+class AIMODULE_API UPawnAction_Move : public UPawnAction
+{
+	GENERATED_UCLASS_BODY()
+
+	UPROPERTY(Category = PawnAction, EditAnywhere, BlueprintReadWrite)
+	class AActor* GoalActor;
+
+	UPROPERTY(Category = PawnAction, EditAnywhere, BlueprintReadWrite)
+	FVector GoalLocation;
+
+	UPROPERTY(Category = PawnAction, EditAnywhere, meta = (ClampMin = "0.01"), BlueprintReadWrite)
+	float AcceptableRadius;
+
+	/** "None" will result in default filter being used */
+	UPROPERTY(Category = PawnAction, EditAnywhere, BlueprintReadWrite)
+	TSubclassOf<class UNavigationQueryFilter> FilterClass;
+
+	UPROPERTY(Category = PawnAction, EditAnywhere, BlueprintReadWrite)
+	uint32 bAllowStrafe : 1;
+
+	/** if set, movement will use path finding */
+	uint32 bUsePathfinding : 1;
+
+	/** if set, GoalLocation will be projected on navigation before using  */
+	uint32 bProjectGoalToNavigation : 1;
+
+	/** if set, path to GoalActor will be updated with goal's movement */
+	uint32 bUpdatePathToGoal : 1;
+
+	/** if set, other actions with the same priority will be aborted when path is changed */
+	uint32 bAbortChildActionOnPathChange : 1;
+
+	static UPawnAction_Move* CreateAction(UWorld* World, class AActor* GoalActor, EPawnActionMoveMode::Type Mode);
+	static UPawnAction_Move* CreateAction(UWorld* World, const FVector& GoalLocation, EPawnActionMoveMode::Type Mode);
+
+	static bool CheckAlreadyAtGoal(class AAIController* Controller, const FVector& TestLocation, float Radius);
+	static bool CheckAlreadyAtGoal(class AAIController* Controller, const AActor* TestGoal, float Radius);
+
+	virtual void HandleAIMessage(UBrainComponent*, const struct FAIMessage&) override;
+
+	void SetPath(FNavPathSharedPtr InPath);
+	void OnPathUpdated(FNavigationPath* UpdatedPath);
+
+protected:
+	/** currently followed path */
+	FNavPathSharedPtr Path;
+
+	/** original path observer */
+	FNavigationPath::FPathObserverDelegate PathObserver;
+
+	virtual bool Start() override;
+	virtual bool Pause() override;
+	virtual bool Resume() override;
+	virtual EPawnActionAbortState::Type PerformAbort(EAIForceParam::Type ShouldForce) override;
+
+	virtual EPathFollowingRequestResult::Type RequestMove(AAIController* Controller);
+	
+	bool PerformMoveAction();
+	void DeferredPerformMoveAction();
+};
diff --git a/Engine/Source/Runtime/AIModule/Classes/Actions/PawnAction_Repeat.h b/Engine/Source/Runtime/AIModule/Classes/Actions/PawnAction_Repeat.h
new file mode 100644
--- /dev/null
+++ b/Engine/Source/Runtime/AIModule/Classes/Actions/PawnAction_Repeat.h
@@ -0,0 +1,38 @@
+// Copyright 1998-2014 Epic Games, Inc. All Rights Reserved.
+
+#pragma once
+
+#include "PawnAction_Repeat.generated.h"
+
+UCLASS()
+class AIMODULE_API UPawnAction_Repeat : public UPawnAction
+{
+	GENERATED_UCLASS_BODY()
+
+	enum
+	{
+		LoopForever = -1
+	};
+
+	/** Action to repeat. This instance won't really be run, it's a source for copying actions to be actually performed */
+	UPROPERTY()
+	UPawnAction* ActionToRepeat;
+
+	UPROPERTY(Transient)
+	UPawnAction* RecentActionCopy;
+
+	UPROPERTY(Category = PawnAction, EditAnywhere, BlueprintReadOnly)
+	TEnumAsByte<EPawnActionFailHandling::Type> ChildFailureHandlingMode;
+
+	int32 RepeatsLeft;
+
+	/** @param NumberOfRepeats number of times to repeat action. UPawnAction_Repeat::LoopForever loops forever */
+	static UPawnAction_Repeat* CreateAction(class UWorld* World, class UPawnAction* ActionToRepeat, int32 NumberOfRepeats);
+
+protected:
+	virtual bool Start() override;
+	virtual bool Resume() override;
+	virtual void OnChildFinished(UPawnAction* Action, EPawnActionResult::Type WithResult) override;
+
+	bool PushActionCopy();
+};
\ No newline at end of file
diff --git a/Engine/Source/Runtime/AIModule/Classes/Actions/PawnAction_Sequence.h b/Engine/Source/Runtime/AIModule/Classes/Actions/PawnAction_Sequence.h
new file mode 100644
--- /dev/null
+++ b/Engine/Source/Runtime/AIModule/Classes/Actions/PawnAction_Sequence.h
@@ -0,0 +1,31 @@
+// Copyright 1998-2014 Epic Games, Inc. All Rights Reserved.
+
+#pragma once
+
+#include "PawnAction_Sequence.generated.h"
+
+UCLASS()
+class AIMODULE_API UPawnAction_Sequence : public UPawnAction
+{
+	GENERATED_UCLASS_BODY()
+
+	UPROPERTY()
+	TArray<UPawnAction*> ActionSequence;
+
+	UPROPERTY(Category = PawnAction, EditAnywhere, BlueprintReadOnly)
+	TEnumAsByte<EPawnActionFailHandling::Type> ChildFailureHandlingMode;
+
+	UPROPERTY(Transient)
+	UPawnAction* RecentActionCopy;
+
+	uint32 CurrentActionIndex;
+
+	static UPawnAction_Sequence* CreateAction(UWorld* World, TArray<UPawnAction*>& ActionSequence);
+
+protected:
+	virtual bool Start() override;
+	virtual bool Resume() override;
+	virtual void OnChildFinished(UPawnAction* Action, EPawnActionResult::Type WithResult) override;
+
+	bool PushNextActionCopy();
+};
\ No newline at end of file
diff --git a/Engine/Source/Runtime/AIModule/Classes/Actions/PawnActionsComponent.h b/Engine/Source/Runtime/AIModule/Classes/Actions/PawnActionsComponent.h
new file mode 100644
--- /dev/null
+++ b/Engine/Source/Runtime/AIModule/Classes/Actions/PawnActionsComponent.h
@@ -0,0 +1,130 @@
+// Copyright 1998-2014 Epic Games, Inc. All Rights Reserved.
+
+#pragma once
+#include "Components/ActorComponent.h"
+#include "PawnActionsComponent.generated.h"
+
+DECLARE_LOG_CATEGORY_EXTERN(LogPawnAction, Warning, All);
+
+USTRUCT()
+struct FPawnActionEvent
+{
+	GENERATED_USTRUCT_BODY()
+
+	UPROPERTY()
+	class UPawnAction* Action;
+
+	EPawnActionEventType::Type EventType;
+
+	EAIRequestPriority::Type Priority;
+
+	// used to maintain order of equally-important messages
+	uint32 Index;
+
+	FPawnActionEvent() : Action(NULL), EventType(EPawnActionEventType::Invalid), Priority(EAIRequestPriority::MAX), Index(uint32(-1))
+	{}
+
+	FPawnActionEvent(class UPawnAction* Action, EPawnActionEventType::Type EventType, uint32 Index);
+};
+
+USTRUCT()
+struct FPawnActionStack
+{
+	GENERATED_USTRUCT_BODY()
+
+private:
+	UPROPERTY()
+	class UPawnAction* TopAction;
+public:
+	void Pause();
+	void Resume();
+
+	/** All it does is tie actions into a double-linked list making NewTopAction
+	 *	new stack's top */
+	void PushAction(class UPawnAction* NewTopAction);
+
+	/** Looks through the double-linked action list looking for specified action
+	 *	and if found action will be popped along with all it's siblings */
+	void PopAction(class UPawnAction* ActionToPop);
+	
+	FORCEINLINE class UPawnAction* GetTop() { return TopAction; }
+	FORCEINLINE const class UPawnAction* GetTop() const { return TopAction; }
+	FORCEINLINE bool IsEmpty() const { return TopAction == NULL; }
+};
+
+UCLASS()
+class AIMODULE_API UPawnActionsComponent : public UActorComponent
+{
+	GENERATED_UCLASS_BODY()
+
+protected:
+	UPROPERTY(BlueprintReadOnly, Category="PawnActions")
+	APawn* ControlledPawn;
+
+	UPROPERTY()
+	TArray<FPawnActionStack> ActionStacks;
+
+	UPROPERTY()
+	TArray<FPawnActionEvent> ActionEvents;
+
+	UPROPERTY(Transient)
+	class UPawnAction* CurrentAction;
+
+	/** set when logic was locked by hi priority stack */
+	uint32 bLockedAILogic : 1;
+
+private:
+	uint32 ActionEventIndex;
+
+public:
+	//----------------------------------------------------------------------//
+	// blueprint interface
+	//----------------------------------------------------------------------//
+
+	UFUNCTION(BlueprintCallable, Category = "AI|PawnActions")
+	static bool PerformAction(class APawn* Pawn, class UPawnAction* Action, TEnumAsByte<EAIRequestPriority::Type> Priority = EAIRequestPriority::HardScript);
+
+	//----------------------------------------------------------------------//
+	// 
+	//----------------------------------------------------------------------//
+	/** Use it to save component work to figure out what it's controlling
+	 *	or if component can't/won't be able to figure it out properly
+	 *	@NOTE will throw a log warning if trying to set ControlledPawn if it's already set */
+	void SetControlledPawn(APawn* NewPawn);
+	FORCEINLINE APawn* GetControlledPawn() { return ControlledPawn; }
+	FORCEINLINE const APawn* GetControlledPawn() const { return ControlledPawn; }
+	FORCEINLINE AController* GetController() { return ControlledPawn ? ControlledPawn->GetController() : NULL; }
+	FORCEINLINE UPawnAction* GetCurrentAction() { return CurrentAction; }
+
+	bool OnEvent(class UPawnAction* Action, EPawnActionEventType::Type Event);
+
+	UFUNCTION(BlueprintCallable, Category = PawnAction)
+	bool PushAction(class UPawnAction* NewAction, EAIRequestPriority::Type Priority, UObject* Instigator = NULL);
+
+	/** Aborts given action instance */
+	UFUNCTION(BlueprintCallable, Category = PawnAction)
+	bool AbortAction(class UPawnAction* ActionToAbort);
+
+	/** removes all actions instigated with Priority by Instigator
+	 *	@param Priority if equal to EAIRequestPriority::MAX then all priority queues will be searched. 
+	 *		This is less efficient so use with caution 
+	 *	@return number of action abortions requested (performed asyncronously) */
+	uint32 AbortActionsInstigatedBy(UObject* const Instigator, EAIRequestPriority::Type Priority);
+	
+	void TickComponent(float DeltaTime, enum ELevelTick TickType, FActorComponentTickFunction *ThisTickFunction) override;
+
+	FORCEINLINE class UPawnAction* GetActiveAction(EAIRequestPriority::Type Priority) { return ActionStacks[Priority].GetTop();  }
+
+#if ENABLE_VISUAL_LOG
+	void DescribeSelfToVisLog(struct FVisLogEntry* Snapshot) const;
+#endif // ENABLE_VISUAL_LOG
+
+protected:
+	/** Finds the action that should be running. If it's different from CurrentAction
+	 *	then CurrentAction gets paused and newly selected action gets started up */
+	void UpdateCurrentAction();
+
+	APawn* CacheControlledPawn();
+
+	void UpdateAILogicLock();
+};
\ No newline at end of file
diff --git a/Engine/Source/Runtime/AIModule/Private/Actions/PawnAction.cpp b/Engine/Source/Runtime/AIModule/Private/Actions/PawnAction.cpp
new file mode 100644
--- /dev/null
+++ b/Engine/Source/Runtime/AIModule/Private/Actions/PawnAction.cpp
@@ -0,0 +1,403 @@
+// Copyright 1998-2014 Epic Games, Inc. All Rights Reserved.
+
+#include "AIModulePrivate.h"
+#include "Actions/PawnAction.h"
+#include "Actions/PawnActionsComponent.h"
+
+namespace
+{
+	FString GetActionResultName(int32 Index)
+	{
+		static const UEnum* Enum = FindObject<UEnum>(ANY_PACKAGE, TEXT("EPawnActionResult"));
+		check(Enum);
+		return Enum->GetEnumName(Index);
+	}
+}
+
+UPawnAction::UPawnAction(const class FPostConstructInitializeProperties& PCIP)
+	: Super(PCIP)
+{
+	// actions start their lives paused
+	bPaused = true;
+	bHasBeenStarted = false;
+	bFailedToStart = false;
+	IndexOnStack = INDEX_NONE;
+
+	FinishResult = EPawnActionResult::InProgress;
+}
+
+UWorld* UPawnAction::GetWorld() const
+{
+	return OwnerComponent ? OwnerComponent->GetWorld() : Cast<UWorld>(GetOuter());
+}
+
+void UPawnAction::Tick(float DeltaTime)
+{
+}
+
+EPawnActionAbortState::Type UPawnAction::Abort(EAIForceParam::Type ShouldForce)
+{
+	// if already aborting, and this request is not Forced, just skip it
+	if (AbortState != EPawnActionAbortState::NotBeingAborted && ShouldForce != EAIForceParam::Force)
+	{
+		UE_VLOG(GetPawn(), LogPawnAction, Log, TEXT("%s> Discarding Abort request due to already in abort state"), *GetName());
+		return AbortState;
+	}
+
+	const bool bForce = ShouldForce == EAIForceParam::Force;
+	EPawnActionAbortState::Type Result = EPawnActionAbortState::NotBeingAborted;
+	EPawnActionAbortState::Type ChildResult = EPawnActionAbortState::AbortDone;
+
+	if (ChildAction != NULL)
+	{
+		ChildResult = ChildAction->Abort(ShouldForce);
+
+		if (ChildResult == EPawnActionAbortState::NotBeingAborted)
+		{
+			UE_VLOG(GetPawn(), LogPawnAction, Error, TEXT("%s> ChildAction %s failed to carry out proper abortion! Things might get ugly..")
+				, *GetName(), *ChildAction->GetName());
+
+			// fake proper result and hope for the best!
+			ChildResult = EPawnActionAbortState::AbortDone;
+		}
+	}
+
+	if (bForce)
+	{
+		Result = PerformAbort(ShouldForce);
+		if (Result != EPawnActionAbortState::AbortDone)
+		{
+			UE_VLOG(GetPawn(), LogPawnAction, Error, TEXT("%s> failed to force-abort! Things might get ugly..")
+				, *GetName());
+
+			// fake proper result and hope for the best!
+			Result = EPawnActionAbortState::AbortDone;
+		}
+	}
+	else
+	{
+		switch (ChildResult)
+		{
+		case EPawnActionAbortState::MarkPendingAbort:
+			// this means child is awaiting its abort, so should parent
+		case EPawnActionAbortState::LatendAbortInProgress:
+			// this means child is performing time-consuming abort. Parent should wait
+			Result = EPawnActionAbortState::MarkPendingAbort;
+			break;
+
+		case EPawnActionAbortState::AbortDone:
+			Result = IsPaused() ? EPawnActionAbortState::MarkPendingAbort : PerformAbort(ShouldForce);
+			break;
+
+		default:
+			UE_VLOG(GetPawn(), LogPawnAction, Error, TEXT("%s> Unhandled Abort State!")
+				, *GetName());
+			Result = EPawnActionAbortState::AbortDone;
+			break;
+		}
+	}
+
+	SetAbortState(Result);
+
+	return Result;
+}
+
+APawn* UPawnAction::GetPawn()
+{
+	return OwnerComponent ? OwnerComponent->GetControlledPawn() : NULL;
+}
+
+AController* UPawnAction::GetController() 
+{
+	return OwnerComponent ? OwnerComponent->GetController() : NULL; 
+}
+
+void UPawnAction::SetAbortState(EPawnActionAbortState::Type NewAbortState)
+{
+	// allowing only progression
+	if (NewAbortState <= AbortState)
+	{
+		return;
+	}
+
+	AbortState = NewAbortState;
+	if (AbortState == EPawnActionAbortState::AbortDone)
+	{
+		SendEvent(EPawnActionEventType::FinishedAborting);
+	}
+}
+
+void UPawnAction::SendEvent(EPawnActionEventType::Type Event)
+{
+	if (OwnerComponent && OwnerComponent->IsPendingKill() == false)
+	{
+		// this will get communicated to parent action if needed, latently 
+		OwnerComponent->OnEvent(this, Event);
+	}
+
+	ActionObserver.ExecuteIfBound(this, Event);
+}
+
+void UPawnAction::StopWaitingForMessages()
+{
+	MessageHandlers.Reset();
+}
+
+void UPawnAction::SetFinishResult(EPawnActionResult::Type Result)
+{
+	// once return value had been set it's no longer possible to back to InProgress
+	if (Result == EPawnActionResult::InProgress)
+	{
+		UE_VLOG(GetPawn(), LogPawnAction, Warning, TEXT("%s> UPawnAction::SetFinishResult setting FinishResult as EPawnActionResult::InProgress - should not be happening"), *GetName());
+		return;
+	}
+
+	if (FinishResult != Result)
+	{
+		FinishResult = Result;
+	}
+}
+
+void UPawnAction::SetOwnerComponent(class UPawnActionsComponent* Component)
+{
+	if (OwnerComponent != NULL && OwnerComponent != Component)
+	{
+		UE_VLOG(GetPawn(), LogPawnAction, Warning, TEXT("%s> UPawnAction::SetOwnerComponent called to change already set valid owner component"), *GetName());
+	}
+
+	OwnerComponent = Component;
+	if (Component != NULL)
+	{
+		AAIController* AIController = Cast<AAIController>(Component->GetController());
+		if (AIController != NULL)
+		{
+			BrainComp = AIController->FindComponentByClass<UBrainComponent>();
+		}
+	}
+}
+
+void UPawnAction::SetInstigator(UObject* const InInstigator)
+{ 
+	if (Instigator != NULL && Instigator != InInstigator)
+	{
+		UE_VLOG(GetPawn(), LogPawnAction, Warning, TEXT("%s> setting Instigator to %s when already has instigator set to %s")
+			, *GetName(), *Instigator->GetName(), *InInstigator->GetName());
+	}
+	Instigator = InInstigator; 
+}
+
+void UPawnAction::Finish(TEnumAsByte<EPawnActionResult::Type> WithResult)
+{
+	UE_VLOG(GetPawn(), LogPawnAction, Log, TEXT("%s> finishing with result %s")
+		, *GetName(), *GetActionResultName(WithResult));
+
+	SetFinishResult(WithResult);
+
+	StopWaitingForMessages();
+
+	SendEvent(EPawnActionEventType::FinishedExecution);
+}
+
+bool UPawnAction::Activate()
+{
+	bool bResult = false; 
+
+	UE_VLOG(GetPawn(), LogPawnAction, Log, TEXT("%s> Activating at priority %s! First start? %s Paused? %s")
+		, *GetName()
+		, *GetPriorityName()
+		, bHasBeenStarted ? TEXT("NO") : TEXT("YES")
+		, IsPaused() ? TEXT("YES") : TEXT("NO"));
+
+	if (bHasBeenStarted == true && IsPaused())
+	{
+		bResult = Resume();
+	}
+	else 
+	{
+		bResult = Start();
+		if (bResult == false)
+		{
+			UE_VLOG(GetPawn(), LogPawnAction, Log, TEXT("%s> Failed to start.")
+				, *GetName());
+			bFailedToStart = true;
+			SetFinishResult(EPawnActionResult::Failed);
+			SendEvent(EPawnActionEventType::FailedToStart);
+		}
+	}
+
+	return bResult;
+}
+
+void UPawnAction::OnPopped()
+{
+	// not calling OnFinish if action haven't actually started
+	if (bFailedToStart == false)
+	{
+		OnFinished(FinishResult);
+	}
+}
+
+bool UPawnAction::Start()
+{
+	bHasBeenStarted = true;
+	bPaused = false;
+	return true;
+}
+
+
+bool UPawnAction::Pause()
+{
+	// parent should be paused anyway
+	ensure(ParentAction == NULL || ParentAction->IsPaused() == true);
+
+	// don't pause twice
+	if (bPaused)
+	{
+		return false;
+	}
+	else if (AbortState == EPawnActionAbortState::LatendAbortInProgress || AbortState == EPawnActionAbortState::AbortDone)
+	{
+		UE_VLOG(GetPawn(), LogPawnAction, Log, TEXT("%s> Not pausing due to being in unpausable aborting state")
+			, *GetName(), *ChildAction->GetName());
+		return false;
+	}
+
+	UE_VLOG(GetPawn(), LogPawnAction, Log, TEXT("%s> Pausing..."), *GetName());
+	
+	bPaused = true;
+
+	if (ChildAction)
+	{
+		ChildAction->Pause();
+	}
+
+	return bPaused;
+}
+
+bool UPawnAction::Resume()
+{
+	// parent should be paused anyway
+	ensure(ParentAction == NULL || ParentAction->IsPaused() == true);
+
+	// do not unpause twice
+	if (bPaused == false)
+	{
+		return false;
+	}
+	
+	ensure(ChildAction == NULL);
+
+	if (ChildAction)
+	{
+		UE_VLOG(GetPawn(), LogPawnAction, Log, TEXT("%s> Resuming child, %s"), *GetName(), *ChildAction->GetName());
+		ChildAction->Resume();
+	}
+	else
+	{
+		UE_VLOG(GetPawn(), LogPawnAction, Log, TEXT("%s> Resuming."), *GetName());
+		bPaused = false;
+	}
+
+	return !bPaused;
+}
+
+void UPawnAction::OnFinished(EPawnActionResult::Type WithResult)
+{
+}
+
+void UPawnAction::OnChildFinished(UPawnAction* Action, EPawnActionResult::Type WithResult)
+{
+	if (Action == NULL)
+	{
+		return;
+	}
+
+	UE_VLOG(GetPawn(), LogPawnAction, Log, TEXT("%s> Child \'%s\' finished with result %s")
+		, *GetName(), *Action->GetName(), *GetActionResultName(WithResult));
+
+	ensure(Action->ParentAction == this);
+	ensure(ChildAction == Action);
+	Action->ParentAction = NULL;
+	ChildAction = NULL;
+}
+
+void UPawnAction::OnPushed()
+{
+	IndexOnStack = 0;
+	UPawnAction* PrevAction = ParentAction;
+	while (PrevAction)
+	{
+		++IndexOnStack;
+		PrevAction = PrevAction->ParentAction;
+	}
+
+	UE_VLOG(GetPawn(), LogPawnAction, Log, TEXT("%s> Pushed with priority %s, IndexOnStack: %d, instigator %s")
+		, *GetName(), *GetPriorityName(), IndexOnStack, *GetNameSafe(Instigator));
+}
+
+bool UPawnAction::PushChildAction(UPawnAction* Action)
+{
+	bool bResult = false;
+	
+	if (OwnerComponent != NULL)
+	{
+		UE_CVLOG( ChildAction != NULL
+			, GetPawn(), LogPawnAction, Log, TEXT("%s> Pushing child action %s while already having ChildAction set to %s")
+			, *GetName(), *Action->GetName(), *ChildAction->GetName());
+		
+		// copy runtime data
+		// note that priority and instigator will get assigned as part of PushAction.
+
+		bResult = OwnerComponent->PushAction(Action, GetPriority(), Instigator);
+
+		// adding a check to make sure important data has been set 
+		ensure(Action->GetPriority() == GetPriority() && Action->GetInstigator() == GetInstigator());
+	}
+
+	return bResult;
+}
+
+//----------------------------------------------------------------------//
+// messaging
+//----------------------------------------------------------------------//
+
+void UPawnAction::WaitForMessage(FName MessageType, FAIRequestID RequestID)
+{
+	MessageHandlers.Add(FAIMessageObserver::Create(BrainComp, MessageType, RequestID.GetID(), FOnAIMessage::CreateUObject(this, &UPawnAction::HandleAIMessage)));
+}
+
+//----------------------------------------------------------------------//
+// blueprint interface
+//----------------------------------------------------------------------//
+
+UPawnAction* UPawnAction::CreateActionInstance(UObject* WorldContextObject, TSubclassOf<UPawnAction> ActionClass)
+{
+	UWorld* World = GEngine->GetWorldFromContextObject(WorldContextObject);
+	if (World && ActionClass)
+	{
+		return ConstructObject<UPawnAction>(ActionClass, World);
+	}
+	return NULL;
+}
+
+//----------------------------------------------------------------------//
+// debug
+//----------------------------------------------------------------------//
+
+FString UPawnAction::GetStateDescription() const
+{
+	static const UEnum* AbortStateEnum = FindObject<UEnum>(ANY_PACKAGE, TEXT("EPawnActionAbortState")); 
+		
+	if (AbortState != EPawnActionAbortState::NotBeingAborted)
+	{
+		return *AbortStateEnum->GetEnumText(AbortState).ToString();
+	}
+	return IsPaused() ? TEXT("Paused") : TEXT("Active");
+}
+
+FString UPawnAction::GetPriorityName() const
+{
+	static const UEnum* Enum = FindObject<UEnum>(ANY_PACKAGE, TEXT("EAIRequestPriority"));
+	check(Enum);
+	return Enum->GetEnumName(GetPriority());
+}
diff --git a/Engine/Source/Runtime/AIModule/Private/Actions/PawnAction_BlueprintBase.cpp b/Engine/Source/Runtime/AIModule/Private/Actions/PawnAction_BlueprintBase.cpp
new file mode 100644
--- /dev/null
+++ b/Engine/Source/Runtime/AIModule/Private/Actions/PawnAction_BlueprintBase.cpp
@@ -0,0 +1,64 @@
+// Copyright 1998-2014 Epic Games, Inc. All Rights Reserved.
+
+#include "AIModulePrivate.h"
+#include "Actions/PawnAction_BlueprintBase.h"
+
+namespace
+{
+	FORCEINLINE bool HasBlueprintFunction(FName FuncName, const UObject* Ob, const UClass* StopAtClass)
+	{
+		return (Ob->GetClass()->FindFunctionByName(FuncName)->GetOuter() != StopAtClass);
+	}
+}
+
+UPawnAction_BlueprintBase::UPawnAction_BlueprintBase(const class FPostConstructInitializeProperties& PCIP)
+	: Super(PCIP)
+{
+	UClass* StopAtClass = UPawnAction_BlueprintBase::StaticClass();
+	bWantsTick = HasBlueprintFunction(TEXT("ActionTick"), this, StopAtClass);
+}
+
+void UPawnAction_BlueprintBase::Tick(float DeltaTime)
+{
+	// no need to call super implementation
+	ActionTick(GetPawn(), DeltaTime);
+}
+
+bool UPawnAction_BlueprintBase::Start()
+{
+	const bool bHasBeenEverStarted = HasBeenStarted();
+	const bool bSuperResult = Super::Start();
+
+	if (bHasBeenEverStarted == false && bSuperResult == true)
+	{
+		ActionStart(GetPawn());
+	}
+
+	return bSuperResult;
+}
+
+bool UPawnAction_BlueprintBase::Pause()
+{
+	bool bResult = Super::Pause();
+	if (bResult)
+	{
+		ActionPause(GetPawn());
+	}
+	return bResult;
+}
+
+bool UPawnAction_BlueprintBase::Resume()
+{
+	bool bResult = Super::Resume();
+	if (bResult)
+	{
+		ActionResume(GetPawn());
+	}
+
+	return bResult;
+}
+
+void UPawnAction_BlueprintBase::OnFinished(EPawnActionResult::Type WithResult)
+{
+	ActionFinished(GetPawn(), WithResult);
+}
diff --git a/Engine/Source/Runtime/AIModule/Private/Actions/PawnAction_Move.cpp b/Engine/Source/Runtime/AIModule/Private/Actions/PawnAction_Move.cpp
new file mode 100644
--- /dev/null
+++ b/Engine/Source/Runtime/AIModule/Private/Actions/PawnAction_Move.cpp
@@ -0,0 +1,243 @@
+// Copyright 1998-2014 Epic Games, Inc. All Rights Reserved.
+
+#include "AIModulePrivate.h"
+#include "Actions/PawnAction_Move.h"
+
+UPawnAction_Move::UPawnAction_Move(const class FPostConstructInitializeProperties& PCIP)
+	: Super(PCIP)
+	, GoalLocation(FAISystem::InvalidLocation)
+	, AcceptableRadius(30.f)
+	, bUsePathfinding(true)
+	, bProjectGoalToNavigation(false)
+	, bUpdatePathToGoal(true)
+	, bAbortChildActionOnPathChange(false)
+{
+}
+
+UPawnAction_Move* UPawnAction_Move::CreateAction(UWorld* World, class AActor* GoalActor, EPawnActionMoveMode::Type Mode)
+{
+	if (GoalActor == NULL)
+	{
+		return NULL;
+	}
+
+	UPawnAction_Move* Action = UPawnAction::CreateActionInstance<UPawnAction_Move>(World);
+	if (Action)
+	{
+		Action->GoalActor = GoalActor;
+		Action->bUsePathfinding = (Mode == EPawnActionMoveMode::UsePathfinding);
+	}
+
+	return Action;
+}
+
+UPawnAction_Move* UPawnAction_Move::CreateAction(UWorld* World, const FVector& GoalLocation, EPawnActionMoveMode::Type Mode)
+{
+	if (FAISystem::IsValidLocation(GoalLocation) == false)
+	{
+		return NULL;
+	}
+
+	UPawnAction_Move* Action = UPawnAction::CreateActionInstance<UPawnAction_Move>(World);
+	if (Action)
+	{
+		Action->GoalLocation = GoalLocation;
+		Action->bUsePathfinding = (Mode == EPawnActionMoveMode::UsePathfinding);
+	}
+
+	return Action;
+}
+
+bool UPawnAction_Move::Start()
+{
+	bool bResult = Super::Start();
+	if (bResult)
+	{
+		bResult = PerformMoveAction();
+	}
+
+	return bResult;
+}
+
+EPathFollowingRequestResult::Type UPawnAction_Move::RequestMove(AAIController* Controller)
+{
+	EPathFollowingRequestResult::Type RequestResult = EPathFollowingRequestResult::Failed;
+
+	if (GoalActor != NULL)
+	{
+		const bool bAtGoal = CheckAlreadyAtGoal(Controller, GoalActor, AcceptableRadius);
+		RequestResult = bAtGoal ? EPathFollowingRequestResult::AlreadyAtGoal : 
+			bUpdatePathToGoal ? Controller->MoveToActor(GoalActor, AcceptableRadius, true, bUsePathfinding, bAllowStrafe, FilterClass) :
+			Controller->MoveToLocation(GoalActor->GetActorLocation(), AcceptableRadius, true, bUsePathfinding, bProjectGoalToNavigation, bAllowStrafe);
+	}
+	else if (FAISystem::IsValidLocation(GoalLocation))
+	{
+		const bool bAtGoal = CheckAlreadyAtGoal(Controller, GoalLocation, AcceptableRadius);
+		RequestResult = bAtGoal ? EPathFollowingRequestResult::AlreadyAtGoal :
+			Controller->MoveToLocation(GoalLocation, AcceptableRadius, true, bUsePathfinding, bProjectGoalToNavigation, bAllowStrafe, FilterClass);
+	}
+	else
+	{
+		UE_VLOG(Controller, LogPawnAction, Warning, TEXT("UPawnAction_Move::Start: no valid move goal set"));
+	}
+
+	return RequestResult;
+}
+
+bool UPawnAction_Move::PerformMoveAction()
+{
+	AAIController* MyController = Cast<AAIController>(GetController());
+	if (MyController == NULL)
+	{
+		return false;
+	}
+
+	if (MyController->ShouldPostponePathUpdates())
+	{
+		UE_VLOG(MyController, LogPawnAction, Log, TEXT("Can't path right now, waiting..."));
+		MyController->GetWorldTimerManager().SetTimer(this, &UPawnAction_Move::DeferredPerformMoveAction, 0.1f);
+		return true;
+	}
+
+	const EPathFollowingRequestResult::Type RequestResult = RequestMove(MyController);
+	bool bResult = true;
+
+	if (RequestResult == EPathFollowingRequestResult::RequestSuccessful)
+	{
+		RequestID = MyController->GetCurrentMoveRequestID();
+		WaitForMessage(UBrainComponent::AIMessage_MoveFinished, RequestID);
+		WaitForMessage(UBrainComponent::AIMessage_RepathFailed);
+	}
+	else if (RequestResult == EPathFollowingRequestResult::AlreadyAtGoal)
+	{
+		// note that this will result in latently notifying actions component about finishing
+		// so it's no problem we run it from withing Start function
+		Finish(EPawnActionResult::Success);
+	}
+	else
+	{
+		bResult = false;
+	}
+
+	return bResult;
+}
+
+void UPawnAction_Move::DeferredPerformMoveAction()
+{
+	const bool bResult = PerformMoveAction();
+	if (!bResult)
+	{
+		Finish(EPawnActionResult::Failed);
+	}
+}
+
+bool UPawnAction_Move::Pause()
+{
+	bool bResult = Super::Pause();
+	if (bResult)
+	{
+		AAIController* MyController = Cast<AAIController>(GetController());
+		if (MyController)
+		{
+			bResult = MyController->PauseMove(RequestID);
+		}
+	}
+	return bResult;
+}
+
+bool UPawnAction_Move::Resume()
+{
+	bool bResult = Super::Resume();
+	if (bResult)
+	{
+		bResult = false;
+		AAIController* MyController = Cast<AAIController>(GetController());
+		if (MyController)
+		{
+			if (MyController->ResumeMove(RequestID) == false)
+			{
+				// try requesting a new move
+				StopWaitingForMessages();
+				
+				bResult = PerformMoveAction();
+			}
+			else
+			{
+				bResult = true;
+			}
+		}
+	}
+	return bResult;
+}
+
+EPawnActionAbortState::Type UPawnAction_Move::PerformAbort(EAIForceParam::Type ShouldForce)
+{
+	AAIController* MyController = Cast<AAIController>(GetController());
+
+	if (MyController && MyController->PathFollowingComponent)
+	{
+		MyController->PathFollowingComponent->AbortMove(TEXT("BehaviorTree abort"), RequestID);
+	}
+
+	return Super::PerformAbort(ShouldForce);
+}
+
+void UPawnAction_Move::HandleAIMessage(UBrainComponent*, const FAIMessage& Message)
+{
+	if (Message.MessageName == UBrainComponent::AIMessage_MoveFinished && Message.HasFlag(EPathFollowingMessage::OtherRequest))
+	{
+		// move was aborted by another request from different action, don't finish yet
+		return;
+	}
+
+	const bool bFail = Message.MessageName == UBrainComponent::AIMessage_RepathFailed
+		|| Message.Status == FAIMessage::Failure;
+
+	Finish(bFail ? EPawnActionResult::Failed : EPawnActionResult::Success);
+}
+
+void UPawnAction_Move::SetPath(FNavPathSharedPtr InPath)
+{
+	Path = InPath;
+	Path->GetObserver(PathObserver);
+	Path->SetObserver(FNavigationPath::FPathObserverDelegate::CreateUObject(this, &UPawnAction_Move::OnPathUpdated));
+}
+
+void UPawnAction_Move::OnPathUpdated(FNavigationPath* UpdatedPath)
+{
+	UE_VLOG(GetController(), LogPawnAction, Log, TEXT("%s> Path updated!"), *GetName());
+	if (PathObserver.IsBound() && PathObserver.GetUObject() != this)
+	{
+		PathObserver.ExecuteIfBound(UpdatedPath);
+	}
+
+	if (bAbortChildActionOnPathChange && GetChildAction())
+	{
+		UE_VLOG(GetController(), LogPawnAction, Log, TEXT(">> aborting child action: %s"), *GetNameSafe(GetChildAction()));
+		GetChildAction()->Abort(EAIForceParam::Force);
+	}
+}
+
+bool UPawnAction_Move::CheckAlreadyAtGoal(class AAIController* Controller, const FVector& TestLocation, float Radius)
+{
+	const bool bAlreadyAtGoal = Controller->PathFollowingComponent->HasReached(TestLocation, Radius);
+	if (bAlreadyAtGoal)
+	{
+		Controller->PathFollowingComponent->AbortMove(TEXT("Aborting move due to new move request finishing with AlreadyAtGoal"), FAIRequestID::AnyRequest, true, false, EPathFollowingMessage::OtherRequest);
+		Controller->PathFollowingComponent->SetLastMoveAtGoal(true);
+	}
+
+	return bAlreadyAtGoal;
+}
+
+bool UPawnAction_Move::CheckAlreadyAtGoal(class AAIController* Controller, const AActor* TestGoal, float Radius)
+{
+	const bool bAlreadyAtGoal = Controller->PathFollowingComponent->HasReached(TestGoal, Radius);
+	if (bAlreadyAtGoal)
+	{
+		Controller->PathFollowingComponent->AbortMove(TEXT("Aborting move due to new move request finishing with AlreadyAtGoal"), FAIRequestID::AnyRequest, true, false, EPathFollowingMessage::OtherRequest);
+		Controller->PathFollowingComponent->SetLastMoveAtGoal(true);
+	}
+
+	return bAlreadyAtGoal;
+}
diff --git a/Engine/Source/Runtime/AIModule/Private/Actions/PawnAction_Repeat.cpp b/Engine/Source/Runtime/AIModule/Private/Actions/PawnAction_Repeat.cpp
new file mode 100644
--- /dev/null
+++ b/Engine/Source/Runtime/AIModule/Private/Actions/PawnAction_Repeat.cpp
@@ -0,0 +1,96 @@
+// Copyright 1998-2014 Epic Games, Inc. All Rights Reserved.
+
+#include "AIModulePrivate.h"
+#include "Actions/PawnAction_Repeat.h"
+
+UPawnAction_Repeat::UPawnAction_Repeat(const class FPostConstructInitializeProperties& PCIP)
+	: Super(PCIP)
+{
+	ChildFailureHandlingMode = EPawnActionFailHandling::IgnoreFailure;
+}
+
+UPawnAction_Repeat* UPawnAction_Repeat::CreateAction(UWorld* World, UPawnAction* ActionToRepeat, int32 NumberOfRepeats)
+{
+	if (ActionToRepeat == NULL || !(NumberOfRepeats > 0 || NumberOfRepeats == UPawnAction_Repeat::LoopForever))
+	{
+		return NULL;
+	}
+
+	UPawnAction_Repeat* Action = UPawnAction::CreateActionInstance<UPawnAction_Repeat>(World);
+	if (Action)
+	{
+		Action->ActionToRepeat = ActionToRepeat;
+		Action->RepeatsLeft = NumberOfRepeats;
+	}
+
+	return Action;
+}
+
+bool UPawnAction_Repeat::Start()
+{
+	bool bResult = Super::Start();
+	
+	if (bResult)
+	{
+		UE_VLOG(GetPawn(), LogPawnAction, Log, TEXT("Starting repeating action: %s. Requested repeats: %d")
+			, *GetNameSafe(ActionToRepeat), RepeatsLeft);
+		bResult = PushActionCopy();
+	}
+
+	return bResult;
+}
+
+bool UPawnAction_Repeat::Resume()
+{
+	bool bResult = Super::Resume();
+
+	if (bResult)
+	{
+		bResult = PushActionCopy();
+	}
+
+	return bResult;
+}
+
+void UPawnAction_Repeat::OnChildFinished(UPawnAction* Action, EPawnActionResult::Type WithResult)
+{
+	if (RecentActionCopy == Action)
+	{
+		if (WithResult == EPawnActionResult::Success || (WithResult == EPawnActionResult::Failed && ChildFailureHandlingMode == EPawnActionFailHandling::IgnoreFailure))
+		{
+			PushActionCopy();
+		}
+		else
+		{
+			Finish(EPawnActionResult::Failed);
+		}
+	}
+
+	Super::OnChildFinished(Action, WithResult);
+}
+
+bool UPawnAction_Repeat::PushActionCopy()
+{
+	if (ActionToRepeat == NULL)
+	{
+		Finish(EPawnActionResult::Failed);
+		return false;
+	}
+	else if (RepeatsLeft == 0)
+	{
+		Finish(EPawnActionResult::Success);
+		return true;
+	}
+
+	if (RepeatsLeft > 0)
+	{
+		--RepeatsLeft;
+	}
+
+	UPawnAction* ActionCopy = Cast<UPawnAction>(StaticDuplicateObject(ActionToRepeat, this, NULL));
+	UE_VLOG(GetPawn(), LogPawnAction, Log, TEXT("%s> pushing repeted action copy %s, repeats left: %d")
+		, *GetName(), *GetNameSafe(ActionCopy), RepeatsLeft);
+	ensure(ActionCopy);
+	RecentActionCopy = ActionCopy;
+	return PushChildAction(ActionCopy); 
+}
diff --git a/Engine/Source/Runtime/AIModule/Private/Actions/PawnAction_Sequence.cpp b/Engine/Source/Runtime/AIModule/Private/Actions/PawnAction_Sequence.cpp
new file mode 100644
--- /dev/null
+++ b/Engine/Source/Runtime/AIModule/Private/Actions/PawnAction_Sequence.cpp
@@ -0,0 +1,90 @@
+// Copyright 1998-2014 Epic Games, Inc. All Rights Reserved.
+
+#include "AIModulePrivate.h"
+#include "Actions/PawnAction_Sequence.h"
+
+UPawnAction_Sequence::UPawnAction_Sequence(const class FPostConstructInitializeProperties& PCIP)
+	: Super(PCIP)
+{
+}
+
+UPawnAction_Sequence* UPawnAction_Sequence::CreateAction(UWorld* World, TArray<UPawnAction*>& ActionSequence)
+{
+	ActionSequence.Remove(NULL);
+	if (ActionSequence.Num() <= 0)
+	{
+		return NULL;
+	}
+
+	UPawnAction_Sequence* Action = UPawnAction::CreateActionInstance<UPawnAction_Sequence>(World);
+	if (Action)
+	{
+		Action->ActionSequence = ActionSequence;
+	}
+
+	return Action;
+}
+
+bool UPawnAction_Sequence::Start()
+{
+	bool bResult = Super::Start();
+
+	if (bResult)
+	{
+		UE_VLOG(GetPawn(), LogPawnAction, Log, TEXT("Starting sequence. Items:"), *GetName());
+		for (auto Action : ActionSequence)
+		{
+			UE_VLOG(GetPawn(), LogPawnAction, Log, TEXT("    %s"), *GetNameSafe(Action));
+		}
+
+		bResult = PushNextActionCopy();
+	}
+
+	return bResult;
+}
+
+bool UPawnAction_Sequence::Resume()
+{
+	bool bResult = Super::Resume();
+
+	if (bResult)
+	{
+		bResult = PushNextActionCopy();
+	}
+
+	return bResult;
+}
+
+void UPawnAction_Sequence::OnChildFinished(UPawnAction* Action, EPawnActionResult::Type WithResult)
+{
+	if (RecentActionCopy == Action)
+	{
+		if (WithResult == EPawnActionResult::Success || (WithResult == EPawnActionResult::Failed && ChildFailureHandlingMode == EPawnActionFailHandling::IgnoreFailure))
+		{
+			PushNextActionCopy();
+		}
+		else
+		{
+			Finish(EPawnActionResult::Failed);
+		}
+	}
+
+	Super::OnChildFinished(Action, WithResult);
+}
+
+bool UPawnAction_Sequence::PushNextActionCopy()
+{
+	if (CurrentActionIndex >= uint32(ActionSequence.Num()))
+	{
+		Finish(EPawnActionResult::Success);
+		return true;
+	}
+
+	UPawnAction* ActionCopy = Cast<UPawnAction>(StaticDuplicateObject(ActionSequence[CurrentActionIndex], this, NULL));
+	UE_VLOG(GetPawn(), LogPawnAction, Log, TEXT("%s> pushing action %s")
+		, *GetName(), *GetNameSafe(ActionCopy));
+	++CurrentActionIndex;	
+	ensure(ActionCopy);
+	RecentActionCopy = ActionCopy;
+	return PushChildAction(ActionCopy);
+}
diff --git a/Engine/Source/Runtime/AIModule/Private/Actions/PawnActionsComponent.cpp b/Engine/Source/Runtime/AIModule/Private/Actions/PawnActionsComponent.cpp
new file mode 100644
--- /dev/null
+++ b/Engine/Source/Runtime/AIModule/Private/Actions/PawnActionsComponent.cpp
@@ -0,0 +1,453 @@
+// Copyright 1998-2014 Epic Games, Inc. All Rights Reserved.
+
+#include "AIModulePrivate.h"
+#include "Actions/PawnActionsComponent.h"
+
+DEFINE_LOG_CATEGORY(LogPawnAction);
+
+//----------------------------------------------------------------------//
+// helpers
+//----------------------------------------------------------------------//
+
+namespace
+{
+	FString GetEventName(int32 Index)
+	{
+		static const UEnum* Enum = FindObject<UEnum>(ANY_PACKAGE, TEXT("EPawnActionEventType"));
+		check(Enum);
+		return Enum->GetEnumName(Index);
+	}
+
+	FString GetPriorityName(int32 Index)
+	{
+		static const UEnum* Enum = FindObject<UEnum>(ANY_PACKAGE, TEXT("EAIRequestPriority"));
+		check(Enum);
+		return Enum->GetEnumName(Index);
+	}
+
+	FString GetActionSignature(UPawnAction* Action)
+	{
+		if (Action == NULL)
+		{
+			return TEXT("NULL");
+		}
+		
+		return FString::Printf(TEXT("[%s, %s]"), *Action->GetName(), *GetPriorityName(Action->GetPriority()));
+	}
+}
+
+//----------------------------------------------------------------------//
+// FPawnActionEvent
+//----------------------------------------------------------------------//
+
+namespace
+{
+	struct FPawnActionEvenSort
+	{
+		FORCEINLINE bool operator()(const FPawnActionEvent& A, const FPawnActionEvent& B) const
+		{
+			return A.Priority < B.Priority
+				|| (A.Priority == B.Priority
+					&& (A.EventType < B.EventType
+						|| (A.EventType == B.EventType && A.Index < B.Index)));
+		}
+	};
+}
+
+FPawnActionEvent::FPawnActionEvent(UPawnAction* InAction, EPawnActionEventType::Type InEventType, uint32 InIndex)
+	: Action(InAction), EventType(InEventType), Index(InIndex)
+{
+	ensure(InAction);
+	if (InAction)
+	{
+		Priority = InAction->GetPriority();
+	}
+}
+
+//----------------------------------------------------------------------//
+// FPawnActionStack
+//----------------------------------------------------------------------//
+
+void FPawnActionStack::Pause()
+{
+	if (TopAction != NULL)
+	{
+		TopAction->Pause();
+	}
+}
+
+void FPawnActionStack::Resume()
+{
+	if (TopAction != NULL)
+	{
+		TopAction->Resume();
+	}
+}
+
+void FPawnActionStack::PushAction(UPawnAction* NewTopAction)
+{
+	if (TopAction != NULL)
+	{
+		TopAction->Pause();
+		ensure(TopAction->ChildAction == NULL);
+		TopAction->ChildAction = NewTopAction;
+		NewTopAction->ParentAction = TopAction;
+	}
+
+	TopAction = NewTopAction;
+	NewTopAction->OnPushed();
+}
+
+void FPawnActionStack::PopAction(UPawnAction* ActionToPop)
+{
+	if (ActionToPop == NULL)
+	{
+		return;
+	}
+
+	// first check if it's there
+	UPawnAction* CutPoint = TopAction;
+	while (CutPoint != NULL && CutPoint != ActionToPop)
+	{
+		CutPoint = CutPoint->ParentAction;
+	}
+
+	if (CutPoint == ActionToPop)
+	{
+		UPawnAction* ActionBeingRemoved = TopAction;
+		// note StopAction can be null
+		UPawnAction* StopAction = ActionToPop->ParentAction;
+
+		while (ActionBeingRemoved != StopAction)
+		{
+			UPawnAction* NextAction = ActionBeingRemoved->ParentAction;
+			
+			ActionBeingRemoved->OnPopped();
+			if (ActionBeingRemoved->ParentAction)
+			{
+				ActionBeingRemoved->ParentAction->OnChildFinished(ActionBeingRemoved, ActionBeingRemoved->FinishResult);
+			}
+			ActionBeingRemoved = NextAction;
+		}
+
+		TopAction = StopAction;
+	}
+}
+
+//----------------------------------------------------------------------//
+// UPawnActionsComponent
+//----------------------------------------------------------------------//
+
+UPawnActionsComponent::UPawnActionsComponent(const class FPostConstructInitializeProperties& PCIP)
+	: Super(PCIP)
+{
+	PrimaryComponentTick.bCanEverTick = true;
+	PrimaryComponentTick.bStartWithTickEnabled = true;
+
+	bAutoActivate = true;
+	bLockedAILogic = false;
+
+	ActionEventIndex = 0;
+
+	ActionStacks.AddZeroed(EAIRequestPriority::MAX);
+}
+
+void UPawnActionsComponent::TickComponent(float DeltaTime, enum ELevelTick TickType, FActorComponentTickFunction *ThisTickFunction)
+{
+	Super::TickComponent(DeltaTime, TickType, ThisTickFunction);
+
+	if (ControlledPawn == NULL)
+	{
+		CacheControlledPawn();
+	}
+
+	if (ActionEvents.Num() > 1)
+	{
+		ActionEvents.Sort(FPawnActionEvenSort());
+	}
+
+	if (ActionEvents.Num() > 0)
+	{
+		for (int32 EventIndex = 0; EventIndex < ActionEvents.Num(); ++EventIndex)
+		{
+			FPawnActionEvent& Event = ActionEvents[EventIndex];
+			switch (Event.EventType)
+			{
+			case EPawnActionEventType::FinishedAborting:
+			case EPawnActionEventType::FinishedExecution:
+			case EPawnActionEventType::FailedToStart:
+				ActionStacks[Event.Priority].PopAction(Event.Action);
+				break;
+			case EPawnActionEventType::Push:
+				ActionStacks[Event.Priority].PushAction(Event.Action);
+				break;
+			default:
+				break;
+			}
+		}
+
+		ActionEvents.Reset();
+
+		UpdateCurrentAction();
+	}
+
+	if (CurrentAction)
+	{
+		CurrentAction->TickAction(DeltaTime);
+	}
+
+	// it's possible we got new events with CurrentAction's tick
+	if (ActionEvents.Num() == 0 && (CurrentAction == NULL || CurrentAction->WantsTick() == false))
+	{
+		SetComponentTickEnabled(false);
+	}
+}
+
+void UPawnActionsComponent::UpdateCurrentAction()
+{
+	UE_VLOG(ControlledPawn, LogPawnAction, Log, TEXT("Picking new current actions. Old CurrentAction %s")
+		, *GetActionSignature(CurrentAction));
+
+	UPawnAction* NewCurrentAction = NULL;
+	int32 Priority = EAIRequestPriority::MAX - 1;
+	do 
+	{
+		NewCurrentAction = ActionStacks[Priority].GetTop();
+
+	} while (NewCurrentAction == NULL && --Priority >= 0);
+
+	if (CurrentAction != NewCurrentAction)
+	{
+		UE_VLOG(ControlledPawn, LogPawnAction, Log, TEXT("New action: %s")
+			, *GetActionSignature(NewCurrentAction));
+
+		if (CurrentAction != NULL && CurrentAction->IsActive())
+		{
+			CurrentAction->Pause();
+		}
+		CurrentAction = NewCurrentAction;
+		bool bNewActionStartedSuccessfully = true;
+		if (CurrentAction != NULL)
+		{
+			bNewActionStartedSuccessfully = CurrentAction->Activate();
+		}
+
+		if (bNewActionStartedSuccessfully == false)
+		{
+			UE_VLOG(ControlledPawn, LogPawnAction, Warning, TEXT("CurrentAction %s failed to activate. Removing and re-running action selection")
+				, *GetActionSignature(NewCurrentAction));
+
+			CurrentAction = NULL;			
+		}
+		// @HACK temporary solution to have actions and old BT tasks work together
+		else if (CurrentAction == NULL || CurrentAction->GetPriority() != EAIRequestPriority::Logic)
+		{
+			UpdateAILogicLock();
+		}
+	}
+	else
+	{
+		if (CurrentAction == NULL)
+		{
+			UpdateAILogicLock();
+		}
+		else
+		{ 
+			UE_VLOG(ControlledPawn, LogPawnAction, Warning, TEXT("Still doing the same action"));
+		}
+	}
+}
+
+void UPawnActionsComponent::UpdateAILogicLock()
+{
+	if (ControlledPawn && ControlledPawn->GetController())
+	{
+		UBrainComponent* BrainComp = ControlledPawn->GetController()->FindComponentByClass<UBrainComponent>();
+		if (BrainComp)
+		{
+			if (CurrentAction != NULL && CurrentAction->GetPriority() > EAIRequestPriority::Logic)
+			{
+				UE_VLOG(ControlledPawn, LogPawnAction, Log, TEXT("Locking AI logic"));
+				BrainComp->LockResource(EAILockSource::Script);
+				bLockedAILogic = true;
+			}
+			else if (bLockedAILogic)
+			{
+				UE_VLOG(ControlledPawn, LogPawnAction, Log, TEXT("Clearing AI logic lock"));
+				bLockedAILogic = false;
+				BrainComp->ClearResourceLock(EAILockSource::Script);
+				if (BrainComp->IsResourceLocked() == false)
+				{
+					UE_VLOG(ControlledPawn, LogPawnAction, Log, TEXT("Reseting AI logic"));
+					BrainComp->RestartLogic();
+				}
+				// @todo consider checking if lock priority is < Logic
+				else
+				{
+					UE_VLOG(ControlledPawn, LogPawnAction, Log, TEXT("AI logic still locked with other priority"));
+					BrainComp->RequestLogicRestartOnUnlock();
+				}
+			}
+		}
+	}
+}
+
+bool UPawnActionsComponent::AbortAction(UPawnAction* ActionToAbort)
+{
+	if (ActionToAbort == NULL)
+	{
+		return false;
+	}
+
+	ActionToAbort->Abort(EAIForceParam::DoNotForce);
+	return true;
+}
+
+uint32 UPawnActionsComponent::AbortActionsInstigatedBy(UObject* const Instigator, EAIRequestPriority::Type Priority)
+{
+	uint32 AbortedActionsCount = 0;
+
+	if (Priority == EAIRequestPriority::MAX)
+	{
+		// call for every regular priority 
+		for (int32 PriorityIndex = 0; PriorityIndex < EAIRequestPriority::MAX; ++PriorityIndex)
+		{
+			AbortedActionsCount += AbortActionsInstigatedBy(Instigator, EAIRequestPriority::Type(PriorityIndex));
+		}
+	}
+	else
+	{
+		UPawnAction* Action = ActionStacks[Priority].GetTop();
+		while (Action)
+		{
+			if (Action->GetInstigator() == Instigator)
+			{
+				Action->Abort(EAIForceParam::Force);
+				++AbortedActionsCount;
+			}
+			Action = Action->ParentAction;
+		}
+	}
+
+	return AbortedActionsCount;
+}
+
+bool UPawnActionsComponent::PushAction(UPawnAction* NewAction, EAIRequestPriority::Type Priority, UObject* const Instigator)
+{
+	NewAction->ExecutionPriority = Priority;
+	NewAction->SetOwnerComponent(this);
+	NewAction->SetInstigator(Instigator);
+	return OnEvent(NewAction, EPawnActionEventType::Push);
+}
+
+bool UPawnActionsComponent::OnEvent(UPawnAction* Action, EPawnActionEventType::Type Event)
+{
+	if (Action == NULL || Event == EPawnActionEventType::Invalid)
+	{
+		// ignore
+		UE_VLOG(ControlledPawn, LogPawnAction, Warning, TEXT("Ignoring Action Event: Action %s Event %s")
+			, *GetNameSafe(Action), *GetEventName(Event));
+		return false;
+	}
+	ActionEvents.Add(FPawnActionEvent(Action, Event, ActionEventIndex++));
+
+	// if it's a first even enable tick
+	if (ActionEvents.Num() == 1)
+	{
+		SetComponentTickEnabled(true);
+	}
+
+	return true;
+}
+
+void UPawnActionsComponent::SetControlledPawn(APawn* NewPawn)
+{
+	if (ControlledPawn != NULL && ControlledPawn != NewPawn)
+	{
+		UE_VLOG(ControlledPawn, LogPawnAction, Warning, TEXT("Trying to set ControlledPawn to new value while ControlledPawn is already set!"));
+	}
+	else
+	{
+		ControlledPawn = NewPawn;
+	}
+}
+
+APawn* UPawnActionsComponent::CacheControlledPawn()
+{
+	if (ControlledPawn == NULL)
+	{
+		AActor* ActorOwner = GetOwner();
+		if (ActorOwner)
+		{
+			ControlledPawn = Cast<APawn>(ActorOwner);
+			if (ControlledPawn == NULL)
+			{
+				AController* Controller = Cast<AController>(ActorOwner);
+				if (Controller != NULL)
+				{
+					ControlledPawn = Controller->GetPawn();
+				}
+			}
+		}
+	}
+
+	return ControlledPawn;
+}
+
+
+//----------------------------------------------------------------------//
+// blueprint interface
+//----------------------------------------------------------------------//
+bool UPawnActionsComponent::PerformAction(APawn* Pawn, UPawnAction* Action, TEnumAsByte<EAIRequestPriority::Type> Priority)
+{
+	bool bSuccess = false;
+
+	ensure(Priority < EAIRequestPriority::MAX);
+
+	if (Pawn && Pawn->GetController() && Action)
+	{
+		UPawnActionsComponent* ActionComp = Pawn->GetController()->FindComponentByClass<UPawnActionsComponent>();
+		if (ActionComp)
+		{
+			ActionComp->PushAction(Action, Priority);
+			bSuccess = true;
+		}
+	}
+
+	return bSuccess;
+}
+
+//----------------------------------------------------------------------//
+// debug
+//----------------------------------------------------------------------//
+#if ENABLE_VISUAL_LOG
+void UPawnActionsComponent::DescribeSelfToVisLog(FVisLogEntry* Snapshot) const
+{
+	static const FString Category = TEXT("PawnActions");
+
+	if (IsPendingKill())
+	{
+		return;
+	}
+
+	for (int32 PriorityIndex = 0; PriorityIndex < ActionStacks.Num(); ++PriorityIndex)
+	{
+		const UPawnAction* Action = ActionStacks[PriorityIndex].GetTop();
+		if (Action == NULL)
+		{
+			continue;
+		}
+
+		FVisLogEntry::FStatusCategory StatusCategory;
+		StatusCategory.Category = Category + TEXT(": ") + GetPriorityName(PriorityIndex);
+
+		while (Action)
+		{
+			StatusCategory.Add(Action->GetName(), Action->GetStateDescription());
+			Action = Action->GetParentAction();
+		}
+
+		Snapshot->Status.Add(StatusCategory);
+	}
+}
+#endif // ENABLE_VISUAL_LOG
\ No newline at end of file
