PawnActions' non-unity compilation fixes merged over from dev branch #UE4

diff --git a/Engine/Source/Runtime/AIModule/Classes/Actions/PawnAction.h b/Engine/Source/Runtime/AIModule/Classes/Actions/PawnAction.h
--- a/Engine/Source/Runtime/AIModule/Classes/Actions/PawnAction.h
+++ b/Engine/Source/Runtime/AIModule/Classes/Actions/PawnAction.h
@@ -8,6 +8,7 @@
 class UPawnAction;
 class UPawnActionsComponent;
 
+DECLARE_LOG_CATEGORY_EXTERN(LogPawnAction, Warning, All);
 DECLARE_DELEGATE_TwoParams(FPawnActionEventDelegate, UPawnAction*, EPawnActionEventType::Type);
 
 UENUM()
diff --git a/Engine/Source/Runtime/AIModule/Classes/Actions/PawnActionsComponent.h b/Engine/Source/Runtime/AIModule/Classes/Actions/PawnActionsComponent.h
--- a/Engine/Source/Runtime/AIModule/Classes/Actions/PawnActionsComponent.h
+++ b/Engine/Source/Runtime/AIModule/Classes/Actions/PawnActionsComponent.h
@@ -4,8 +4,6 @@
 #include "Components/ActorComponent.h"
 #include "PawnActionsComponent.generated.h"
 
-DECLARE_LOG_CATEGORY_EXTERN(LogPawnAction, Warning, All);
-
 USTRUCT()
 struct FPawnActionEvent
 {
diff --git a/Engine/Source/Runtime/AIModule/Private/AIController.cpp b/Engine/Source/Runtime/AIModule/Private/AIController.cpp
--- a/Engine/Source/Runtime/AIModule/Private/AIController.cpp
+++ b/Engine/Source/Runtime/AIModule/Private/AIController.cpp
@@ -1,7 +1,6 @@
 // Copyright 1998-2014 Epic Games, Inc. All Rights Reserved.
 
 #include "AIModulePrivate.h"
-#include "VisualLog.h"
 #include "Kismet/GameplayStatics.h"
 #include "DisplayDebugHelpers.h"
 #include "BrainComponent.h"
diff --git a/Engine/Source/Runtime/AIModule/Private/AIModulePrivate.h b/Engine/Source/Runtime/AIModule/Private/AIModulePrivate.h
--- a/Engine/Source/Runtime/AIModule/Private/AIModulePrivate.h
+++ b/Engine/Source/Runtime/AIModule/Private/AIModulePrivate.h
@@ -45,6 +45,10 @@
 #include "BehaviorTree/Decorators/BTDecorator_BlackboardBase.h"
 #include "BehaviorTree/Services/BTService_BlackboardBase.h"
 
+#include "Actions/PawnAction.h"
+
+#include "VisualLog.h"
+
 DECLARE_LOG_CATEGORY_EXTERN(LogAINavigation, Warning, All);
 DECLARE_LOG_CATEGORY_EXTERN(LogBehaviorTree, Display, All);
 DECLARE_LOG_CATEGORY_EXTERN(LogCrowdFollowing, Warning, All);
\ No newline at end of file
diff --git a/Engine/Source/Runtime/AIModule/Private/Actions/PawnAction.cpp b/Engine/Source/Runtime/AIModule/Private/Actions/PawnAction.cpp
--- a/Engine/Source/Runtime/AIModule/Private/Actions/PawnAction.cpp
+++ b/Engine/Source/Runtime/AIModule/Private/Actions/PawnAction.cpp
@@ -4,6 +4,8 @@
 #include "Actions/PawnAction.h"
 #include "Actions/PawnActionsComponent.h"
 
+DEFINE_LOG_CATEGORY(LogPawnAction);
+
 namespace
 {
 	FString GetActionResultName(int32 Index)
diff --git a/Engine/Source/Runtime/AIModule/Private/Actions/PawnActionsComponent.cpp b/Engine/Source/Runtime/AIModule/Private/Actions/PawnActionsComponent.cpp
--- a/Engine/Source/Runtime/AIModule/Private/Actions/PawnActionsComponent.cpp
+++ b/Engine/Source/Runtime/AIModule/Private/Actions/PawnActionsComponent.cpp
@@ -3,8 +3,6 @@
 #include "AIModulePrivate.h"
 #include "Actions/PawnActionsComponent.h"
 
-DEFINE_LOG_CATEGORY(LogPawnAction);
-
 //----------------------------------------------------------------------//
 // helpers
 //----------------------------------------------------------------------//
diff --git a/Engine/Source/Runtime/AIModule/Private/BehaviorTree/BTCompositeNode.cpp b/Engine/Source/Runtime/AIModule/Private/BehaviorTree/BTCompositeNode.cpp
--- a/Engine/Source/Runtime/AIModule/Private/BehaviorTree/BTCompositeNode.cpp
+++ b/Engine/Source/Runtime/AIModule/Private/BehaviorTree/BTCompositeNode.cpp
@@ -2,7 +2,6 @@
 
 #include "AIModulePrivate.h"
 #include "BehaviorTree/BTCompositeNode.h"
-#include "VisualLog.h"
 
 UBTCompositeNode::UBTCompositeNode(const class FPostConstructInitializeProperties& PCIP) : Super(PCIP)
 {
diff --git a/Engine/Source/Runtime/AIModule/Private/BehaviorTree/BehaviorTreeComponent.cpp b/Engine/Source/Runtime/AIModule/Private/BehaviorTree/BehaviorTreeComponent.cpp
--- a/Engine/Source/Runtime/AIModule/Private/BehaviorTree/BehaviorTreeComponent.cpp
+++ b/Engine/Source/Runtime/AIModule/Private/BehaviorTree/BehaviorTreeComponent.cpp
@@ -8,7 +8,6 @@
 #include "BehaviorTree/BehaviorTreeManager.h"
 #include "BehaviorTree/BlackboardComponent.h"
 #include "BehaviorTree/BehaviorTreeComponent.h"
-#include "VisualLog.h"
 
 //----------------------------------------------------------------------//
 // UBehaviorTreeComponent
diff --git a/Engine/Source/Runtime/AIModule/Private/BehaviorTree/BehaviorTreeTypes.cpp b/Engine/Source/Runtime/AIModule/Private/BehaviorTree/BehaviorTreeTypes.cpp
--- a/Engine/Source/Runtime/AIModule/Private/BehaviorTree/BehaviorTreeTypes.cpp
+++ b/Engine/Source/Runtime/AIModule/Private/BehaviorTree/BehaviorTreeTypes.cpp
@@ -5,7 +5,6 @@
 #include "BehaviorTree/BTCompositeNode.h"
 #include "BehaviorTree/Tasks/BTTask_RunBehavior.h"
 #include "BehaviorTree/BehaviorTreeTypes.h"
-#include "VisualLog.h"
 
 //----------------------------------------------------------------------//
 // FBehaviorTreeInstance
diff --git a/Engine/Source/Runtime/AIModule/Private/BehaviorTree/BlackboardComponent.cpp b/Engine/Source/Runtime/AIModule/Private/BehaviorTree/BlackboardComponent.cpp
--- a/Engine/Source/Runtime/AIModule/Private/BehaviorTree/BlackboardComponent.cpp
+++ b/Engine/Source/Runtime/AIModule/Private/BehaviorTree/BlackboardComponent.cpp
@@ -4,7 +4,6 @@
 #include "EngineUtils.h"
 #include "BehaviorTree/Blackboard/BlackboardKeyAllTypes.h"
 #include "BehaviorTree/BlackboardComponent.h"
-#include "VisualLog.h"
 
 UBlackboardComponent::UBlackboardComponent(const class FPostConstructInitializeProperties& PCIP) : Super(PCIP)
 {
diff --git a/Engine/Source/Runtime/AIModule/Private/BehaviorTree/Decorators/BTDecorator_Blackboard.cpp b/Engine/Source/Runtime/AIModule/Private/BehaviorTree/Decorators/BTDecorator_Blackboard.cpp
--- a/Engine/Source/Runtime/AIModule/Private/BehaviorTree/Decorators/BTDecorator_Blackboard.cpp
+++ b/Engine/Source/Runtime/AIModule/Private/BehaviorTree/Decorators/BTDecorator_Blackboard.cpp
@@ -4,7 +4,6 @@
 #include "BehaviorTree/Decorators/BTDecorator_Blackboard.h"
 #include "BehaviorTree/BlackboardComponent.h"
 #include "BehaviorTree/Blackboard/BlackboardKeyAllTypes.h"
-#include "VisualLog.h"
 
 UBTDecorator_Blackboard::UBTDecorator_Blackboard(const class FPostConstructInitializeProperties& PCIP) : Super(PCIP)
 {
diff --git a/Engine/Source/Runtime/AIModule/Private/BehaviorTree/Tasks/BTTask_MoveTo.cpp b/Engine/Source/Runtime/AIModule/Private/BehaviorTree/Tasks/BTTask_MoveTo.cpp
--- a/Engine/Source/Runtime/AIModule/Private/BehaviorTree/Tasks/BTTask_MoveTo.cpp
+++ b/Engine/Source/Runtime/AIModule/Private/BehaviorTree/Tasks/BTTask_MoveTo.cpp
@@ -6,7 +6,6 @@
 #include "BehaviorTree/BlackboardComponent.h"
 #include "BehaviorTree/Blackboard/BlackboardKeyType_Object.h"
 #include "BehaviorTree/Blackboard/BlackboardKeyType_Vector.h"
-#include "VisualLog.h"
 
 UBTTask_MoveTo::UBTTask_MoveTo(const class FPostConstructInitializeProperties& PCIP) 
 	: Super(PCIP)
diff --git a/Engine/Source/Runtime/AIModule/Private/BehaviorTree/Tasks/BTTask_RunEQSQuery.cpp b/Engine/Source/Runtime/AIModule/Private/BehaviorTree/Tasks/BTTask_RunEQSQuery.cpp
--- a/Engine/Source/Runtime/AIModule/Private/BehaviorTree/Tasks/BTTask_RunEQSQuery.cpp
+++ b/Engine/Source/Runtime/AIModule/Private/BehaviorTree/Tasks/BTTask_RunEQSQuery.cpp
@@ -5,8 +5,6 @@
 #include "EnvironmentQuery/Items/EnvQueryItemType.h"
 #include "EnvironmentQuery/EnvQueryManager.h"
 #include "EnvironmentQuery/EnvQuery.h"
-#include "VisualLog.h"
-
 
 UBTTask_RunEQSQuery::UBTTask_RunEQSQuery(const class FPostConstructInitializeProperties& PCIP) : Super(PCIP)
 {
diff --git a/Engine/Source/Runtime/AIModule/Private/BrainComponent.cpp b/Engine/Source/Runtime/AIModule/Private/BrainComponent.cpp
--- a/Engine/Source/Runtime/AIModule/Private/BrainComponent.cpp
+++ b/Engine/Source/Runtime/AIModule/Private/BrainComponent.cpp
@@ -1,7 +1,6 @@
 // Copyright 1998-2014 Epic Games, Inc. All Rights Reserved.
 
 #include "AIModulePrivate.h"
-#include "VisualLog.h"
 #include "BrainComponent.h"
 #include "BehaviorTree/BlackboardComponent.h"
 
diff --git a/Engine/Source/Runtime/AIModule/Private/EnvironmentQuery/EQSRenderingComponent.cpp b/Engine/Source/Runtime/AIModule/Private/EnvironmentQuery/EQSRenderingComponent.cpp
--- a/Engine/Source/Runtime/AIModule/Private/EnvironmentQuery/EQSRenderingComponent.cpp
+++ b/Engine/Source/Runtime/AIModule/Private/EnvironmentQuery/EQSRenderingComponent.cpp
@@ -5,7 +5,6 @@
 #include "EnvironmentQuery/EQSRenderingComponent.h"
 #include "EnvironmentQuery/EQSQueryResultSourceInterface.h"
 #include "EnvironmentQuery/Items/EnvQueryItemType_VectorBase.h"
-#include "VisualLog.h"
 
 static const int32 EQSMaxItemsDrawn = 10000;
 
diff --git a/Engine/Source/Runtime/AIModule/Private/EnvironmentQuery/Generators/EnvQueryGenerator_OnCircle.cpp b/Engine/Source/Runtime/AIModule/Private/EnvironmentQuery/Generators/EnvQueryGenerator_OnCircle.cpp
--- a/Engine/Source/Runtime/AIModule/Private/EnvironmentQuery/Generators/EnvQueryGenerator_OnCircle.cpp
+++ b/Engine/Source/Runtime/AIModule/Private/EnvironmentQuery/Generators/EnvQueryGenerator_OnCircle.cpp
@@ -5,7 +5,7 @@
 #include "EnvironmentQuery/Items/EnvQueryItemType_Point.h"
 #include "EnvironmentQuery/Generators/EnvQueryGenerator_OnCircle.h"
 #include "AI/Navigation/RecastNavMesh.h"
-#include "VisualLog.h"
+
 
 #define LOCTEXT_NAMESPACE "EnvQueryGenerator"
 
diff --git a/Engine/Source/Runtime/AIModule/Private/Navigation/CrowdFollowingComponent.cpp b/Engine/Source/Runtime/AIModule/Private/Navigation/CrowdFollowingComponent.cpp
--- a/Engine/Source/Runtime/AIModule/Private/Navigation/CrowdFollowingComponent.cpp
+++ b/Engine/Source/Runtime/AIModule/Private/Navigation/CrowdFollowingComponent.cpp
@@ -6,7 +6,7 @@
 #include "Navigation/NavigationComponent.h"
 #include "AI/Navigation/NavLinkCustomInterface.h"
 #include "GameFramework/CharacterMovementComponent.h"
-#include "VisualLog.h"
+
 
 DEFINE_LOG_CATEGORY(LogCrowdFollowing);
 
diff --git a/Engine/Source/Runtime/AIModule/Private/Navigation/CrowdManager.cpp b/Engine/Source/Runtime/AIModule/Private/Navigation/CrowdManager.cpp
--- a/Engine/Source/Runtime/AIModule/Private/Navigation/CrowdManager.cpp
+++ b/Engine/Source/Runtime/AIModule/Private/Navigation/CrowdManager.cpp
@@ -11,7 +11,6 @@
 #include "Navigation/CrowdManager.h"
 #include "Navigation/CrowdFollowingComponent.h"
 #include "Navigation/CrowdAgentInterface.h"
-#include "VisualLog.h"
 
 #include "DrawDebugHelpers.h"
 
diff --git a/Engine/Source/Runtime/AIModule/Private/Navigation/NavigationComponent.cpp b/Engine/Source/Runtime/AIModule/Private/Navigation/NavigationComponent.cpp
--- a/Engine/Source/Runtime/AIModule/Private/Navigation/NavigationComponent.cpp
+++ b/Engine/Source/Runtime/AIModule/Private/Navigation/NavigationComponent.cpp
@@ -1,7 +1,6 @@
 // Copyright 1998-2014 Epic Games, Inc. All Rights Reserved.
 
 #include "AIModulePrivate.h"
-#include "VisualLog.h"
 #include "AI/Navigation/NavLinkCustomComponent.h"
 #include "DrawDebugHelpers.h"
 #include "Navigation/NavigationComponent.h"
diff --git a/Engine/Source/Runtime/AIModule/Private/Navigation/PathFollowingComponent.cpp b/Engine/Source/Runtime/AIModule/Private/Navigation/PathFollowingComponent.cpp
--- a/Engine/Source/Runtime/AIModule/Private/Navigation/PathFollowingComponent.cpp
+++ b/Engine/Source/Runtime/AIModule/Private/Navigation/PathFollowingComponent.cpp
@@ -11,7 +11,7 @@
 #include "GameFramework/Character.h"
 #include "Engine/Canvas.h"
 #include "TimerManager.h"
-#include "VisualLog.h"
+
 #include "Navigation/PathFollowingComponent.h"
 
 #define USE_PHYSIC_FOR_VISIBILITY_TESTS 1 // Physic will be used for visibility tests if set or only raycasts on navmesh if not
