behavior tree restart on value/result change in blackboard decorator

diff --git a/Engine/Source/Editor/BehaviorTreeEditor/Private/DetailCustomizations/BlackboardDecoratorDetails.cpp b/Engine/Source/Editor/BehaviorTreeEditor/Private/DetailCustomizations/BlackboardDecoratorDetails.cpp
--- a/Engine/Source/Editor/BehaviorTreeEditor/Private/DetailCustomizations/BlackboardDecoratorDetails.cpp
+++ b/Engine/Source/Editor/BehaviorTreeEditor/Private/DetailCustomizations/BlackboardDecoratorDetails.cpp
@@ -27,6 +27,11 @@ void FBlackboardDecoratorDetails::CustomizeDetails( IDetailLayoutBuilder& Detail
 	const bool bIsEnabled = CachedBlackboardAsset.IsValid();
 	TAttribute<bool> PropertyEditCheck = TAttribute<bool>::Create(TAttribute<bool>::FGetter::CreateSP(this, &FBehaviorDecoratorDetails::IsEditingEnabled));
 
+	IDetailCategoryBuilder& FlowCategory = DetailLayout.EditCategory("FlowControl");
+	NotifyObserverProperty = DetailLayout.GetProperty(GET_MEMBER_NAME_CHECKED(UBTDecorator_Blackboard, NotifyObserver));
+	IDetailPropertyRow& AbortRow = FlowCategory.AddProperty(NotifyObserverProperty);
+	AbortRow.IsEnabled(PropertyEditCheck);
+
 	IDetailCategoryBuilder& BBCategory = DetailLayout.EditCategory( "Blackboard" );
 	IDetailPropertyRow& KeySelectorRow = BBCategory.AddProperty(DetailLayout.GetProperty(GET_MEMBER_NAME_CHECKED(UBTDecorator_Blackboard, BlackboardKey)));
 	KeySelectorRow.IsEnabled(bIsEnabled);
diff --git a/Engine/Source/Editor/BehaviorTreeEditor/Private/DetailCustomizations/BlackboardDecoratorDetails.h b/Engine/Source/Editor/BehaviorTreeEditor/Private/DetailCustomizations/BlackboardDecoratorDetails.h
--- a/Engine/Source/Editor/BehaviorTreeEditor/Private/DetailCustomizations/BlackboardDecoratorDetails.h
+++ b/Engine/Source/Editor/BehaviorTreeEditor/Private/DetailCustomizations/BlackboardDecoratorDetails.h
@@ -35,6 +35,7 @@ private:
 
 	TSharedPtr<IPropertyHandle> IntValueProperty;
 	TSharedPtr<IPropertyHandle> KeyIDProperty;
+	TSharedPtr<IPropertyHandle> NotifyObserverProperty;
 	
 	/** cached type of property selected by KeyName */
 	TSubclassOf<class UBlackboardKeyType> CachedKeyType;
diff --git a/Engine/Source/Runtime/AIModule/Classes/BehaviorTree/Decorators/BTDecorator_Blackboard.h b/Engine/Source/Runtime/AIModule/Classes/BehaviorTree/Decorators/BTDecorator_Blackboard.h
--- a/Engine/Source/Runtime/AIModule/Classes/BehaviorTree/Decorators/BTDecorator_Blackboard.h
+++ b/Engine/Source/Runtime/AIModule/Classes/BehaviorTree/Decorators/BTDecorator_Blackboard.h
@@ -9,6 +9,16 @@
  *  Decorator for accessing blackboard values
  */
 
+UENUM()
+namespace EBTBlackboardRestart
+{
+	enum Type
+	{
+		ValueChange		UMETA(DisplayName="On Value Change", ToolTip="Restart on every change of observed blackboard value"),
+		ResultChange	UMETA(DisplayName="On Result Change", ToolTip="Restart only when result of evaluated condition is changed"),
+	};
+}
+
 UCLASS(HideCategories=(Condition))
 class AIMODULE_API UBTDecorator_Blackboard : public UBTDecorator_BlackboardBase
 {
@@ -41,6 +51,10 @@ protected:
 	UPROPERTY()
 	uint8 OperationType;
 
+	/** when observer can try to request abort? */
+	UPROPERTY(Category=FlowControl, EditAnywhere)
+	TEnumAsByte<EBTBlackboardRestart::Type> NotifyObserver;
+
 #if WITH_EDITORONLY_DATA
 
 	UPROPERTY(Category=Blackboard, EditAnywhere, meta=(DisplayName="Key Query"))
diff --git a/Engine/Source/Runtime/AIModule/Private/BehaviorTree/Decorators/BTDecorator_Blackboard.cpp b/Engine/Source/Runtime/AIModule/Private/BehaviorTree/Decorators/BTDecorator_Blackboard.cpp
--- a/Engine/Source/Runtime/AIModule/Private/BehaviorTree/Decorators/BTDecorator_Blackboard.cpp
+++ b/Engine/Source/Runtime/AIModule/Private/BehaviorTree/Decorators/BTDecorator_Blackboard.cpp
@@ -8,6 +8,7 @@
 UBTDecorator_Blackboard::UBTDecorator_Blackboard(const class FPostConstructInitializeProperties& PCIP) : Super(PCIP)
 {
 	NodeName = "Blackboard";
+	NotifyObserver = EBTBlackboardRestart::ResultChange;
 }
 
 bool UBTDecorator_Blackboard::CalculateRawConditionValue(class UBehaviorTreeComponent* OwnerComp, uint8* NodeMemory) const
@@ -40,20 +41,40 @@ bool UBTDecorator_Blackboard::EvaluateOnBlackboard(const UBlackboardComponent* B
 void UBTDecorator_Blackboard::OnBlackboardChange(const class UBlackboardComponent* Blackboard, FBlackboard::FKey ChangedKeyID)
 {
 	UBehaviorTreeComponent* BehaviorComp = Blackboard ? (UBehaviorTreeComponent*)Blackboard->GetBrainComponent() : NULL;
-	if (BlackboardKey.GetSelectedKeyID() == ChangedKeyID && BehaviorComp)
+	if (BlackboardKey.GetSelectedKeyID() == ChangedKeyID &&
+		GetFlowAbortMode() != EBTFlowAbortMode::None &&
+		BehaviorComp)
 	{
-		const bool bIsExecutingBranch = BehaviorComp->IsExecutingBranch(this, GetChildIndex());
-		const bool bPass = EvaluateOnBlackboard(Blackboard);
+		if (NotifyObserver == EBTBlackboardRestart::ResultChange)
+		{
+			const bool bIsExecutingBranch = BehaviorComp->IsExecutingBranch(this, GetChildIndex());
+			const bool bPass = EvaluateOnBlackboard(Blackboard);
 
-		UE_VLOG(BehaviorComp->GetOwner(), LogBehaviorTree, Verbose, TEXT("%s, OnBlackboardChange[%s] pass:%d executing:%d => %s"),
-			*UBehaviorTreeTypes::DescribeNodeHelper(this),
-			*Blackboard->GetKeyName(ChangedKeyID).ToString(), bPass, bIsExecutingBranch,
-			(bIsExecutingBranch && !bPass) || (!bIsExecutingBranch && bPass) ? TEXT("restart") : TEXT("skip"));
+			UE_VLOG(BehaviorComp->GetOwner(), LogBehaviorTree, Verbose, TEXT("%s, OnBlackboardChange[%s] pass:%d executing:%d => %s"),
+				*UBehaviorTreeTypes::DescribeNodeHelper(this),
+				*Blackboard->GetKeyName(ChangedKeyID).ToString(), bPass, bIsExecutingBranch,
+				(bIsExecutingBranch && !bPass) || (!bIsExecutingBranch && bPass) ? TEXT("restart") : TEXT("skip"));
 
-		if ((bIsExecutingBranch && !bPass) || 
-			(!bIsExecutingBranch && bPass))
+			if ((bIsExecutingBranch && !bPass) ||
+				(!bIsExecutingBranch && bPass))
+			{
+				BehaviorComp->RequestExecution(this);
+			}
+		}
+		else
 		{
-			BehaviorComp->RequestExecution(this);
+			UE_VLOG(BehaviorComp->GetOwner(), LogBehaviorTree, Verbose, TEXT("%s, OnBlackboardChange[%s] => restart"),
+				*UBehaviorTreeTypes::DescribeNodeHelper(this),
+				*Blackboard->GetKeyName(ChangedKeyID).ToString());
+
+			// force result Aborted to restart from this decorator
+			// can't use helper function
+
+			const int32 InstanceIdx = BehaviorComp->FindInstanceContainingNode((UBTNode*)GetParentNode());
+			if (InstanceIdx != INDEX_NONE)
+			{
+				BehaviorComp->RequestExecution(GetParentNode(), InstanceIdx, this, GetChildIndex(), EBTNodeResult::Aborted);
+			}
 		}
 	}
 }
