fixed DeltaTime in behavior tree's auxiliary node tick event

diff --git a/Engine/Source/Runtime/AIModule/Classes/BehaviorTree/BTAuxiliaryNode.h b/Engine/Source/Runtime/AIModule/Classes/BehaviorTree/BTAuxiliaryNode.h
--- a/Engine/Source/Runtime/AIModule/Classes/BehaviorTree/BTAuxiliaryNode.h
+++ b/Engine/Source/Runtime/AIModule/Classes/BehaviorTree/BTAuxiliaryNode.h
@@ -8,6 +8,7 @@
 struct FBTAuxiliaryMemory : public FBTInstancedNodeMemory
 {
 	float NextTickRemainingTime;
+	float AccumulatedDeltaTime;
 };
 
 /** 
diff --git a/Engine/Source/Runtime/AIModule/Private/BehaviorTree/BTAuxiliaryNode.cpp b/Engine/Source/Runtime/AIModule/Private/BehaviorTree/BTAuxiliaryNode.cpp
--- a/Engine/Source/Runtime/AIModule/Private/BehaviorTree/BTAuxiliaryNode.cpp
+++ b/Engine/Source/Runtime/AIModule/Private/BehaviorTree/BTAuxiliaryNode.cpp
@@ -39,20 +39,27 @@ void UBTAuxiliaryNode::WrappedTickNode(class UBehaviorTreeComponent* OwnerComp,
 {
 	if (bNotifyTick)
 	{
+		float UseDeltaTime = DeltaSeconds;
+
 		if (bTickIntervals)
 		{
 			FBTAuxiliaryMemory* AuxMemory = GetSpecialNodeMemory<FBTAuxiliaryMemory>(NodeMemory);
 			AuxMemory->NextTickRemainingTime -= DeltaSeconds;
+			AuxMemory->AccumulatedDeltaTime += DeltaSeconds;
+			
 			if (AuxMemory->NextTickRemainingTime > 0.0f)
 			{
 				return;
-			}	
+			}
+
+			UseDeltaTime = AuxMemory->AccumulatedDeltaTime;
+			AuxMemory->AccumulatedDeltaTime = 0.0f;
 		}
 
 		const UBTNode* NodeOb = bCreateNodeInstance ? GetNodeInstance(OwnerComp, NodeMemory) : this;
 		if (NodeOb)
 		{
-			((UBTAuxiliaryNode*)NodeOb)->TickNode(OwnerComp, NodeMemory, DeltaSeconds);
+			((UBTAuxiliaryNode*)NodeOb)->TickNode(OwnerComp, NodeMemory, UseDeltaTime);
 		}
 	}
 }
diff --git a/Engine/Source/Runtime/AIModule/Private/BehaviorTree/Decorators/BTDecorator_TimeLimit.cpp b/Engine/Source/Runtime/AIModule/Private/BehaviorTree/Decorators/BTDecorator_TimeLimit.cpp
--- a/Engine/Source/Runtime/AIModule/Private/BehaviorTree/Decorators/BTDecorator_TimeLimit.cpp
+++ b/Engine/Source/Runtime/AIModule/Private/BehaviorTree/Decorators/BTDecorator_TimeLimit.cpp
@@ -21,6 +21,7 @@ void UBTDecorator_TimeLimit::OnBecomeRelevant(UBehaviorTreeComponent* OwnerComp,
 {
 	FBTAuxiliaryMemory* DecoratorMemory = GetSpecialNodeMemory<FBTAuxiliaryMemory>(NodeMemory);
 	DecoratorMemory->NextTickRemainingTime = TimeLimit;
+	DecoratorMemory->AccumulatedDeltaTime = 0.0f;
 }
 
 void UBTDecorator_TimeLimit::TickNode(UBehaviorTreeComponent* OwnerComp, uint8* NodeMemory, float DeltaSeconds)
