EQS support for BP-implemented EQS generators #UE4

Also added a basic EQS genetator to retrieve all actors of specified class.

diff --git a/Engine/Source/Editor/EnvironmentQueryEditor/Private/EdGraphSchema_EnvironmentQuery.cpp b/Engine/Source/Editor/EnvironmentQueryEditor/Private/EdGraphSchema_EnvironmentQuery.cpp
--- a/Engine/Source/Editor/EnvironmentQueryEditor/Private/EdGraphSchema_EnvironmentQuery.cpp
+++ b/Engine/Source/Editor/EnvironmentQueryEditor/Private/EdGraphSchema_EnvironmentQuery.cpp
@@ -7,7 +7,7 @@
 #include "EnvironmentQuery/EnvQueryGenerator.h"
 #include "EnvironmentQuery/Generators/EnvQueryGenerator_Composite.h"
 #include "EnvironmentQuery/EnvQueryTest.h"
-
+#include "Kismet2/KismetEditorUtilities.h"
 
 #define LOCTEXT_NAMESPACE "EnvironmentQuerySchema"
 #define SNAP_GRID (16) // @todo ensure this is the same as SNodePanel::GetSnapGridSize()
@@ -175,7 +175,8 @@ static void GatherEQSGenerators(TArray<UClass*>& Classes)
 		{
 			UClass* TestClass = *It;
 			if (TestClass->HasAnyClassFlags(CLASS_Abstract) ||
-				!TestClass->HasAnyClassFlags(CLASS_Native) ||
+				// ignore temporary BP classes
+				(!TestClass->HasAnyClassFlags(CLASS_Native) && (FKismetEditorUtilities::IsClassABlueprintSkeleton(TestClass) || TestClass->HasAnyClassFlags(CLASS_NewerVersionExists))) || 
 				!TestClass->IsChildOf(UEnvQueryGenerator::StaticClass()) ||
 				TestClass == UEnvQueryGenerator_Composite::StaticClass())
 			{
diff --git a/Engine/Source/Runtime/AIModule/Classes/EnvironmentQuery/Generators/EnvQueryGenerator_ActorsOfClass.h b/Engine/Source/Runtime/AIModule/Classes/EnvironmentQuery/Generators/EnvQueryGenerator_ActorsOfClass.h
new file mode 100644
--- /dev/null
+++ b/Engine/Source/Runtime/AIModule/Classes/EnvironmentQuery/Generators/EnvQueryGenerator_ActorsOfClass.h
@@ -0,0 +1,29 @@
+// Copyright 1998-2014 Epic Games, Inc. All Rights Reserved.
+
+#pragma once
+#include "EnvironmentQuery/EnvQueryGenerator.h"
+#include "EnvQueryGenerator_ActorsOfClass.generated.h"
+
+class AActor;
+
+UCLASS()	
+class UEnvQueryGenerator_ActorsOfClass : public UEnvQueryGenerator
+{
+	GENERATED_UCLASS_BODY()
+
+	/** max distance of path between point and context */
+	UPROPERTY(EditDefaultsOnly, Category=Generator)
+	FEnvFloatParam Radius;
+
+	UPROPERTY(EditDefaultsOnly, Category=Generator)
+	TSubclassOf<AActor> SearchedActorClass;
+
+	/** context */
+	UPROPERTY(EditAnywhere, Category=Generator)
+	TSubclassOf<class UEnvQueryContext> SearchCenter;
+
+	void GenerateItems(struct FEnvQueryInstance& QueryInstance); 
+
+	virtual FText GetDescriptionTitle() const override;
+	virtual FText GetDescriptionDetails() const override;
+};
diff --git a/Engine/Source/Runtime/AIModule/Classes/EnvironmentQuery/Generators/EnvQueryGenerator_BlueprintBase.h b/Engine/Source/Runtime/AIModule/Classes/EnvironmentQuery/Generators/EnvQueryGenerator_BlueprintBase.h
new file mode 100644
--- /dev/null
+++ b/Engine/Source/Runtime/AIModule/Classes/EnvironmentQuery/Generators/EnvQueryGenerator_BlueprintBase.h
@@ -0,0 +1,54 @@
+// Copyright 1998-2014 Epic Games, Inc. All Rights Reserved.
+
+#pragma once
+#include "EnvironmentQuery/EnvQueryGenerator.h"
+#include "EnvQueryGenerator_BlueprintBase.generated.h"
+
+struct FEnvQueryInstance;
+class AActor;
+class UEnvQueryContext;
+class UEnvQueryItemType;
+
+UCLASS(Abstract, Blueprintable)
+class AIMODULE_API UEnvQueryGenerator_BlueprintBase : public UEnvQueryGenerator
+{
+	GENERATED_UCLASS_BODY()
+
+	/** max distance of between point and context */
+	UPROPERTY(EditDefaultsOnly, Category = Generator)
+	FEnvFloatParam Radius; 
+
+	/** A short description of what test does, like "Generate pawn named Joe" */
+	UPROPERTY(EditAnywhere, Category = Generator)
+	FText GeneratorsActionDescription;
+
+	/** context */
+	UPROPERTY(EditAnywhere, Category = Generator)
+	TSubclassOf<UEnvQueryContext> Context;
+
+	UPROPERTY(EditDefaultsOnly, Category = Generator)
+	TSubclassOf<UEnvQueryItemType> GeneratedItemType;
+
+	virtual void PostInitProperties() override;
+
+	UFUNCTION(BlueprintImplementableEvent, Category = Generator)
+	virtual void DoItemGeneration(const TArray<FVector>& ContextLocations);
+
+	void GenerateItems(FEnvQueryInstance& QueryInstance); 
+
+	virtual FText GetDescriptionTitle() const override;
+	virtual FText GetDescriptionDetails() const override;
+	
+	UFUNCTION(BlueprintCallable, Category = "EQS")
+	void AddGeneratedVector(FVector GeneratedVector);
+
+	UFUNCTION(BlueprintCallable, Category = "EQS")
+	void AddGeneratedActor(AActor* GeneratedActor);
+
+	UFUNCTION(BlueprintCallable, Category = "EQS")
+	UObject* GetQuerier() const;
+
+private:
+	/** this is valid and set only within GenerateItems call */
+	FEnvQueryInstance* CachedQueryInstance;
+};
diff --git a/Engine/Source/Runtime/AIModule/Private/EnvironmentQuery/Generators/EnvQueryGenerator_ActorsOfClass.cpp b/Engine/Source/Runtime/AIModule/Private/EnvironmentQuery/Generators/EnvQueryGenerator_ActorsOfClass.cpp
new file mode 100644
--- /dev/null
+++ b/Engine/Source/Runtime/AIModule/Private/EnvironmentQuery/Generators/EnvQueryGenerator_ActorsOfClass.cpp
@@ -0,0 +1,72 @@
+// Copyright 1998-2014 Epic Games, Inc. All Rights Reserved.
+
+#include "AIModulePrivate.h"
+#include "EnvironmentQuery/Contexts/EnvQueryContext_Querier.h"
+#include "EnvironmentQuery/Items/EnvQueryItemType_Point.h"
+#include "EnvironmentQuery/Generators/EnvQueryGenerator_ActorsOfClass.h"
+#include "AI/Navigation/RecastNavMesh.h"
+
+#define LOCTEXT_NAMESPACE "EnvQueryGenerator"
+
+UEnvQueryGenerator_ActorsOfClass::UEnvQueryGenerator_ActorsOfClass(const class FPostConstructInitializeProperties& PCIP) : Super(PCIP)
+{
+	GenerateDelegate.BindUObject(this, &UEnvQueryGenerator_ActorsOfClass::GenerateItems);
+
+	SearchCenter = UEnvQueryContext_Querier::StaticClass();
+	ItemType = UEnvQueryItemType_Actor::StaticClass();
+	Radius.Value = 500.0f;
+	SearchedActorClass = AActor::StaticClass();
+}
+
+void UEnvQueryGenerator_ActorsOfClass::GenerateItems(struct FEnvQueryInstance& QueryInstance)
+{
+	float RadiusValue = 0.0f;
+	float DensityValue = 0.0f;
+
+	UWorld* World = GEngine->GetWorldFromContextObject(QueryInstance.Owner.Get());
+	// @TODO add some logging here
+	if (World == NULL || QueryInstance.GetParamValue(Radius, RadiusValue, TEXT("Radius")) == false
+		|| SearchedActorClass == NULL)
+	{
+		return;
+	}
+
+	const float RadiusSq = FMath::Square(RadiusValue);
+
+	TArray<FVector> ContextLocations;
+	QueryInstance.PrepareContext(SearchCenter, ContextLocations);
+
+	for (TActorIterator<AActor> ItActor = TActorIterator<AActor>(World, SearchedActorClass); ItActor; ++ItActor)
+	{
+		for (int32 ContextIndex = 0; ContextIndex < ContextLocations.Num(); ++ContextIndex)
+		{
+			if (FVector::DistSquared(ContextLocations[ContextIndex], ItActor->GetActorLocation()) < RadiusSq)
+			{
+				QueryInstance.AddItemData<UEnvQueryItemType_Actor>(*ItActor);
+				break;
+			}
+		}
+	}
+}
+
+FText UEnvQueryGenerator_ActorsOfClass::GetDescriptionTitle() const
+{
+	FFormatNamedArguments Args;
+	Args.Add(TEXT("DescriptionTitle"), Super::GetDescriptionTitle());
+	Args.Add(TEXT("ActorsClass"), FText::FromString(GetNameSafe(SearchedActorClass)));
+	Args.Add(TEXT("DescribeContext"), UEnvQueryTypes::DescribeContext(SearchCenter));
+
+	return FText::Format(LOCTEXT("DescriptionGenerateActorsAroundContext", "{DescriptionTitle}: generate set of actors of class {ActorsClass} around {DescribeContext}"), Args);
+};
+
+FText UEnvQueryGenerator_ActorsOfClass::GetDescriptionDetails() const
+{
+	FFormatNamedArguments Args;
+	Args.Add(TEXT("Radius"), FText::FromString(UEnvQueryTypes::DescribeFloatParam(Radius)));
+
+	FText Desc = FText::Format(LOCTEXT("ActorsOfClassDescription", "radius: {Radius}"), Args);
+
+	return Desc;
+}
+
+#undef LOCTEXT_NAMESPACE
\ No newline at end of file
diff --git a/Engine/Source/Runtime/AIModule/Private/EnvironmentQuery/Generators/EnvQueryGenerator_BlueprintBase.cpp b/Engine/Source/Runtime/AIModule/Private/EnvironmentQuery/Generators/EnvQueryGenerator_BlueprintBase.cpp
new file mode 100644
--- /dev/null
+++ b/Engine/Source/Runtime/AIModule/Private/EnvironmentQuery/Generators/EnvQueryGenerator_BlueprintBase.cpp
@@ -0,0 +1,80 @@
+// Copyright 1998-2014 Epic Games, Inc. All Rights Reserved.
+
+#include "AIModulePrivate.h"
+#include "EnvironmentQuery/Contexts/EnvQueryContext_Querier.h"
+#include "EnvironmentQuery/Items/EnvQueryItemType_Point.h"
+#include "EnvironmentQuery/Items/EnvQueryItemType_Point.h"
+#include "EnvironmentQuery/Generators/EnvQueryGenerator_BlueprintBase.h"
+#include "AI/Navigation/RecastNavMesh.h"
+#include "VisualLog.h"
+
+#define LOCTEXT_NAMESPACE "EnvQueryGenerator"
+
+UEnvQueryGenerator_BlueprintBase::UEnvQueryGenerator_BlueprintBase(const class FPostConstructInitializeProperties& PCIP)
+	: Super(PCIP)
+{
+	GenerateDelegate.BindUObject(this, &UEnvQueryGenerator_BlueprintBase::GenerateItems);
+
+	Context = UEnvQueryContext_Querier::StaticClass();
+	ItemType = UEnvQueryItemType_Point::StaticClass();
+	Radius.Value = 1000.0f;
+}
+
+void UEnvQueryGenerator_BlueprintBase::PostInitProperties()
+{
+	Super::PostInitProperties();
+	ItemType = GeneratedItemType;
+}
+
+void UEnvQueryGenerator_BlueprintBase::GenerateItems(FEnvQueryInstance& QueryInstance)
+{
+	float RadiusValue = 0.f;
+	if (QueryInstance.GetParamValue(Radius, RadiusValue, TEXT("Radius")) == false)
+	{
+		return;
+	}
+
+	TArray<FVector> ContextLocations;
+	QueryInstance.PrepareContext(Context, ContextLocations);
+
+	CachedQueryInstance = &QueryInstance;
+	DoItemGeneration(ContextLocations);
+	CachedQueryInstance = NULL;
+}
+
+void UEnvQueryGenerator_BlueprintBase::AddGeneratedVector(FVector Vector)
+{
+	check(CachedQueryInstance);
+	ensure(ItemType->GetClass()->IsChildOf(UEnvQueryItemType_VectorBase::StaticClass()));
+	CachedQueryInstance->AddItemData<UEnvQueryItemType_Point>(Vector);
+}
+
+void UEnvQueryGenerator_BlueprintBase::AddGeneratedActor(AActor* Actor)
+{
+	check(CachedQueryInstance);
+	ensure(ItemType->GetClass()->IsChildOf(UEnvQueryItemType_ActorBase::StaticClass()));
+	CachedQueryInstance->AddItemData<UEnvQueryItemType_Actor>(Actor);
+}
+
+UObject* UEnvQueryGenerator_BlueprintBase::GetQuerier() const
+{
+	check(CachedQueryInstance);
+	return CachedQueryInstance->Owner.Get();
+}
+
+FText UEnvQueryGenerator_BlueprintBase::GetDescriptionTitle() const
+{
+	FFormatNamedArguments Args;
+	Args.Add(TEXT("DescriptionTitle"), Super::GetDescriptionTitle());
+	Args.Add(TEXT("DescribeGeneratorAction"), GeneratorsActionDescription);
+	Args.Add(TEXT("DescribeContext"), UEnvQueryTypes::DescribeContext(Context));
+
+	return FText::Format(LOCTEXT("DescriptionBlueprintImplementedGenerator", "{DescriptionTitle}: {DescribeGeneratorAction} around {DescribeContext}"), Args);
+}
+
+FText UEnvQueryGenerator_BlueprintBase::GetDescriptionDetails() const
+{
+	return FText::FromString(TEXT("None"));
+}
+
+#undef LOCTEXT_NAMESPACE
\ No newline at end of file
