Attempting to fix Mac CIS blindly.

diff --git a/Engine/Source/Runtime/AIModule/Private/EnvironmentQuery/Generators/EnvQueryGenerator_ActorsOfClass.cpp b/Engine/Source/Runtime/AIModule/Private/EnvironmentQuery/Generators/EnvQueryGenerator_ActorsOfClass.cpp
--- a/Engine/Source/Runtime/AIModule/Private/EnvironmentQuery/Generators/EnvQueryGenerator_ActorsOfClass.cpp
+++ b/Engine/Source/Runtime/AIModule/Private/EnvironmentQuery/Generators/EnvQueryGenerator_ActorsOfClass.cpp
@@ -3,6 +3,7 @@
 #include "AIModulePrivate.h"
 #include "EnvironmentQuery/Contexts/EnvQueryContext_Querier.h"
 #include "EnvironmentQuery/Items/EnvQueryItemType_Point.h"
+#include "EnvironmentQuery/Items/EnvQueryItemType_Actor.h"
 #include "EnvironmentQuery/Generators/EnvQueryGenerator_ActorsOfClass.h"
 #include "AI/Navigation/RecastNavMesh.h"
 
