AIModule CIS fixes

diff --git a/Engine/Source/Runtime/AIModule/Private/AIModulePrivate.h b/Engine/Source/Runtime/AIModule/Private/AIModulePrivate.h
--- a/Engine/Source/Runtime/AIModule/Private/AIModulePrivate.h
+++ b/Engine/Source/Runtime/AIModule/Private/AIModulePrivate.h
@@ -27,6 +27,7 @@
 #include "Components/CapsuleComponent.h"
 #include "Engine/Engine.h"
 #include "TimerManager.h"					// Game play timers
+#include "EngineUtils.h"
 
 //----------------------------------------------------------------------//
 // AIModule headers
diff --git a/Engine/Source/Runtime/AIModule/Private/EnvironmentQuery/Generators/EnvQueryGenerator_BlueprintBase.cpp b/Engine/Source/Runtime/AIModule/Private/EnvironmentQuery/Generators/EnvQueryGenerator_BlueprintBase.cpp
--- a/Engine/Source/Runtime/AIModule/Private/EnvironmentQuery/Generators/EnvQueryGenerator_BlueprintBase.cpp
+++ b/Engine/Source/Runtime/AIModule/Private/EnvironmentQuery/Generators/EnvQueryGenerator_BlueprintBase.cpp
@@ -5,6 +5,7 @@
 #include "EnvironmentQuery/Items/EnvQueryItemType_Point.h"
 #include "EnvironmentQuery/Items/EnvQueryItemType_Point.h"
 #include "EnvironmentQuery/Generators/EnvQueryGenerator_BlueprintBase.h"
+#include "EnvironmentQuery/Items/EnvQueryItemType_Actor.h"
 #include "AI/Navigation/RecastNavMesh.h"
 #include "VisualLog.h"
 
