fixed crash in behavior tree trying to process requests with unregistered component

diff --git a/Engine/Source/Runtime/AIModule/Private/BehaviorTree/BehaviorTreeComponent.cpp b/Engine/Source/Runtime/AIModule/Private/BehaviorTree/BehaviorTreeComponent.cpp
--- a/Engine/Source/Runtime/AIModule/Private/BehaviorTree/BehaviorTreeComponent.cpp
+++ b/Engine/Source/Runtime/AIModule/Private/BehaviorTree/BehaviorTreeComponent.cpp
@@ -140,7 +140,12 @@ void UBehaviorTreeComponent::StopTree()
 	TaskMessageObservers.Reset();
 	ExecutionRequest = FBTNodeExecutionInfo();
 	ActiveInstanceIdx = 0;
-	GetOwner()->GetWorldTimerManager().ClearTimer(this, &UBehaviorTreeComponent::ProcessExecutionRequest);
+
+	if (IsRegistered())
+	{
+		GetOwner()->GetWorldTimerManager().ClearTimer(this, &UBehaviorTreeComponent::ProcessExecutionRequest);
+	}
+
 	// make sure to allow new execution requests since we just removed last request timer
 	bRequestedFlowUpdate = false;
 }
@@ -467,7 +472,7 @@ static void FindCommonParent(const TArray<FBehaviorTreeInstance>& Instances, con
 
 void UBehaviorTreeComponent::ScheduleExecutionUpdate()
 {
-	if (!bRequestedFlowUpdate)
+	if (!bRequestedFlowUpdate && IsRegistered())
 	{
 		bRequestedFlowUpdate = true;
 
@@ -752,6 +757,12 @@ void UBehaviorTreeComponent::TickComponent(float DeltaTime, enum ELevelTick Tick
 void UBehaviorTreeComponent::ProcessExecutionRequest()
 {
 	bRequestedFlowUpdate = false;
+	if (!IsRegistered())
+	{
+		// it shouldn't be called, component is no longer valid
+		return;
+	}
+
 	if (bIsPaused)
 	{
 		UE_VLOG(GetOwner(), LogBehaviorTree, Verbose, TEXT("Ignoring ProcessExecutionRequest call due to BTComponent still being paused"));
