Enum Blackboard key types are now selectable again.

Removed EditInline property that was preventing the Enum asset picker from appearing.
Also removed extraneous EditInline keywords from other properties, as they were only added in an attempt to try to auto-expand details properties.

TTP# 341122 - FIXIF: Can't choose enum file in blackboard

diff --git a/Engine/Source/Runtime/AIModule/Classes/BehaviorTree/Blackboard/BlackboardKeyType_Class.h b/Engine/Source/Runtime/AIModule/Classes/BehaviorTree/Blackboard/BlackboardKeyType_Class.h
--- a/Engine/Source/Runtime/AIModule/Classes/BehaviorTree/Blackboard/BlackboardKeyType_Class.h
+++ b/Engine/Source/Runtime/AIModule/Classes/BehaviorTree/Blackboard/BlackboardKeyType_Class.h
@@ -8,7 +8,7 @@ class UBlackboardKeyType_Class : public UBlackboardKeyType
 {
 	GENERATED_UCLASS_BODY()
 
-	UPROPERTY(Category=Blackboard, EditInline, EditDefaultsOnly, meta=(AllowAbstract="1"))
+	UPROPERTY(Category=Blackboard, EditDefaultsOnly, meta=(AllowAbstract="1"))
 	UClass* BaseClass;
 	
 	static UClass* GetValue(const uint8* RawData);
diff --git a/Engine/Source/Runtime/AIModule/Classes/BehaviorTree/Blackboard/BlackboardKeyType_Enum.h b/Engine/Source/Runtime/AIModule/Classes/BehaviorTree/Blackboard/BlackboardKeyType_Enum.h
--- a/Engine/Source/Runtime/AIModule/Classes/BehaviorTree/Blackboard/BlackboardKeyType_Enum.h
+++ b/Engine/Source/Runtime/AIModule/Classes/BehaviorTree/Blackboard/BlackboardKeyType_Enum.h
@@ -8,7 +8,7 @@ class UBlackboardKeyType_Enum : public UBlackboardKeyType
 {
 	GENERATED_UCLASS_BODY()
 
-	UPROPERTY(Category=Blackboard, EditInline, EditDefaultsOnly)
+	UPROPERTY(Category=Blackboard, EditDefaultsOnly)
 	UEnum* EnumType;
 
 	static uint8 GetValue(const uint8* RawData);
diff --git a/Engine/Source/Runtime/AIModule/Classes/BehaviorTree/Blackboard/BlackboardKeyType_NativeEnum.h b/Engine/Source/Runtime/AIModule/Classes/BehaviorTree/Blackboard/BlackboardKeyType_NativeEnum.h
--- a/Engine/Source/Runtime/AIModule/Classes/BehaviorTree/Blackboard/BlackboardKeyType_NativeEnum.h
+++ b/Engine/Source/Runtime/AIModule/Classes/BehaviorTree/Blackboard/BlackboardKeyType_NativeEnum.h
@@ -9,10 +9,10 @@ class UBlackboardKeyType_NativeEnum : public UBlackboardKeyType
 {
 	GENERATED_UCLASS_BODY()
 
-	UPROPERTY(Category=Blackboard, EditInline, EditDefaultsOnly)
+	UPROPERTY(Category=Blackboard, EditDefaultsOnly)
 	FString EnumName;
 
-	UPROPERTY(Category=Blackboard, EditInline, VisibleDefaultsOnly)
+	UPROPERTY(Category=Blackboard, VisibleDefaultsOnly)
 	bool bIsEnumNameValid;
 
 	UPROPERTY()
diff --git a/Engine/Source/Runtime/AIModule/Classes/BehaviorTree/Blackboard/BlackboardKeyType_Object.h b/Engine/Source/Runtime/AIModule/Classes/BehaviorTree/Blackboard/BlackboardKeyType_Object.h
--- a/Engine/Source/Runtime/AIModule/Classes/BehaviorTree/Blackboard/BlackboardKeyType_Object.h
+++ b/Engine/Source/Runtime/AIModule/Classes/BehaviorTree/Blackboard/BlackboardKeyType_Object.h
@@ -8,7 +8,7 @@ class UBlackboardKeyType_Object : public UBlackboardKeyType
 {
 	GENERATED_UCLASS_BODY()
 
-	UPROPERTY(Category=Blackboard, EditInline, EditDefaultsOnly, meta=(AllowAbstract="1"))
+	UPROPERTY(Category=Blackboard, EditDefaultsOnly, meta=(AllowAbstract="1"))
 	UClass* BaseClass;
 
 	static UObject* GetValue(const uint8* RawData);
