added more checks in path following's segment switch to catch a rare crash
TTP# 340745

diff --git a/Engine/Source/Runtime/AIModule/Private/Navigation/CrowdFollowingComponent.cpp b/Engine/Source/Runtime/AIModule/Private/Navigation/CrowdFollowingComponent.cpp
--- a/Engine/Source/Runtime/AIModule/Private/Navigation/CrowdFollowingComponent.cpp
+++ b/Engine/Source/Runtime/AIModule/Private/Navigation/CrowdFollowingComponent.cpp
@@ -458,6 +458,13 @@ void LogPathPartHelper(AActor* LogOwner, FNavMeshPath* NavMeshPath, int32 StartI
 {
 #if ENABLE_VISUAL_LOG && WITH_RECAST
 	ARecastNavMesh* NavMesh = Cast<ARecastNavMesh>(NavMeshPath->GetOwner());
+	if (NavMesh == NULL ||
+		!NavMeshPath->PathCorridor.IsValidIndex(StartIdx) ||
+		!NavMeshPath->PathCorridor.IsValidIndex(EndIdx))
+	{
+		return;
+	}
+
 	NavMesh->BeginBatchQuery();
 	
 	FVector SegmentStart = NavMeshPath->PathPoints[0].Location;
@@ -500,7 +507,7 @@ void UCrowdFollowingComponent::SetMoveSegment(int32 SegmentStartIndex)
 
 	PathStartIndex = SegmentStartIndex;
 	LastPathPolyIndex = PathStartIndex;
-	if (!Path.IsValid())
+	if (!Path.IsValid() || GetOwner() == NULL)
 	{
 		return;
 	}
