blueprint based BT task will return InProgress when it's using ReceiveExecute or ReceiveTick events

diff --git a/Engine/Source/Runtime/AIModule/Private/BehaviorTree/Tasks/BTTask_BlueprintBase.cpp b/Engine/Source/Runtime/AIModule/Private/BehaviorTree/Tasks/BTTask_BlueprintBase.cpp
--- a/Engine/Source/Runtime/AIModule/Private/BehaviorTree/Tasks/BTTask_BlueprintBase.cpp
+++ b/Engine/Source/Runtime/AIModule/Private/BehaviorTree/Tasks/BTTask_BlueprintBase.cpp
@@ -31,7 +31,9 @@ void UBTTask_BlueprintBase::PostInitProperties()
 
 EBTNodeResult::Type UBTTask_BlueprintBase::ExecuteTask(class UBehaviorTreeComponent* OwnerComp, uint8* NodeMemory)
 {
-	CurrentCallResult = bImplementsReceiveExecute ? EBTNodeResult::InProgress : EBTNodeResult::Failed;
+	// fail when task doesn't react to execution (start or tick)
+	CurrentCallResult = (bImplementsReceiveExecute || bImplementsReceiveTick) ? EBTNodeResult::InProgress : EBTNodeResult::Failed;
+
 	if (bImplementsReceiveExecute)
 	{
 		bStoreFinishResult = true;
