Minor fixed to BP-implementable EQS generator #UE4

diff --git a/Engine/Source/Runtime/AIModule/Classes/EnvironmentQuery/Generators/EnvQueryGenerator_BlueprintBase.h b/Engine/Source/Runtime/AIModule/Classes/EnvironmentQuery/Generators/EnvQueryGenerator_BlueprintBase.h
--- a/Engine/Source/Runtime/AIModule/Classes/EnvironmentQuery/Generators/EnvQueryGenerator_BlueprintBase.h
+++ b/Engine/Source/Runtime/AIModule/Classes/EnvironmentQuery/Generators/EnvQueryGenerator_BlueprintBase.h
@@ -30,6 +30,7 @@ class AIMODULE_API UEnvQueryGenerator_BlueprintBase : public UEnvQueryGenerator
 	TSubclassOf<UEnvQueryItemType> GeneratedItemType;
 
 	virtual void PostInitProperties() override;
+	virtual UWorld* GetWorld() const;
 
 	UFUNCTION(BlueprintImplementableEvent, Category = Generator)
 	virtual void DoItemGeneration(const TArray<FVector>& ContextLocations);
diff --git a/Engine/Source/Runtime/AIModule/Private/EnvironmentQuery/Generators/EnvQueryGenerator_BlueprintBase.cpp b/Engine/Source/Runtime/AIModule/Private/EnvironmentQuery/Generators/EnvQueryGenerator_BlueprintBase.cpp
--- a/Engine/Source/Runtime/AIModule/Private/EnvironmentQuery/Generators/EnvQueryGenerator_BlueprintBase.cpp
+++ b/Engine/Source/Runtime/AIModule/Private/EnvironmentQuery/Generators/EnvQueryGenerator_BlueprintBase.cpp
@@ -17,8 +17,9 @@ UEnvQueryGenerator_BlueprintBase::UEnvQueryGenerator_BlueprintBase(const class F
 	GenerateDelegate.BindUObject(this, &UEnvQueryGenerator_BlueprintBase::GenerateItems);
 
 	Context = UEnvQueryContext_Querier::StaticClass();
-	ItemType = UEnvQueryItemType_Point::StaticClass();
+	ItemType = UEnvQueryItemType_Actor::StaticClass();
 	Radius.Value = 1000.0f;
+	GeneratedItemType = UEnvQueryItemType_Actor::StaticClass();
 }
 
 void UEnvQueryGenerator_BlueprintBase::PostInitProperties()
@@ -27,6 +28,11 @@ void UEnvQueryGenerator_BlueprintBase::PostInitProperties()
 	ItemType = GeneratedItemType;
 }
 
+UWorld* UEnvQueryGenerator_BlueprintBase::GetWorld() const
+{
+	return CachedQueryInstance ? CachedQueryInstance->World : NULL;
+}
+
 void UEnvQueryGenerator_BlueprintBase::GenerateItems(FEnvQueryInstance& QueryInstance)
 {
 	float RadiusValue = 0.f;
@@ -53,7 +59,7 @@ void UEnvQueryGenerator_BlueprintBase::AddGeneratedVector(FVector Vector)
 void UEnvQueryGenerator_BlueprintBase::AddGeneratedActor(AActor* Actor)
 {
 	check(CachedQueryInstance);
-	ensure(ItemType->GetClass()->IsChildOf(UEnvQueryItemType_ActorBase::StaticClass()));
+	//ensure(ItemType->GetClass()->IsChildOf(UEnvQueryItemType_ActorBase::StaticClass()));
 	CachedQueryInstance->AddItemData<UEnvQueryItemType_Actor>(Actor);
 }
 
