FBlackboardKeySelector::AddRotatorFilter now uses CreateNamedObject

TTP# 341526 - [GitHub] 290 : FBlackboardKeySelector::AddRotatorFilter now uses CreateNamedObject

https://github.com/EpicGames/UnrealEngine/pull/290

diff --git a/Engine/Source/Runtime/AIModule/Private/BehaviorTree/BehaviorTreeTypes.cpp b/Engine/Source/Runtime/AIModule/Private/BehaviorTree/BehaviorTreeTypes.cpp
--- a/Engine/Source/Runtime/AIModule/Private/BehaviorTree/BehaviorTreeTypes.cpp
+++ b/Engine/Source/Runtime/AIModule/Private/BehaviorTree/BehaviorTreeTypes.cpp
@@ -240,7 +240,7 @@ void FBlackboardKeySelector::AddVectorFilter(UObject* Owner)
 
 void FBlackboardKeySelector::AddRotatorFilter(UObject* Owner)
 {
-	AllowedTypes.Add(NewObject<UBlackboardKeyType_Rotator>(Owner));
+	AllowedTypes.Add(NewNamedObject<UBlackboardKeyType_Rotator>(Owner, TEXT("BlackboardKeyType_Rotator")));
 }
 
 void FBlackboardKeySelector::AddStringFilter(UObject* Owner)
