Fixed marking a FBlackboardKeySelector as accepting "None" as value breaking key ID caching #UE4

diff --git a/Engine/Source/Runtime/AIModule/Private/BehaviorTree/BehaviorTreeTypes.cpp b/Engine/Source/Runtime/AIModule/Private/BehaviorTree/BehaviorTreeTypes.cpp
--- a/Engine/Source/Runtime/AIModule/Private/BehaviorTree/BehaviorTreeTypes.cpp
+++ b/Engine/Source/Runtime/AIModule/Private/BehaviorTree/BehaviorTreeTypes.cpp
@@ -144,7 +144,7 @@ void FBehaviorTreeSearchData::AddUniqueUpdate(const FBehaviorTreeSearchUpdate& U
 //----------------------------------------------------------------------//
 void FBlackboardKeySelector::CacheSelectedKey(class UBlackboardData* BlackboardAsset)
 {
-	if (BlackboardAsset && !(bNoneIsAllowedValue && SelectedKeyID == FBlackboard::InvalidKey))
+	if (BlackboardAsset && !(bNoneIsAllowedValue && SelectedKeyName == NAME_None))
 	{
 		if (SelectedKeyName == NAME_None)
 		{
