fixed EQSTestingPawn crashing the engine if it gets loaded with a level #UE4

diff --git a/Engine/Source/Runtime/AIModule/Private/EnvironmentQuery/EQSTestingPawn.cpp b/Engine/Source/Runtime/AIModule/Private/EnvironmentQuery/EQSTestingPawn.cpp
--- a/Engine/Source/Runtime/AIModule/Private/EnvironmentQuery/EQSTestingPawn.cpp
+++ b/Engine/Source/Runtime/AIModule/Private/EnvironmentQuery/EQSTestingPawn.cpp
@@ -132,17 +132,6 @@ void AEQSTestingPawn::PostLoad()
 	{
 		SpriteComponent->bHiddenInGame = !bShouldBeVisibleInGame;
 	}
-
-	UEnvQueryManager* EQS = GetWorld() ? UAISystem::GetCurrentEQSManager(GetWorld()) : NULL;
-	if (EQS)
-	{
-		RunEQSQuery();
-	}
-	else
-	{
-		QueryTemplate = NULL;
-		QueryParams.Reset();
-	}
 }
 
 void AEQSTestingPawn::RunEQSQuery()
