fixed resetting cooldown decorator memory on reentering the same subtree

diff --git a/Engine/Source/Runtime/AIModule/Classes/BehaviorTree/BehaviorTreeTypes.h b/Engine/Source/Runtime/AIModule/Classes/BehaviorTree/BehaviorTreeTypes.h
--- a/Engine/Source/Runtime/AIModule/Classes/BehaviorTree/BehaviorTreeTypes.h
+++ b/Engine/Source/Runtime/AIModule/Classes/BehaviorTree/BehaviorTreeTypes.h
@@ -209,6 +209,9 @@ struct FBehaviorTreeInstanceId
 	/** execution index path from root */
 	TArray<uint16> Path;
 
+	/** persistent instance memory */
+	TArray<uint8> InstanceMemory;
+
 	/** index of first node instance (BehaviorTreeComponent.NodeInstances) */
 	int32 FirstNodeInstance;
 
diff --git a/Engine/Source/Runtime/AIModule/Classes/BehaviorTree/Decorators/BTDecorator_Cooldown.h b/Engine/Source/Runtime/AIModule/Classes/BehaviorTree/Decorators/BTDecorator_Cooldown.h
--- a/Engine/Source/Runtime/AIModule/Classes/BehaviorTree/Decorators/BTDecorator_Cooldown.h
+++ b/Engine/Source/Runtime/AIModule/Classes/BehaviorTree/Decorators/BTDecorator_Cooldown.h
@@ -7,7 +7,8 @@
 struct FBTCooldownDecoratorMemory
 {
 	float LastUseTimestamp;
-	uint8 RequestedRestart;
+	uint8 bRequestedRestart : 1;
+	uint8 bInitialized : 1;
 };
 
 /**
diff --git a/Engine/Source/Runtime/AIModule/Private/BehaviorTree/BehaviorTreeComponent.cpp b/Engine/Source/Runtime/AIModule/Private/BehaviorTree/BehaviorTreeComponent.cpp
--- a/Engine/Source/Runtime/AIModule/Private/BehaviorTree/BehaviorTreeComponent.cpp
+++ b/Engine/Source/Runtime/AIModule/Private/BehaviorTree/BehaviorTreeComponent.cpp
@@ -1258,7 +1258,7 @@ bool UBehaviorTreeComponent::PushInstance(class UBehaviorTree* TreeAsset)
 	const bool bLoaded = BTManager->LoadTree(TreeAsset, RootNode, InstanceMemorySize);
 	if (bLoaded)
 	{
-		FBehaviorTreeInstance NewInstance(InstanceMemorySize);
+		FBehaviorTreeInstance NewInstance;
 		NewInstance.InstanceIdIndex = UpdateInstanceId(TreeAsset, ActiveNode, InstanceStack.Num() - 1);
 		NewInstance.RootNode = RootNode;
 		NewInstance.ActiveNode = NULL;
@@ -1267,6 +1267,12 @@ bool UBehaviorTreeComponent::PushInstance(class UBehaviorTree* TreeAsset)
 		// initialize memory and node instances
 		FBehaviorTreeInstanceId& InstanceInfo = KnownInstances[NewInstance.InstanceIdIndex];
 		int32 NodeInstanceIndex = InstanceInfo.FirstNodeInstance;
+		if (InstanceInfo.InstanceMemory.Num() != InstanceMemorySize)
+		{
+			InstanceInfo.InstanceMemory.AddZeroed(InstanceMemorySize);
+		}
+
+		NewInstance.InstanceMemory = InstanceInfo.InstanceMemory;
 		NewInstance.Initialize(this, RootNode, NodeInstanceIndex);
 		NewInstance.InjectNodes(this, RootNode, NodeInstanceIndex);
 
diff --git a/Engine/Source/Runtime/AIModule/Private/BehaviorTree/BehaviorTreeTypes.cpp b/Engine/Source/Runtime/AIModule/Private/BehaviorTree/BehaviorTreeTypes.cpp
--- a/Engine/Source/Runtime/AIModule/Private/BehaviorTree/BehaviorTreeTypes.cpp
+++ b/Engine/Source/Runtime/AIModule/Private/BehaviorTree/BehaviorTreeTypes.cpp
@@ -71,7 +71,7 @@ void FBehaviorTreeInstance::InjectNodes(class UBehaviorTreeComponent* OwnerComp,
 
 void FBehaviorTreeInstance::Cleanup(class UBehaviorTreeComponent* OwnerComp)
 {
-	const FBehaviorTreeInstanceId& Info = OwnerComp->KnownInstances[InstanceIdIndex];
+	FBehaviorTreeInstanceId& Info = OwnerComp->KnownInstances[InstanceIdIndex];
 	if (Info.FirstNodeInstance >= 0)
 	{
 		const int32 LastNodeIdx = OwnerComp->KnownInstances.IsValidIndex(InstanceIdIndex + 1) ?
@@ -83,6 +83,8 @@ void FBehaviorTreeInstance::Cleanup(class UBehaviorTreeComponent* OwnerComp)
 			OwnerComp->NodeInstances[Idx]->OnInstanceDestroyed(OwnerComp);
 		}
 	}
+
+	Info.InstanceMemory = InstanceMemory;
 }
 
 //----------------------------------------------------------------------//
diff --git a/Engine/Source/Runtime/AIModule/Private/BehaviorTree/Decorators/BTDecorator_Cooldown.cpp b/Engine/Source/Runtime/AIModule/Private/BehaviorTree/Decorators/BTDecorator_Cooldown.cpp
--- a/Engine/Source/Runtime/AIModule/Private/BehaviorTree/Decorators/BTDecorator_Cooldown.cpp
+++ b/Engine/Source/Runtime/AIModule/Private/BehaviorTree/Decorators/BTDecorator_Cooldown.cpp
@@ -31,25 +31,31 @@ bool UBTDecorator_Cooldown::CalculateRawConditionValue(class UBehaviorTreeCompon
 void UBTDecorator_Cooldown::InitializeMemory(class UBehaviorTreeComponent* OwnerComp, uint8* NodeMemory) const
 {
 	FBTCooldownDecoratorMemory* DecoratorMemory = (FBTCooldownDecoratorMemory*)NodeMemory;
-	DecoratorMemory->LastUseTimestamp = -FLT_MAX;
+	if (!DecoratorMemory->bInitialized)
+	{
+		DecoratorMemory->bInitialized = true;
+		DecoratorMemory->LastUseTimestamp = -FLT_MAX;
+	}
+
+	DecoratorMemory->bRequestedRestart = false;
 }
 
 void UBTDecorator_Cooldown::OnNodeDeactivation(struct FBehaviorTreeSearchData& SearchData, EBTNodeResult::Type NodeResult)
 {
 	FBTCooldownDecoratorMemory* DecoratorMemory = GetNodeMemory<FBTCooldownDecoratorMemory>(SearchData);
 	DecoratorMemory->LastUseTimestamp = SearchData.OwnerComp->GetWorld()->GetTimeSeconds();
-	DecoratorMemory->RequestedRestart = false;
+	DecoratorMemory->bRequestedRestart = false;
 }
 
 void UBTDecorator_Cooldown::TickNode(class UBehaviorTreeComponent* OwnerComp, uint8* NodeMemory, float DeltaSeconds)
 {
 	FBTCooldownDecoratorMemory* DecoratorMemory = (FBTCooldownDecoratorMemory*)NodeMemory;
-	if (!DecoratorMemory->RequestedRestart)
+	if (!DecoratorMemory->bRequestedRestart)
 	{
 		const float TimePassed = (OwnerComp->GetWorld()->GetTimeSeconds() - DecoratorMemory->LastUseTimestamp);
 		if (TimePassed >= CoolDownTime)
 		{
-			DecoratorMemory->RequestedRestart = true;
+			DecoratorMemory->bRequestedRestart = true;
 			OwnerComp->RequestExecution(this);
 		}
 	}
