better memory management for behavior tree nodes
- added CleanupMemory function to BTNode class
- added init and cleanup modes to identify subtree serialization vs ctor/dtor

diff --git a/Engine/Source/Runtime/AIModule/Classes/BehaviorTree/BTNode.h b/Engine/Source/Runtime/AIModule/Classes/BehaviorTree/BTNode.h
--- a/Engine/Source/Runtime/AIModule/Classes/BehaviorTree/BTNode.h
+++ b/Engine/Source/Runtime/AIModule/Classes/BehaviorTree/BTNode.h
@@ -24,8 +24,11 @@ class AIMODULE_API UBTNode : public UObject
 	/** initialize any asset related data */
 	virtual void InitializeFromAsset(class UBehaviorTree* Asset);
 	
-	/** initialize memory in instance */
-	virtual void InitializeMemory(class UBehaviorTreeComponent* OwnerComp, uint8* NodeMemory) const;
+	/** initialize memory block */
+	virtual void InitializeMemory(class UBehaviorTreeComponent* OwnerComp, uint8* NodeMemory, EBTMemoryInit::Type InitType) const;
+
+	/** cleanup memory block */
+	virtual void CleanupMemory(class UBehaviorTreeComponent* OwnerComp, uint8* NodeMemory, EBTMemoryClear::Type CleanupType) const;
 
 	/** gathers description of all runtime parameters */
 	virtual void DescribeRuntimeValues(const class UBehaviorTreeComponent* OwnerComp, uint8* NodeMemory, EBTDescriptionVerbosity::Type Verbosity, TArray<FString>& Values) const;
@@ -39,8 +42,11 @@ class AIMODULE_API UBTNode : public UObject
 	/** called when node instance is removed from tree */
 	virtual void OnInstanceDestroyed(class UBehaviorTreeComponent* OwnerComp);
 
-	/** initialize memory and instancing */
-	void InitializeForInstance(class UBehaviorTreeComponent* OwnerComp, uint8* NodeMemory, int32& NextInstancedIndex) const;
+	/** called on creating subtree to set up memory and instancing */
+	void InitializeInSubtree(class UBehaviorTreeComponent* OwnerComp, uint8* NodeMemory, int32& NextInstancedIndex, EBTMemoryInit::Type InitType) const;
+
+	/** called on removing subtree to cleanup memory */
+	void CleanupInSubtree(class UBehaviorTreeComponent* OwnerComp, uint8* NodeMemory, EBTMemoryClear::Type CleanupType) const;
 
 	/** size of special, hidden memory block for internal mechanics */
 	virtual uint16 GetSpecialMemorySize() const;
diff --git a/Engine/Source/Runtime/AIModule/Classes/BehaviorTree/BehaviorTreeComponent.h b/Engine/Source/Runtime/AIModule/Classes/BehaviorTree/BehaviorTreeComponent.h
--- a/Engine/Source/Runtime/AIModule/Classes/BehaviorTree/BehaviorTreeComponent.h
+++ b/Engine/Source/Runtime/AIModule/Classes/BehaviorTree/BehaviorTreeComponent.h
@@ -50,6 +50,10 @@ public:
 	virtual bool IsPaused() const override;
 	// End UBrainComponent overrides
 
+	// Begin UObject overrides
+	virtual void BeginDestroy() override;
+	// End UObject overrides
+
 	/** starts execution from root */
 	bool StartTree(class UBehaviorTree* Asset, EBTExecutionMode::Type ExecuteMode = EBTExecutionMode::Looped);
 
@@ -197,6 +201,9 @@ protected:
 	/** add unique Id of newly created subtree to KnownInstances list and return its index */
 	uint8 UpdateInstanceId(class UBehaviorTree* TreeAsset, const class UBTNode* OriginNode, int32 OriginInstanceIdx);
 
+	/** remove instanced nodes, known subtree instances and safely clears their persistent memory */
+	void RemoveAllInstances();
+
 	/** find next task to execute */
 	UBTTaskNode* FindNextTask(class UBTCompositeNode* ParentNode, uint16 ParentInstanceIdx, EBTNodeResult::Type LastResult);
 
diff --git a/Engine/Source/Runtime/AIModule/Classes/BehaviorTree/BehaviorTreeTypes.h b/Engine/Source/Runtime/AIModule/Classes/BehaviorTree/BehaviorTreeTypes.h
--- a/Engine/Source/Runtime/AIModule/Classes/BehaviorTree/BehaviorTreeTypes.h
+++ b/Engine/Source/Runtime/AIModule/Classes/BehaviorTree/BehaviorTreeTypes.h
@@ -63,6 +63,24 @@ namespace EBTExecutionMode
 	};
 }
 
+namespace EBTMemoryInit
+{
+	enum Type
+	{
+		Initialize,		// first time initialization
+		RestoreSubtree,	// loading saved data on reentering subtree
+	};
+}
+
+namespace EBTMemoryClear
+{
+	enum Type
+	{
+		Destroy,		// final clear
+		StoreSubtree,	// saving data on leaving subtree
+	};
+}
+
 UENUM()
 namespace EBTFlowAbortMode
 {
@@ -204,6 +222,9 @@ struct FBehaviorTreeInstanceId
 	/** behavior tree asset */
 	class UBehaviorTree* TreeAsset;
 
+	/** root node in template for cleanup purposes */
+	class UBTCompositeNode* RootNode;
+
 	/** execution index path from root */
 	TArray<uint16> Path;
 
@@ -213,7 +234,7 @@ struct FBehaviorTreeInstanceId
 	/** index of first node instance (BehaviorTreeComponent.NodeInstances) */
 	int32 FirstNodeInstance;
 
-	FBehaviorTreeInstanceId() :	TreeAsset(0), FirstNodeInstance(-1) {}
+	FBehaviorTreeInstanceId() : TreeAsset(0), RootNode(0), FirstNodeInstance(-1) {}
 
 	bool operator==(const FBehaviorTreeInstanceId& Other) const
 	{
@@ -264,13 +285,18 @@ struct FBehaviorTreeInstance
 #endif // STATS
 
 	/** initialize memory and create node instances */
-	void Initialize(class UBehaviorTreeComponent* OwnerComp, UBTCompositeNode* Node, int32& InstancedIndex);
+	void Initialize(class UBehaviorTreeComponent* OwnerComp, UBTCompositeNode* Node, int32& InstancedIndex, EBTMemoryInit::Type InitType);
 
 	/** update injected nodes */
 	void InjectNodes(class UBehaviorTreeComponent* OwnerComp, UBTCompositeNode* Node, int32& InstancedIndex);
 
 	/** cleanup node instances */
-	void Cleanup(class UBehaviorTreeComponent* OwnerComp);
+	void Cleanup(class UBehaviorTreeComponent* OwnerComp, EBTMemoryClear::Type CleanupType);
+
+protected:
+
+	/** worker for updating all nodes */
+	void CleanupNodes(class UBehaviorTreeComponent* OwnerComp, UBTCompositeNode* Node, EBTMemoryClear::Type CleanupType);
 };
 
 struct FBTNodeIndex
diff --git a/Engine/Source/Runtime/AIModule/Classes/BehaviorTree/Decorators/BTDecorator_Cooldown.h b/Engine/Source/Runtime/AIModule/Classes/BehaviorTree/Decorators/BTDecorator_Cooldown.h
--- a/Engine/Source/Runtime/AIModule/Classes/BehaviorTree/Decorators/BTDecorator_Cooldown.h
+++ b/Engine/Source/Runtime/AIModule/Classes/BehaviorTree/Decorators/BTDecorator_Cooldown.h
@@ -8,7 +8,6 @@ struct FBTCooldownDecoratorMemory
 {
 	float LastUseTimestamp;
 	uint8 bRequestedRestart : 1;
-	uint8 bInitialized : 1;
 };
 
 /**
@@ -29,7 +28,7 @@ class AIMODULE_API UBTDecorator_Cooldown : public UBTDecorator
 	// End UObject Interface
 
 	virtual bool CalculateRawConditionValue(class UBehaviorTreeComponent* OwnerComp, uint8* NodeMemory) const override;
-	virtual void InitializeMemory(class UBehaviorTreeComponent* OwnerComp, uint8* NodeMemory) const override;
+	virtual void InitializeMemory(class UBehaviorTreeComponent* OwnerComp, uint8* NodeMemory, EBTMemoryInit::Type InitType) const override;
 	virtual uint16 GetInstanceMemorySize() const override;
 	virtual void DescribeRuntimeValues(const class UBehaviorTreeComponent* OwnerComp, uint8* NodeMemory, EBTDescriptionVerbosity::Type Verbosity, TArray<FString>& Values) const override;
 	virtual FString GetStaticDescription() const override;
diff --git a/Engine/Source/Runtime/AIModule/Classes/BehaviorTree/Tasks/BTTask_RunBehavior.h b/Engine/Source/Runtime/AIModule/Classes/BehaviorTree/Tasks/BTTask_RunBehavior.h
--- a/Engine/Source/Runtime/AIModule/Classes/BehaviorTree/Tasks/BTTask_RunBehavior.h
+++ b/Engine/Source/Runtime/AIModule/Classes/BehaviorTree/Tasks/BTTask_RunBehavior.h
@@ -25,6 +25,7 @@ class UBTTask_RunBehavior : public UBTTaskNode
 
 	virtual EBTNodeResult::Type ExecuteTask(UBehaviorTreeComponent* OwnerComp, uint8* NodeMemory) override;
 	virtual uint16 GetInstanceMemorySize() const override;
+	virtual void CleanupMemory(class UBehaviorTreeComponent* OwnerComp, uint8* NodeMemory, EBTMemoryClear::Type CleanupType) const override;
 	virtual FString GetStaticDescription() const override;
 
 #if WITH_EDITOR
diff --git a/Engine/Source/Runtime/AIModule/Private/BehaviorTree/BTNode.cpp b/Engine/Source/Runtime/AIModule/Private/BehaviorTree/BTNode.cpp
--- a/Engine/Source/Runtime/AIModule/Private/BehaviorTree/BTNode.cpp
+++ b/Engine/Source/Runtime/AIModule/Private/BehaviorTree/BTNode.cpp
@@ -44,7 +44,12 @@ void UBTNode::InitializeNode(class UBTCompositeNode* InParentNode, uint16 InExec
 	TreeDepth = InTreeDepth;
 }
 
-void UBTNode::InitializeMemory(class UBehaviorTreeComponent* OwnerComp, uint8* NodeMemory) const
+void UBTNode::InitializeMemory(class UBehaviorTreeComponent* OwnerComp, uint8* NodeMemory, EBTMemoryInit::Type InitType) const
+{
+	// empty in base class 
+}
+
+void UBTNode::CleanupMemory(class UBehaviorTreeComponent* OwnerComp, uint8* NodeMemory, EBTMemoryClear::Type CleanupType) const
 {
 	// empty in base class 
 }
@@ -59,7 +64,7 @@ void UBTNode::OnInstanceDestroyed(class UBehaviorTreeComponent* OwnerComp)
 	// empty in base class
 }
 
-void UBTNode::InitializeForInstance(class UBehaviorTreeComponent* OwnerComp, uint8* NodeMemory, int32& NextInstancedIndex) const
+void UBTNode::InitializeInSubtree(class UBehaviorTreeComponent* OwnerComp, uint8* NodeMemory, int32& NextInstancedIndex, EBTMemoryInit::Type InitType) const
 {
 	if (bCreateNodeInstance)
 	{
@@ -84,7 +89,15 @@ void UBTNode::InitializeForInstance(class UBehaviorTreeComponent* OwnerComp, uin
 	}
 	else
 	{
-		InitializeMemory(OwnerComp, NodeMemory);
+		InitializeMemory(OwnerComp, NodeMemory, InitType);
+	}
+}
+
+void UBTNode::CleanupInSubtree(class UBehaviorTreeComponent* OwnerComp, uint8* NodeMemory, EBTMemoryClear::Type CleanupType) const
+{
+	if (!bCreateNodeInstance && !bIsInjected)
+	{
+		CleanupMemory(OwnerComp, NodeMemory, CleanupType);
 	}
 }
 
diff --git a/Engine/Source/Runtime/AIModule/Private/BehaviorTree/BehaviorTreeComponent.cpp b/Engine/Source/Runtime/AIModule/Private/BehaviorTree/BehaviorTreeComponent.cpp
--- a/Engine/Source/Runtime/AIModule/Private/BehaviorTree/BehaviorTreeComponent.cpp
+++ b/Engine/Source/Runtime/AIModule/Private/BehaviorTree/BehaviorTreeComponent.cpp
@@ -27,6 +27,12 @@ UBehaviorTreeComponent::UBehaviorTreeComponent(const class FPostConstructInitial
 	SearchData.OwnerComp = this;
 }
 
+void UBehaviorTreeComponent::BeginDestroy()
+{
+	RemoveAllInstances();
+	Super::BeginDestroy();
+}
+
 void UBehaviorTreeComponent::RestartLogic()
 {
 	UE_VLOG(GetOwner(), LogBehaviorTree, Log, TEXT("UBehaviorTreeComponent::RestartLogic"));
@@ -110,6 +116,8 @@ bool UBehaviorTreeComponent::StartTree(class UBehaviorTree* TreeAsset, EBTExecut
 	}
 
 	StopTree();
+	RemoveAllInstances();
+
 	bLoopExecution = (ExecuteMode == EBTExecutionMode::Looped);
 	bIsRunning = true;
 
@@ -138,7 +146,7 @@ void UBehaviorTreeComponent::StopTree()
 		}
 
 		// clear instance
-		InstanceStack[InstanceIndex].Cleanup(this);
+		InstanceStack[InstanceIndex].Cleanup(this, EBTMemoryClear::StoreSubtree);
 	}
 
 	InstanceStack.Reset();
@@ -836,7 +844,7 @@ void UBehaviorTreeComponent::ProcessExecutionRequest()
 	{
 		for (int32 InstanceIndex = ActiveInstanceIdx + 1; InstanceIndex < InstanceStack.Num(); InstanceIndex++)
 		{
-			InstanceStack[InstanceIndex].Cleanup(this);
+			InstanceStack[InstanceIndex].Cleanup(this, EBTMemoryClear::StoreSubtree);
 		}
 
 		InstanceStack.SetNum(ActiveInstanceIdx + 1);
@@ -904,7 +912,7 @@ void UBehaviorTreeComponent::ProcessExecutionRequest()
 					ApplySearchData(NULL);
 
 					// and leave subtree
-					InstanceStack[ActiveInstanceIdx].Cleanup(this);
+					InstanceStack[ActiveInstanceIdx].Cleanup(this, EBTMemoryClear::StoreSubtree);
 					InstanceStack.Pop();
 					ActiveInstanceIdx--;
 
@@ -1272,13 +1280,15 @@ bool UBehaviorTreeComponent::PushInstance(class UBehaviorTree* TreeAsset)
 		// initialize memory and node instances
 		FBehaviorTreeInstanceId& InstanceInfo = KnownInstances[NewInstance.InstanceIdIndex];
 		int32 NodeInstanceIndex = InstanceInfo.FirstNodeInstance;
-		if (InstanceInfo.InstanceMemory.Num() != InstanceMemorySize)
+		const bool bFirstTime = (InstanceInfo.InstanceMemory.Num() != InstanceMemorySize);
+		if (bFirstTime)
 		{
 			InstanceInfo.InstanceMemory.AddZeroed(InstanceMemorySize);
+			InstanceInfo.RootNode = RootNode;
 		}
 
 		NewInstance.InstanceMemory = InstanceInfo.InstanceMemory;
-		NewInstance.Initialize(this, RootNode, NodeInstanceIndex);
+		NewInstance.Initialize(this, RootNode, NodeInstanceIndex, bFirstTime ? EBTMemoryInit::Initialize : EBTMemoryInit::RestoreSubtree);
 		NewInstance.InjectNodes(this, RootNode, NodeInstanceIndex);
 
 		INC_DWORD_STAT(STAT_AI_BehaviorTree_NumInstances);
@@ -1414,6 +1424,29 @@ uint8* UBehaviorTreeComponent::GetNodeMemory(UBTNode* Node, int32 InstanceIdx) c
 	return InstanceStack.IsValidIndex(InstanceIdx) ? (uint8*)Node->GetNodeMemory<uint8>(InstanceStack[InstanceIdx]) : NULL;
 }
 
+void UBehaviorTreeComponent::RemoveAllInstances()
+{
+	if (InstanceStack.Num())
+	{
+		StopTree();
+	}
+
+	FBehaviorTreeInstance DummyInstance;
+	for (int32 Idx = 0; Idx < KnownInstances.Num(); Idx++)
+	{
+		const FBehaviorTreeInstanceId& Info = KnownInstances[Idx];
+		DummyInstance.InstanceMemory = Info.InstanceMemory;
+		DummyInstance.InstanceIdIndex = Idx;
+		DummyInstance.RootNode = Info.RootNode;
+
+		DummyInstance.Cleanup(this, EBTMemoryClear::Destroy);
+
+		DummyInstance.InstanceMemory.Reset();
+	}
+
+	NodeInstances.Reset();
+}
+
 FString UBehaviorTreeComponent::GetDebugInfoString() const 
 { 
 	FString DebugInfo;
diff --git a/Engine/Source/Runtime/AIModule/Private/BehaviorTree/BehaviorTreeTypes.cpp b/Engine/Source/Runtime/AIModule/Private/BehaviorTree/BehaviorTreeTypes.cpp
--- a/Engine/Source/Runtime/AIModule/Private/BehaviorTree/BehaviorTreeTypes.cpp
+++ b/Engine/Source/Runtime/AIModule/Private/BehaviorTree/BehaviorTreeTypes.cpp
@@ -9,7 +9,7 @@
 //----------------------------------------------------------------------//
 // FBehaviorTreeInstance
 //----------------------------------------------------------------------//
-void FBehaviorTreeInstance::Initialize(class UBehaviorTreeComponent* OwnerComp, UBTCompositeNode* Node, int32& InstancedIndex)
+void FBehaviorTreeInstance::Initialize(class UBehaviorTreeComponent* OwnerComp, UBTCompositeNode* Node, int32& InstancedIndex, EBTMemoryInit::Type InitType)
 {
 	if (Node == NULL)
 	{
@@ -18,10 +18,10 @@ void FBehaviorTreeInstance::Initialize(class UBehaviorTreeComponent* OwnerComp,
 
 	for (int32 ServiceIndex = 0; ServiceIndex < Node->Services.Num(); ServiceIndex++)
 	{
-		Node->Services[ServiceIndex]->InitializeForInstance(OwnerComp, Node->Services[ServiceIndex]->GetNodeMemory<uint8>(*this), InstancedIndex);
+		Node->Services[ServiceIndex]->InitializeInSubtree(OwnerComp, Node->Services[ServiceIndex]->GetNodeMemory<uint8>(*this), InstancedIndex, InitType);
 	}
 
-	Node->InitializeForInstance(OwnerComp, Node->GetNodeMemory<uint8>(*this), InstancedIndex);
+	Node->InitializeInSubtree(OwnerComp, Node->GetNodeMemory<uint8>(*this), InstancedIndex, InitType);
 
 	for (int32 ChildIndex = 0; ChildIndex < Node->Children.Num(); ChildIndex++)
 	{
@@ -29,16 +29,16 @@ void FBehaviorTreeInstance::Initialize(class UBehaviorTreeComponent* OwnerComp,
 
 		for (int32 DecoratorIndex = 0; DecoratorIndex < ChildInfo.Decorators.Num(); DecoratorIndex++)
 		{
-			ChildInfo.Decorators[DecoratorIndex]->InitializeForInstance(OwnerComp, ChildInfo.Decorators[DecoratorIndex]->GetNodeMemory<uint8>(*this), InstancedIndex);
+			ChildInfo.Decorators[DecoratorIndex]->InitializeInSubtree(OwnerComp, ChildInfo.Decorators[DecoratorIndex]->GetNodeMemory<uint8>(*this), InstancedIndex, InitType);
 		}
 
 		if (ChildInfo.ChildComposite)
 		{
-			Initialize(OwnerComp, ChildInfo.ChildComposite, InstancedIndex);
+			Initialize(OwnerComp, ChildInfo.ChildComposite, InstancedIndex, InitType);
 		}
 		else if (ChildInfo.ChildTask)
 		{
-			ChildInfo.ChildTask->InitializeForInstance(OwnerComp, ChildInfo.ChildTask->GetNodeMemory<uint8>(*this), InstancedIndex);
+			ChildInfo.ChildTask->InitializeInSubtree(OwnerComp, ChildInfo.ChildTask->GetNodeMemory<uint8>(*this), InstancedIndex, InitType);
 		}
 	}
 }
@@ -69,7 +69,7 @@ void FBehaviorTreeInstance::InjectNodes(class UBehaviorTreeComponent* OwnerComp,
 	}
 }
 
-void FBehaviorTreeInstance::Cleanup(class UBehaviorTreeComponent* OwnerComp)
+void FBehaviorTreeInstance::Cleanup(class UBehaviorTreeComponent* OwnerComp, EBTMemoryClear::Type CleanupType)
 {
 	FBehaviorTreeInstanceId& Info = OwnerComp->KnownInstances[InstanceIdIndex];
 	if (Info.FirstNodeInstance >= 0)
@@ -84,9 +84,45 @@ void FBehaviorTreeInstance::Cleanup(class UBehaviorTreeComponent* OwnerComp)
 		}
 	}
 
+	CleanupNodes(OwnerComp, RootNode, CleanupType);
 	Info.InstanceMemory = InstanceMemory;
 }
 
+void FBehaviorTreeInstance::CleanupNodes(class UBehaviorTreeComponent* OwnerComp, UBTCompositeNode* Node, EBTMemoryClear::Type CleanupType)
+{
+	if (Node == NULL)
+	{
+		return;
+	}
+
+	for (int32 ServiceIndex = 0; ServiceIndex < Node->Services.Num(); ServiceIndex++)
+	{
+		Node->Services[ServiceIndex]->CleanupInSubtree(OwnerComp, Node->Services[ServiceIndex]->GetNodeMemory<uint8>(*this), CleanupType);
+	}
+
+	Node->CleanupInSubtree(OwnerComp, Node->GetNodeMemory<uint8>(*this), CleanupType);
+
+	for (int32 ChildIndex = 0; ChildIndex < Node->Children.Num(); ChildIndex++)
+	{
+		FBTCompositeChild& ChildInfo = Node->Children[ChildIndex];
+
+		for (int32 DecoratorIndex = 0; DecoratorIndex < ChildInfo.Decorators.Num(); DecoratorIndex++)
+		{
+			ChildInfo.Decorators[DecoratorIndex]->CleanupInSubtree(OwnerComp, ChildInfo.Decorators[DecoratorIndex]->GetNodeMemory<uint8>(*this), CleanupType);
+		}
+
+		if (ChildInfo.ChildComposite)
+		{
+			CleanupNodes(OwnerComp, ChildInfo.ChildComposite, CleanupType);
+		}
+		else if (ChildInfo.ChildTask)
+		{
+			ChildInfo.ChildTask->CleanupInSubtree(OwnerComp, ChildInfo.ChildTask->GetNodeMemory<uint8>(*this), CleanupType);
+		}
+	}
+}
+
+
 //----------------------------------------------------------------------//
 // FBTNodeIndex
 //----------------------------------------------------------------------//
diff --git a/Engine/Source/Runtime/AIModule/Private/BehaviorTree/BlueprintNodeHelpers.cpp b/Engine/Source/Runtime/AIModule/Private/BehaviorTree/BlueprintNodeHelpers.cpp
--- a/Engine/Source/Runtime/AIModule/Private/BehaviorTree/BlueprintNodeHelpers.cpp
+++ b/Engine/Source/Runtime/AIModule/Private/BehaviorTree/BlueprintNodeHelpers.cpp
@@ -212,10 +212,13 @@ namespace BlueprintNodeHelpers
 		return bFound;
 	}
 
-	void AbortLatentActions(AActor* OwningActor, const UObject* Ob)
+	void AbortLatentActions(UActorComponent* OwnerOb, const UObject* Ob)
 	{
-		UWorld* MyWorld = OwningActor->GetWorld();
-		MyWorld->GetLatentActionManager().RemoveActionsForObject(Ob);
-		MyWorld->GetTimerManager().ClearAllTimersForObject(Ob);
+		if (OwnerOb && !OwnerOb->HasAnyFlags(RF_BeginDestroyed) && OwnerOb->GetOwner())
+		{
+			UWorld* MyWorld = OwnerOb->GetOwner()->GetWorld();
+			MyWorld->GetLatentActionManager().RemoveActionsForObject(Ob);
+			MyWorld->GetTimerManager().ClearAllTimersForObject(Ob);
+		}
 	}
 }
diff --git a/Engine/Source/Runtime/AIModule/Private/BehaviorTree/Decorators/BTDecorator_BlueprintBase.cpp b/Engine/Source/Runtime/AIModule/Private/BehaviorTree/Decorators/BTDecorator_BlueprintBase.cpp
--- a/Engine/Source/Runtime/AIModule/Private/BehaviorTree/Decorators/BTDecorator_BlueprintBase.cpp
+++ b/Engine/Source/Runtime/AIModule/Private/BehaviorTree/Decorators/BTDecorator_BlueprintBase.cpp
@@ -197,7 +197,7 @@ void UBTDecorator_BlueprintBase::OnBlackboardChange(const UBlackboardComponent*
 void UBTDecorator_BlueprintBase::OnInstanceDestroyed(class UBehaviorTreeComponent* OwnerComp)
 {
 	// force dropping all pending latent actions associated with this blueprint
-	BlueprintNodeHelpers::AbortLatentActions(OwnerComp->GetOwner(), this);
+	BlueprintNodeHelpers::AbortLatentActions(OwnerComp, this);
 }
 
 #if WITH_EDITOR
diff --git a/Engine/Source/Runtime/AIModule/Private/BehaviorTree/Decorators/BTDecorator_Cooldown.cpp b/Engine/Source/Runtime/AIModule/Private/BehaviorTree/Decorators/BTDecorator_Cooldown.cpp
--- a/Engine/Source/Runtime/AIModule/Private/BehaviorTree/Decorators/BTDecorator_Cooldown.cpp
+++ b/Engine/Source/Runtime/AIModule/Private/BehaviorTree/Decorators/BTDecorator_Cooldown.cpp
@@ -28,12 +28,11 @@ bool UBTDecorator_Cooldown::CalculateRawConditionValue(class UBehaviorTreeCompon
 	return TimePassed >= CoolDownTime;
 }
 
-void UBTDecorator_Cooldown::InitializeMemory(class UBehaviorTreeComponent* OwnerComp, uint8* NodeMemory) const
+void UBTDecorator_Cooldown::InitializeMemory(class UBehaviorTreeComponent* OwnerComp, uint8* NodeMemory, EBTMemoryInit::Type InitType) const
 {
 	FBTCooldownDecoratorMemory* DecoratorMemory = (FBTCooldownDecoratorMemory*)NodeMemory;
-	if (!DecoratorMemory->bInitialized)
+	if (InitType == EBTMemoryInit::Initialize)
 	{
-		DecoratorMemory->bInitialized = true;
 		DecoratorMemory->LastUseTimestamp = -FLT_MAX;
 	}
 
diff --git a/Engine/Source/Runtime/AIModule/Private/BehaviorTree/Services/BTService_BlueprintBase.cpp b/Engine/Source/Runtime/AIModule/Private/BehaviorTree/Services/BTService_BlueprintBase.cpp
--- a/Engine/Source/Runtime/AIModule/Private/BehaviorTree/Services/BTService_BlueprintBase.cpp
+++ b/Engine/Source/Runtime/AIModule/Private/BehaviorTree/Services/BTService_BlueprintBase.cpp
@@ -48,7 +48,7 @@ void UBTService_BlueprintBase::OnCeaseRelevant(class UBehaviorTreeComponent* Own
 {
 	// force dropping all pending latent actions associated with this blueprint
 	// we can't have those resuming activity when node is/was aborted
-	BlueprintNodeHelpers::AbortLatentActions(OwnerComp->GetOwner(), this);
+	BlueprintNodeHelpers::AbortLatentActions(OwnerComp, this);
 
 	Super::OnCeaseRelevant(OwnerComp, NodeMemory);
 
@@ -120,7 +120,7 @@ void UBTService_BlueprintBase::DescribeRuntimeValues(const class UBehaviorTreeCo
 void UBTService_BlueprintBase::OnInstanceDestroyed(class UBehaviorTreeComponent* OwnerComp)
 {
 	// force dropping all pending latent actions associated with this blueprint
-	BlueprintNodeHelpers::AbortLatentActions(OwnerComp->GetOwner(), this);
+	BlueprintNodeHelpers::AbortLatentActions(OwnerComp, this);
 }
 
 #if WITH_EDITOR
diff --git a/Engine/Source/Runtime/AIModule/Private/BehaviorTree/Tasks/BTTask_BlueprintBase.cpp b/Engine/Source/Runtime/AIModule/Private/BehaviorTree/Tasks/BTTask_BlueprintBase.cpp
--- a/Engine/Source/Runtime/AIModule/Private/BehaviorTree/Tasks/BTTask_BlueprintBase.cpp
+++ b/Engine/Source/Runtime/AIModule/Private/BehaviorTree/Tasks/BTTask_BlueprintBase.cpp
@@ -50,7 +50,7 @@ EBTNodeResult::Type UBTTask_BlueprintBase::AbortTask(class UBehaviorTreeComponen
 {
 	// force dropping all pending latent actions associated with this blueprint
 	// we can't have those resuming activity when node is/was aborted
-	BlueprintNodeHelpers::AbortLatentActions(OwnerComp->GetOwner(), this);
+	BlueprintNodeHelpers::AbortLatentActions(OwnerComp, this);
 
 	CurrentCallResult = bImplementsReceiveAbort ? EBTNodeResult::InProgress : EBTNodeResult::Aborted;
 	if (bImplementsReceiveAbort)
@@ -159,7 +159,7 @@ void UBTTask_BlueprintBase::DescribeRuntimeValues(const class UBehaviorTreeCompo
 void UBTTask_BlueprintBase::OnInstanceDestroyed(class UBehaviorTreeComponent* OwnerComp)
 {
 	// force dropping all pending latent actions associated with this blueprint
-	BlueprintNodeHelpers::AbortLatentActions(OwnerComp->GetOwner(), this);
+	BlueprintNodeHelpers::AbortLatentActions(OwnerComp, this);
 }
 
 #if WITH_EDITOR
diff --git a/Engine/Source/Runtime/AIModule/Private/BehaviorTree/Tasks/BTTask_RunBehavior.cpp b/Engine/Source/Runtime/AIModule/Private/BehaviorTree/Tasks/BTTask_RunBehavior.cpp
--- a/Engine/Source/Runtime/AIModule/Private/BehaviorTree/Tasks/BTTask_RunBehavior.cpp
+++ b/Engine/Source/Runtime/AIModule/Private/BehaviorTree/Tasks/BTTask_RunBehavior.cpp
@@ -33,6 +33,9 @@ uint16 UBTTask_RunBehavior::GetInstanceMemorySize() const
 		UBehaviorTreeManager::InitializeMemoryHelper(BehaviorAsset->RootDecorators, MemoryOffsets, MemorySize);
 	}
 
+	// memory block header: index of first node instance in owner's array
+	MemorySize += sizeof(FBTInstancedNodeMemory);
+
 	return MemorySize;
 }
 
@@ -44,7 +47,10 @@ void UBTTask_RunBehavior::InjectNodes(UBehaviorTreeComponent* OwnerComp, uint8*
 	}
 
 	const int32 NumInjectedDecorators = BehaviorAsset->RootDecorators.Num();
-	const int32 FirstNodeIdx = InstancedIndex;
+	FBTInstancedNodeMemory* NodeMemoryHeader = (FBTInstancedNodeMemory*)NodeMemory;
+	int32 FirstNodeIdx = NodeMemoryHeader->NodeIdx;
+
+	uint8* InjectedMemoryBase = NodeMemory + sizeof(FBTInstancedNodeMemory);
 
 	// initialize on first access
 	if (!OwnerComp->NodeInstances.IsValidIndex(InstancedIndex))
@@ -56,10 +62,16 @@ void UBTTask_RunBehavior::InjectNodes(UBehaviorTreeComponent* OwnerComp, uint8*
 		const int32 AlignedInstanceMemorySize = UBehaviorTreeManager::GetAlignedDataSize(sizeof(FBTInstancedNodeMemory));
 
 		// prepare dummy memory block for nodes that won't require instancing and offset it by special data size
-		// InitializeForInstance will read it through GetSpecialNodeMemory function, which moves pointer back 
+		// InitializeInSubtree will read it through GetSpecialNodeMemory function, which moves pointer back 
 		FBTInstancedNodeMemory DummyMemory;
 		uint8* RawDummyMemory = ((uint8*)&DummyMemory) + AlignedInstanceMemorySize;
 
+		// newly created nodes = full init
+		EBTMemoryInit::Type InitType = EBTMemoryInit::Initialize;
+
+		NodeMemoryHeader->NodeIdx = InstancedIndex;
+		FirstNodeIdx = InstancedIndex;
+
 		// create nodes
 		for (int32 Idx = 0; Idx < NumInjectedDecorators; Idx++)
 		{
@@ -68,16 +80,16 @@ void UBTTask_RunBehavior::InjectNodes(UBehaviorTreeComponent* OwnerComp, uint8*
 			const bool bUsesInstancing = DecoratorOb->HasInstance();
 			if (bUsesInstancing)
 			{
-				DecoratorOb->InitializeForInstance(OwnerComp, NodeMemory + MemoryOffsets[Idx], InstancedIndex);
+				DecoratorOb->InitializeInSubtree(OwnerComp, InjectedMemoryBase + MemoryOffsets[Idx], InstancedIndex, InitType);
 			}
 			else
 			{
 				DecoratorOb->ForceInstancing(true);
-				DecoratorOb->InitializeForInstance(OwnerComp, RawDummyMemory, InstancedIndex);
+				DecoratorOb->InitializeInSubtree(OwnerComp, RawDummyMemory, InstancedIndex, InitType);
 				DecoratorOb->ForceInstancing(false);
 
 				UBTDecorator* InstancedOb = Cast<UBTDecorator>(OwnerComp->NodeInstances.Last());
-				InstancedOb->InitializeMemory(OwnerComp, NodeMemory + MemoryOffsets[Idx]);
+				InstancedOb->InitializeMemory(OwnerComp, InjectedMemoryBase + MemoryOffsets[Idx], InitType);
 			}
 		}
 
@@ -94,6 +106,9 @@ void UBTTask_RunBehavior::InjectNodes(UBehaviorTreeComponent* OwnerComp, uint8*
 	}
 	else
 	{
+		// restoring existing nodes = partial init
+		EBTMemoryInit::Type InitType = EBTMemoryInit::RestoreSubtree;
+
 		// initialize memory for injected nodes
 		for (int32 Idx = 0; Idx < NumInjectedDecorators; Idx++)
 		{
@@ -107,8 +122,8 @@ void UBTTask_RunBehavior::InjectNodes(UBehaviorTreeComponent* OwnerComp, uint8*
 			}
 			else
 			{
-				uint8* InjectedNodeMemory = NodeMemory + (InstancedOb->GetMemoryOffset() - GetMemoryOffset());
-				InstancedOb->InitializeMemory(OwnerComp, InjectedNodeMemory);
+				uint8* InjectedNodeMemory = InjectedMemoryBase + (InstancedOb->GetMemoryOffset() - GetMemoryOffset());
+				InstancedOb->InitializeMemory(OwnerComp, InjectedNodeMemory, InitType);
 			}
 		}
 
@@ -225,6 +240,29 @@ void UBTTask_RunBehavior::InjectNodes(UBehaviorTreeComponent* OwnerComp, uint8*
 	}
 }
 
+void UBTTask_RunBehavior::CleanupMemory(class UBehaviorTreeComponent* OwnerComp, uint8* NodeMemory, EBTMemoryClear::Type CleanupType) const
+{
+	Super::CleanupMemory(OwnerComp, NodeMemory, CleanupType);
+
+	if (GetParentNode() && BehaviorAsset)
+	{
+		FBTInstancedNodeMemory* NodeMemoryHeader = (FBTInstancedNodeMemory*)NodeMemory;
+		uint8* InjectedMemoryBase = NodeMemory + sizeof(FBTInstancedNodeMemory);
+
+		const int32 NumInjectedDecorators = BehaviorAsset->RootDecorators.Num();
+		for (int32 Idx = 0; Idx < NumInjectedDecorators; Idx++)
+		{
+			const int32 InstancedNodeIdx = NodeMemoryHeader->NodeIdx + Idx;
+			UBTDecorator* InstancedOb = Cast<UBTDecorator>(OwnerComp->NodeInstances[InstancedNodeIdx]);
+			if (InstancedOb)
+			{
+				uint8* InjectedNodeMemory = InjectedMemoryBase + (InstancedOb->GetMemoryOffset() - GetMemoryOffset());
+				InstancedOb->CleanupMemory(OwnerComp, InjectedNodeMemory, CleanupType);
+			}
+		}
+	}
+}
+
 #if WITH_EDITOR
 
 FName UBTTask_RunBehavior::GetNodeIconName() const
diff --git a/Engine/Source/Runtime/AIModule/Public/BlueprintNodeHelpers.h b/Engine/Source/Runtime/AIModule/Public/BlueprintNodeHelpers.h
--- a/Engine/Source/Runtime/AIModule/Public/BlueprintNodeHelpers.h
+++ b/Engine/Source/Runtime/AIModule/Public/BlueprintNodeHelpers.h
@@ -18,7 +18,7 @@ namespace BlueprintNodeHelpers
 
 	bool FindNodeOwner(AActor* OwningActor, class UBTNode* Node, UBehaviorTreeComponent*& OwningComp, int32& OwningInstanceIdx);
 
-	void AbortLatentActions(AActor* OwningActor, const UObject* Ob);
+	void AbortLatentActions(UActorComponent* OwnerOb, const UObject* Ob);
 
 	FORCEINLINE bool HasBlueprintFunction(FName FuncName, const UObject* Ob, const UClass* StopAtClass)
 	{
diff --git a/Engine/Source/Runtime/Engine/Public/AI/BlueprintNodeHelpers.h b/Engine/Source/Runtime/Engine/Public/AI/BlueprintNodeHelpers.h
--- a/Engine/Source/Runtime/Engine/Public/AI/BlueprintNodeHelpers.h
+++ b/Engine/Source/Runtime/Engine/Public/AI/BlueprintNodeHelpers.h
@@ -18,7 +18,7 @@ namespace BlueprintNodeHelpers
 
 	bool FindNodeOwner(AActor* OwningActor, class UBTNode* Node, UBehaviorTreeComponent*& OwningComp, int32& OwningInstanceIdx);
 
-	void AbortLatentActions(AActor* OwningActor, const UObject* Ob);
+	void AbortLatentActions(UActorComponent* OwnerOb, const UObject* Ob);
 
 	FORCEINLINE bool HasBlueprintFunction(FName FuncName, const UObject* Ob, const UClass* StopAtClass)
 	{
