Fix mismatching guards (caused shipping build breakage).

diff --git a/Engine/Source/Runtime/AIModule/Classes/EnvironmentQuery/EnvQueryGenerator.h b/Engine/Source/Runtime/AIModule/Classes/EnvironmentQuery/EnvQueryGenerator.h
--- a/Engine/Source/Runtime/AIModule/Classes/EnvironmentQuery/EnvQueryGenerator.h
+++ b/Engine/Source/Runtime/AIModule/Classes/EnvironmentQuery/EnvQueryGenerator.h
@@ -21,7 +21,7 @@ class AIMODULE_API UEnvQueryGenerator : public UObject
 	virtual FText GetDescriptionTitle() const;
 	virtual FText GetDescriptionDetails() const;
 
-#if WITH_EDITOR
+#if WITH_EDITOR && USE_EQS_DEBUGGER
 	virtual void PostEditChangeProperty(struct FPropertyChangedEvent& PropertyChangedEvent) override;
-#endif //WITH_EDITOR
+#endif //WITH_EDITOR && USE_EQS_DEBUGGER
 };
diff --git a/Engine/Source/Runtime/AIModule/Classes/EnvironmentQuery/EnvQueryTest.h b/Engine/Source/Runtime/AIModule/Classes/EnvironmentQuery/EnvQueryTest.h
--- a/Engine/Source/Runtime/AIModule/Classes/EnvironmentQuery/EnvQueryTest.h
+++ b/Engine/Source/Runtime/AIModule/Classes/EnvironmentQuery/EnvQueryTest.h
@@ -152,9 +152,9 @@ public:
 
 	virtual void PostLoad() override;
 
-#if WITH_EDITOR
+#if WITH_EDITOR && USE_EQS_DEBUGGER
 	virtual void PostEditChangeProperty(FPropertyChangedEvent& PropertyChangedEvent) override;
-#endif //WITH_EDITOR
+#endif //WITH_EDITOR && USE_EQS_DEBUGGER
 };
 
 //////////////////////////////////////////////////////////////////////////
diff --git a/Engine/Source/Runtime/AIModule/Private/EnvironmentQuery/EnvQueryGenerator.cpp b/Engine/Source/Runtime/AIModule/Private/EnvironmentQuery/EnvQueryGenerator.cpp
--- a/Engine/Source/Runtime/AIModule/Private/EnvironmentQuery/EnvQueryGenerator.cpp
+++ b/Engine/Source/Runtime/AIModule/Private/EnvironmentQuery/EnvQueryGenerator.cpp
@@ -18,10 +18,10 @@ FText UEnvQueryGenerator::GetDescriptionDetails() const
 	return FText::GetEmpty();
 }
 
-#if WITH_EDITOR
+#if WITH_EDITOR && USE_EQS_DEBUGGER
 void UEnvQueryGenerator::PostEditChangeProperty(FPropertyChangedEvent& PropertyChangedEvent) 
 {
 	Super::PostEditChangeProperty(PropertyChangedEvent);
 	UEnvQueryManager::NotifyAssetUpdate(NULL);
 }
-#endif //WITH_EDITOR
+#endif //WITH_EDITOR && USE_EQS_DEBUGGER
diff --git a/Engine/Source/Runtime/AIModule/Private/EnvironmentQuery/EnvQueryManager.cpp b/Engine/Source/Runtime/AIModule/Private/EnvironmentQuery/EnvQueryManager.cpp
--- a/Engine/Source/Runtime/AIModule/Private/EnvironmentQuery/EnvQueryManager.cpp
+++ b/Engine/Source/Runtime/AIModule/Private/EnvironmentQuery/EnvQueryManager.cpp
@@ -104,7 +104,7 @@ UEnvQueryManager* UEnvQueryManager::GetCurrent(UObject* WorldContextObject)
 	return World ? UAISystem::GetCurrentEQSManager(World) : NULL;
 }
 
-#if WITH_EDITOR
+#if USE_EQS_DEBUGGER
 void UEnvQueryManager::NotifyAssetUpdate(UEnvQuery* Query)
 {
 	if (GEditor == NULL)
@@ -138,7 +138,7 @@ void UEnvQueryManager::NotifyAssetUpdate(UEnvQuery* Query)
 		}
 	}
 }
-#endif // WITH_EDITOR
+#endif // USE_EQS_DEBUGGER
 
 TStatId UEnvQueryManager::GetStatId() const
 {
diff --git a/Engine/Source/Runtime/AIModule/Private/EnvironmentQuery/EnvQueryTest.cpp b/Engine/Source/Runtime/AIModule/Private/EnvironmentQuery/EnvQueryTest.cpp
--- a/Engine/Source/Runtime/AIModule/Private/EnvironmentQuery/EnvQueryTest.cpp
+++ b/Engine/Source/Runtime/AIModule/Private/EnvironmentQuery/EnvQueryTest.cpp
@@ -298,13 +298,13 @@ void UEnvQueryTest::PostLoad()
 	bFormatUpdated = true;
 }
 
-#if WITH_EDITOR
+#if WITH_EDITOR && USE_EQS_DEBUGGER
 void UEnvQueryTest::PostEditChangeProperty(FPropertyChangedEvent& PropertyChangedEvent) 
 {
 	Super::PostEditChangeProperty(PropertyChangedEvent);
 	UEnvQueryManager::NotifyAssetUpdate(NULL);
 }
-#endif //WITH_EDITOR
+#endif //WITH_EDITOR && USE_EQS_DEBUGGER
 
 FText UEnvQueryTest::DescribeFloatTestParams() const
 {
