Fixing CIS after a previous change.

diff --git a/Engine/Source/Runtime/AIModule/Classes/EnvironmentQuery/EnvQueryManager.h b/Engine/Source/Runtime/AIModule/Classes/EnvironmentQuery/EnvQueryManager.h
--- a/Engine/Source/Runtime/AIModule/Classes/EnvironmentQuery/EnvQueryManager.h
+++ b/Engine/Source/Runtime/AIModule/Classes/EnvironmentQuery/EnvQueryManager.h
@@ -136,14 +136,14 @@ class AIMODULE_API UEnvQueryManager : public UObject, public FTickableGameObject
 	static UEnvQueryManager* GetCurrent(UWorld* World);
 	static UEnvQueryManager* GetCurrent(UObject* WorldContextObject);
 
-#if USE_EQS_DEBUGGER
+#if WITH_EDITOR && USE_EQS_DEBUGGER
 	static void NotifyAssetUpdate(UEnvQuery* Query);
 
 	FEQSDebugger& GetDebugger() { return EQSDebugger; }
 
 protected:
 	FEQSDebugger EQSDebugger;
-#endif // USE_EQS_DEBUGGER
+#endif // WITH_EDITOR && USE_EQS_DEBUGGER
 
 protected:
 
diff --git a/Engine/Source/Runtime/AIModule/Private/EnvironmentQuery/EnvQueryManager.cpp b/Engine/Source/Runtime/AIModule/Private/EnvironmentQuery/EnvQueryManager.cpp
--- a/Engine/Source/Runtime/AIModule/Private/EnvironmentQuery/EnvQueryManager.cpp
+++ b/Engine/Source/Runtime/AIModule/Private/EnvironmentQuery/EnvQueryManager.cpp
@@ -104,7 +104,7 @@ UEnvQueryManager* UEnvQueryManager::GetCurrent(UObject* WorldContextObject)
 	return World ? UAISystem::GetCurrentEQSManager(World) : NULL;
 }
 
-#if USE_EQS_DEBUGGER
+#if WITH_EDITOR && USE_EQS_DEBUGGER
 void UEnvQueryManager::NotifyAssetUpdate(UEnvQuery* Query)
 {
 	if (GEditor == NULL)
@@ -138,7 +138,7 @@ void UEnvQueryManager::NotifyAssetUpdate(UEnvQuery* Query)
 		}
 	}
 }
-#endif // USE_EQS_DEBUGGER
+#endif // WITH_EDITOR && USE_EQS_DEBUGGER
 
 TStatId UEnvQueryManager::GetStatId() const
 {
