Hopefully the proper CIS fix as much as EQS is concerned.

diff --git a/Engine/Source/Runtime/AIModule/Classes/EnvironmentQuery/EnvQueryManager.h b/Engine/Source/Runtime/AIModule/Classes/EnvironmentQuery/EnvQueryManager.h
--- a/Engine/Source/Runtime/AIModule/Classes/EnvironmentQuery/EnvQueryManager.h
+++ b/Engine/Source/Runtime/AIModule/Classes/EnvironmentQuery/EnvQueryManager.h
@@ -136,8 +136,10 @@ class AIMODULE_API UEnvQueryManager : public UObject, public FTickableGameObject
 	static UEnvQueryManager* GetCurrent(UWorld* World);
 	static UEnvQueryManager* GetCurrent(UObject* WorldContextObject);
 
-#if WITH_EDITOR && USE_EQS_DEBUGGER
+#if USE_EQS_DEBUGGER
+#if WITH_EDITOR
 	static void NotifyAssetUpdate(UEnvQuery* Query);
+#endif //WITH_EDITOR
 
 	FEQSDebugger& GetDebugger() { return EQSDebugger; }
 
