fixed crash in behavior tree's CleanupMemory
ttp# 344064

diff --git a/Engine/Source/Runtime/AIModule/Private/BehaviorTree/Tasks/BTTask_RunBehavior.cpp b/Engine/Source/Runtime/AIModule/Private/BehaviorTree/Tasks/BTTask_RunBehavior.cpp
--- a/Engine/Source/Runtime/AIModule/Private/BehaviorTree/Tasks/BTTask_RunBehavior.cpp
+++ b/Engine/Source/Runtime/AIModule/Private/BehaviorTree/Tasks/BTTask_RunBehavior.cpp
@@ -244,7 +244,7 @@ void UBTTask_RunBehavior::CleanupMemory(class UBehaviorTreeComponent* OwnerComp,
 {
 	Super::CleanupMemory(OwnerComp, NodeMemory, CleanupType);
 
-	if (GetParentNode() && BehaviorAsset)
+	if (GetParentNode() && OwnerComp && BehaviorAsset)
 	{
 		FBTInstancedNodeMemory* NodeMemoryHeader = (FBTInstancedNodeMemory*)NodeMemory;
 		uint8* InjectedMemoryBase = NodeMemory + sizeof(FBTInstancedNodeMemory);
@@ -253,7 +253,9 @@ void UBTTask_RunBehavior::CleanupMemory(class UBehaviorTreeComponent* OwnerComp,
 		for (int32 Idx = 0; Idx < NumInjectedDecorators; Idx++)
 		{
 			const int32 InstancedNodeIdx = NodeMemoryHeader->NodeIdx + Idx;
-			UBTDecorator* InstancedOb = Cast<UBTDecorator>(OwnerComp->NodeInstances[InstancedNodeIdx]);
+			UBTDecorator* InstancedOb = OwnerComp->NodeInstances.IsValidIndex(InstancedNodeIdx) ?
+				Cast<UBTDecorator>(OwnerComp->NodeInstances[InstancedNodeIdx]) : NULL;
+
 			if (InstancedOb)
 			{
 				uint8* InjectedNodeMemory = InjectedMemoryBase + (InstancedOb->GetMemoryOffset() - GetMemoryOffset());
@@ -270,4 +272,4 @@ FName UBTTask_RunBehavior::GetNodeIconName() const
 	return FName("BTEditor.Graph.BTNode.Task.RunBehavior.Icon");
 }
 
-#endif	// WITH_EDITOR
\ No newline at end of file
+#endif	// WITH_EDITOR
