Made AIController::PerceptionComponent optional rather then always instanced #UE4

diff --git a/Engine/Source/Runtime/AIModule/Classes/AIController.h b/Engine/Source/Runtime/AIModule/Classes/AIController.h
--- a/Engine/Source/Runtime/AIModule/Classes/AIController.h
+++ b/Engine/Source/Runtime/AIModule/Classes/AIController.h
@@ -110,7 +110,7 @@ public:
 	UBrainComponent* BrainComponent;
 
 	UPROPERTY()
-	TSubobjectPtr<UAIPerceptionComponent> PerceptionComponent;
+	UAIPerceptionComponent* PerceptionComponent;
 
 	UFUNCTION(BlueprintImplementableEvent, Category = "AI")
 	void OnPossess(APawn* PossessedPawn);
diff --git a/Engine/Source/Runtime/AIModule/Private/AIController.cpp b/Engine/Source/Runtime/AIModule/Private/AIController.cpp
--- a/Engine/Source/Runtime/AIModule/Private/AIController.cpp
+++ b/Engine/Source/Runtime/AIModule/Private/AIController.cpp
@@ -37,15 +37,12 @@ DEFINE_LOG_CATEGORY(LogAINavigation);
 AAIController::AAIController(const class FPostConstructInitializeProperties& PCIP)
 	: Super(PCIP)
 {
-
 	// set up navigation component
 	NavComponent = PCIP.CreateDefaultSubobject<UNavigationComponent>(this, TEXT("NavComponent"));
 
 	PathFollowingComponent = PCIP.CreateDefaultSubobject<UPathFollowingComponent>(this, TEXT("PathFollowingComponent"));
 	PathFollowingComponent->OnMoveFinished.AddUObject(this, &AAIController::OnMoveCompleted);
 
-	PerceptionComponent = PCIP.CreateDefaultSubobject<UAIPerceptionComponent>(this, "PerceptionComponent");
-
 	bSkipExtraLOSChecks = true;
 	bWantsPlayerState = false;
 }
@@ -259,7 +256,7 @@ void AAIController::ClearFocus(uint8 InPriority)
 
 bool AAIController::LineOfSightTo(const AActor* Other, FVector ViewPoint, bool bAlternateChecks) const
 {
-	if( !Other )
+	if (!Other)
 	{
 		return false;
 	}
@@ -295,13 +292,13 @@ bool AAIController::LineOfSightTo(const AActor* Other, FVector ViewPoint, bool b
 	}
 
 	const FVector OtherActorLocation = Other->GetActorLocation();
-	float distSq = (OtherActorLocation - ViewPoint).SizeSquared();
-	if ( distSq > FARSIGHTTHRESHOLDSQUARED )
+	const float DistSq = (OtherActorLocation - ViewPoint).SizeSquared();
+	if ( DistSq > FARSIGHTTHRESHOLDSQUARED )
 	{
 		return false;
 	}
 
-	if ( !OtherPawn && (distSq > NEARSIGHTTHRESHOLDSQUARED) ) 
+	if ( !OtherPawn && (DistSq > NEARSIGHTTHRESHOLDSQUARED) ) 
 	{
 		return false;
 	}
@@ -332,28 +329,28 @@ bool AAIController::LineOfSightTo(const AActor* Other, FVector ViewPoint, bool b
 		Points[1] = OtherActorLocation + FVector(OtherRadius, OtherRadius, 0);
 		Points[2] = OtherActorLocation - FVector(OtherRadius, OtherRadius, 0);
 		Points[3] = OtherActorLocation + FVector(OtherRadius, -1 * OtherRadius, 0);
-		int32 imin = 0;
-		int32 imax = 0;
-		float currentmin = (Points[0] - ViewPoint).SizeSquared();
-		float currentmax = currentmin;
+		int32 IndexMin = 0;
+		int32 IndexMax = 0;
+		float CurrentMax = (Points[0] - ViewPoint).SizeSquared();
+		float CurrentMin = CurrentMax;
 		for ( int32 PointIndex=1; PointIndex<4; PointIndex++ )
 		{
-			float nextsize = (Points[PointIndex] - ViewPoint).SizeSquared();
-			if (nextsize > currentmax)
+			const float NextSize = (Points[PointIndex] - ViewPoint).SizeSquared();
+			if (NextSize > CurrentMin)
 			{
-				currentmax = nextsize;
-				imax = PointIndex;
+				CurrentMin = NextSize;
+				IndexMax = PointIndex;
 			}
-			else if (nextsize < currentmin)
+			else if (NextSize < CurrentMax)
 			{
-				currentmin = nextsize;
-				imin = PointIndex;
+				CurrentMax = NextSize;
+				IndexMin = PointIndex;
 			}
 		}
 
-		for ( int32 PointIndex=0; PointIndex<4; PointIndex++ )
+		for (int32 PointIndex=0; PointIndex<4; PointIndex++)
 		{
-			if	( (PointIndex != imin) && (PointIndex != imax) )
+			if ((PointIndex != IndexMin) && (PointIndex != IndexMax))
 			{
 				bHit = GetWorld()->LineTraceTest(ViewPoint,  Points[PointIndex], ECC_Visibility, CollisionParams);
 				if ( !bHit )
