Fixed when to draw or not EQS from visual logger #ue4

diff --git a/Engine/Source/Developer/LogVisualizer/Private/SLogVisualizer.cpp b/Engine/Source/Developer/LogVisualizer/Private/SLogVisualizer.cpp
--- a/Engine/Source/Developer/LogVisualizer/Private/SLogVisualizer.cpp
+++ b/Engine/Source/Developer/LogVisualizer/Private/SLogVisualizer.cpp
@@ -2197,12 +2197,17 @@ void SLogVisualizer::DrawOnCanvas(UCanvas* Canvas, APlayerController*)
 			{
 				const FName TagName = CurrentData.TagName;
 				const bool bIsValidByFilter = FilterListPtr->IsFilterEnabled(CurrentData.Category.ToString(), ELogVerbosity::All) && FilterListPtr->IsFilterEnabled(CurrentData.TagName.ToString(), ELogVerbosity::All);
-				if (!bIsValidByFilter)
+				FVisualLogExtensionInterface* Extension = FVisualLog::Get().GetExtensionForTag(TagName);
+				if (!Extension)
 				{
 					continue;
 				}
 
-				if (FVisualLogExtensionInterface* Extension = FVisualLog::Get().GetExtensionForTag(TagName))
+				if (!bIsValidByFilter)
+				{
+					Extension->DisableDrawingForData(World, Canvas, HelperActor, TagName, CurrentData, Entry->TimeStamp);
+				}
+				else
 				{
 					Extension->DrawData(World, Canvas, HelperActor, TagName, CurrentData, Entry->TimeStamp);
 				}
diff --git a/Engine/Source/Runtime/AIModule/Classes/VisualLoggerExtension.h b/Engine/Source/Runtime/AIModule/Classes/VisualLoggerExtension.h
--- a/Engine/Source/Runtime/AIModule/Classes/VisualLoggerExtension.h
+++ b/Engine/Source/Runtime/AIModule/Classes/VisualLoggerExtension.h
@@ -17,6 +17,10 @@ public:
 	FVisualLoggerExtension();
 	virtual void OnTimestampChange(float Timestamp, class UWorld* InWorld, class AActor* HelperActor);
 	virtual void DrawData(class UWorld* InWorld, class UCanvas* Canvas, class AActor* HelperActor, const FName& TagName, const struct FVisLogEntry::FDataBlock& DataBlock, float Timestamp);
+	virtual void DisableDrawingForData(class UWorld* InWorld, class UCanvas* Canvas, class AActor* HelperActor, const FName& TagName, const FVisLogEntry::FDataBlock& DataBlock, float Timestamp);
+
+private:
+	void DisableEQSRendering(class AActor* HelperActor);
 
 protected:
 	uint32 CachedEQSId;
diff --git a/Engine/Source/Runtime/AIModule/Private/VisualLoggerExtension.cpp b/Engine/Source/Runtime/AIModule/Private/VisualLoggerExtension.cpp
--- a/Engine/Source/Runtime/AIModule/Private/VisualLoggerExtension.cpp
+++ b/Engine/Source/Runtime/AIModule/Private/VisualLoggerExtension.cpp
@@ -17,6 +17,22 @@ void FVisualLoggerExtension::OnTimestampChange(float Timestamp, class UWorld* In
 	{
 		CurrentTimestamp = Timestamp;
 		CachedEQSId = INDEX_NONE;
+		DisableEQSRendering(HelperActor);
+	}
+}
+
+void FVisualLoggerExtension::DisableDrawingForData(class UWorld* InWorld, class UCanvas* Canvas, class AActor* HelperActor, const FName& TagName, const FVisLogEntry::FDataBlock& DataBlock, float Timestamp)
+{
+	if (TagName == *EVisLogTags::TAG_EQS && CurrentTimestamp == Timestamp)
+	{
+		DisableEQSRendering(HelperActor);
+	}
+}
+
+void FVisualLoggerExtension::DisableEQSRendering(class AActor* HelperActor)
+{
+	if (HelperActor)
+	{
 		UEQSRenderingComponent* EQSRenderComp = HelperActor->FindComponentByClass<UEQSRenderingComponent>();
 		if (EQSRenderComp)
 		{
@@ -30,29 +46,26 @@ void FVisualLoggerExtension::OnTimestampChange(float Timestamp, class UWorld* In
 
 void FVisualLoggerExtension::DrawData(class UWorld* InWorld, class UCanvas* Canvas, class AActor* HelperActor, const FName& TagName, const FVisLogEntry::FDataBlock& DataBlock, float Timestamp)
 {
-	if (TagName == *EVisLogTags::TAG_EQS)
+	if (TagName == *EVisLogTags::TAG_EQS && HelperActor)
 	{
-		if (HelperActor)
+		UEQSRenderingComponent* EQSRenderComp = HelperActor->FindComponentByClass<UEQSRenderingComponent>();
+		if (!EQSRenderComp)
 		{
-			UEQSRenderingComponent* EQSRenderComp = HelperActor->FindComponentByClass<UEQSRenderingComponent>();
-			if (!EQSRenderComp)
-			{
-				EQSRenderComp = ConstructObject<UEQSRenderingComponent>(UEQSRenderingComponent::StaticClass(), HelperActor);
-				EQSRenderComp->bDrawOnlyWhenSelected = false;
-				EQSRenderComp->RegisterComponent();
-				EQSRenderComp->SetHiddenInGame(true);
-			}
+			EQSRenderComp = ConstructObject<UEQSRenderingComponent>(UEQSRenderingComponent::StaticClass(), HelperActor);
+			EQSRenderComp->bDrawOnlyWhenSelected = false;
+			EQSRenderComp->RegisterComponent();
+			EQSRenderComp->SetHiddenInGame(true);
+		}
 
-			EQSDebug::FQueryData DebugData;
-			UEnvQueryDebugHelpers::BlobArrayToDebugData(DataBlock.Data, DebugData, false);
-			if (DebugData.Id != CachedEQSId || (EQSRenderComp && EQSRenderComp->bHiddenInGame))
-			{
-				CachedEQSId = DebugData.Id;
-				EQSRenderComp->DebugData = DebugData;
-				EQSRenderComp->Activate();
-				EQSRenderComp->SetHiddenInGame(false);
-				EQSRenderComp->MarkRenderStateDirty();
-			}
+		EQSDebug::FQueryData DebugData;
+		UEnvQueryDebugHelpers::BlobArrayToDebugData(DataBlock.Data, DebugData, false);
+		if (DebugData.Id != CachedEQSId || (EQSRenderComp && EQSRenderComp->bHiddenInGame))
+		{
+			CachedEQSId = DebugData.Id;
+			EQSRenderComp->DebugData = DebugData;
+			EQSRenderComp->Activate();
+			EQSRenderComp->SetHiddenInGame(false);
+			EQSRenderComp->MarkRenderStateDirty();
 		}
 	}
 }
diff --git a/Engine/Source/Runtime/Engine/Public/VisualLog.h b/Engine/Source/Runtime/Engine/Public/VisualLog.h
--- a/Engine/Source/Runtime/Engine/Public/VisualLog.h
+++ b/Engine/Source/Runtime/Engine/Public/VisualLog.h
@@ -321,6 +321,7 @@ struct FVisualLogExtensionInterface
 {
 	virtual void OnTimestampChange(float Timestamp, class UWorld* InWorld, class AActor* HelperActor) = 0;
 	virtual void DrawData(class UWorld* InWorld, class UCanvas* Canvas, class AActor* HelperActor, const FName& TagName, const FVisLogEntry::FDataBlock& DataBlock, float Timestamp) = 0;
+	virtual void DisableDrawingForData(class UWorld* InWorld, class UCanvas* Canvas, class AActor* HelperActor, const FName& TagName, const FVisLogEntry::FDataBlock& DataBlock, float Timestamp) = 0;
 };
 
 
