fixed cycles in blackboard's parent chain
ttp# 343978

diff --git a/Engine/Source/Runtime/AIModule/Classes/BehaviorTree/BlackboardData.h b/Engine/Source/Runtime/AIModule/Classes/BehaviorTree/BlackboardData.h
--- a/Engine/Source/Runtime/AIModule/Classes/BehaviorTree/BlackboardData.h
+++ b/Engine/Source/Runtime/AIModule/Classes/BehaviorTree/BlackboardData.h
@@ -115,4 +115,7 @@ protected:
 
 	/** updates parent key cache for editor */
 	void UpdateParentKeys();
+
+	/** check if parent chain contains given blackboard data */
+	bool HasParent(const UBlackboardData* TestParent) const;
 };
diff --git a/Engine/Source/Runtime/AIModule/Private/BehaviorTree/BlackboardData.cpp b/Engine/Source/Runtime/AIModule/Private/BehaviorTree/BlackboardData.cpp
--- a/Engine/Source/Runtime/AIModule/Private/BehaviorTree/BlackboardData.cpp
+++ b/Engine/Source/Runtime/AIModule/Private/BehaviorTree/BlackboardData.cpp
@@ -110,6 +110,15 @@ void UBlackboardData::PostEditChangeProperty(struct FPropertyChangedEvent& Prope
 
 	if (PropertyChangedEvent.Property && PropertyChangedEvent.Property->GetFName() == GET_MEMBER_NAME_CHECKED(UBlackboardData, Parent))
 	{
+		// look for cycles
+		if (Parent && Parent->HasParent(this))
+		{
+			UE_LOG(LogBehaviorTree, Warning, TEXT("Blackboard asset (%s) has (%s) in parent chain! Clearing value to avoid cycle."),
+				*GetNameSafe(Parent), *GetNameSafe(this));
+
+			Parent = NULL;
+		}
+
 		UpdateParentKeys();
 	}
 }
@@ -137,6 +146,16 @@ static bool ContainsKeyName(FName KeyName, const TArray<FBlackboardEntry>& Keys,
 	return false;
 }
 
+bool UBlackboardData::HasParent(const UBlackboardData* TestParent) const
+{
+	if (Parent == TestParent)
+	{
+		return true;
+	}
+
+	return Parent ? Parent->HasParent(TestParent) : false;
+}
+
 void UBlackboardData::UpdateParentKeys()
 {
 	if (Parent == this)
