Fixed freeze issue with AIMoveTo blueprint node. It fixes TTP #344320 "AI: Lag (freeze) after in game accumulates in average 10 000 objects".

diff --git a/Engine/Source/Runtime/AIModule/Classes/AISystem.h b/Engine/Source/Runtime/AIModule/Classes/AISystem.h
--- a/Engine/Source/Runtime/AIModule/Classes/AISystem.h
+++ b/Engine/Source/Runtime/AIModule/Classes/AISystem.h
@@ -10,6 +10,7 @@
 class UBehaviorTreeManager;
 class UEnvQueryManager;
 class UAIPerceptionSystem;
+class UAIAsyncTaskBlueprintProxy;
 
 UCLASS(config=Engine)
 class AIMODULE_API UAISystem : public UAISystemBase
@@ -28,6 +29,9 @@ protected:
 	UPROPERTY(Transient)
 	UAIPerceptionSystem* PerceptionSystem;
 
+	UPROPERTY(Transient)
+	TArray<class UAIAsyncTaskBlueprintProxy*> AllProxyObjects;
+
 public:
 	virtual ~UAISystem();
 
@@ -78,6 +82,10 @@ public:
 
 	virtual UWorld* GetWorld() const override { return GetOuterWorld(); }
 	
+	FORCEINLINE void AddReferenceFromProxyObject(UAIAsyncTaskBlueprintProxy* BlueprintProxy) { AllProxyObjects.AddUnique(BlueprintProxy); }
+
+	FORCEINLINE void RemoveReferenceToProxyObject(UAIAsyncTaskBlueprintProxy* BlueprintProxy) { AllProxyObjects.RemoveSwap(BlueprintProxy); }
+
 	//----------------------------------------------------------------------//
 	// cheats
 	//----------------------------------------------------------------------//
diff --git a/Engine/Source/Runtime/AIModule/Private/Blueprint/AIBlueprintHelperLibrary.cpp b/Engine/Source/Runtime/AIModule/Private/Blueprint/AIBlueprintHelperLibrary.cpp
--- a/Engine/Source/Runtime/AIModule/Private/Blueprint/AIBlueprintHelperLibrary.cpp
+++ b/Engine/Source/Runtime/AIModule/Private/Blueprint/AIBlueprintHelperLibrary.cpp
@@ -15,16 +15,22 @@
 UAIAsyncTaskBlueprintProxy::UAIAsyncTaskBlueprintProxy(const class FPostConstructInitializeProperties& PCIP)
 	: Super(PCIP)
 {
+	MyWorld = Cast<UWorld>(GetOuter());
+	if (HasAnyFlags(RF_ClassDefaultObject) == false)
+	{
+		UAISystem* const AISystem = MyWorld.IsValid() ? UAISystem::GetCurrent(MyWorld.Get()) : NULL;
+		if (AISystem)
+		{
+			AISystem->AddReferenceFromProxyObject(this);
+		}
+	}
 }
 
 void UAIAsyncTaskBlueprintProxy::OnMoveCompleted(FAIRequestID RequestID, EPathFollowingResult::Type MovementResult)
 {
-	if (RequestID.IsEquivalent(MoveRequestId) && AIController.IsValid())
+	if (RequestID.IsEquivalent(MoveRequestId) && AIController.IsValid(true))
 	{
-		if (!AIController->IsPendingKill() && AIController->ReceiveMoveCompleted.IsBound())
-		{
-			AIController->ReceiveMoveCompleted.RemoveDynamic(this, &UAIAsyncTaskBlueprintProxy::OnMoveCompleted);
-		}
+		AIController->ReceiveMoveCompleted.RemoveDynamic(this, &UAIAsyncTaskBlueprintProxy::OnMoveCompleted);
 
 		if (MovementResult == EPathFollowingResult::Success)
 		{
@@ -34,21 +40,32 @@ void UAIAsyncTaskBlueprintProxy::OnMoveCompleted(FAIRequestID RequestID, EPathFo
 		{
 			OnFail.Broadcast(MovementResult);
 		}
+
+		UAISystem* const AISystem = MyWorld.IsValid() ? UAISystem::GetCurrent(MyWorld.Get()) : NULL;
+		if (AISystem)
+		{
+			AISystem->RemoveReferenceToProxyObject(this);
+		}
 	}
 }
 
 void UAIAsyncTaskBlueprintProxy::OnNoPath()
 {
 	OnFail.Broadcast(EPathFollowingResult::Aborted);
+	UAISystem* const AISystem = MyWorld.IsValid() ? UAISystem::GetCurrent(MyWorld.Get()) : NULL;
+	if (AISystem)
+	{
+		AISystem->RemoveReferenceToProxyObject(this);
+	}
 }
 
 void UAIAsyncTaskBlueprintProxy::BeginDestroy()
 {
-	if (AIController.IsValid() && !AIController->IsPendingKill() && AIController->ReceiveMoveCompleted.IsBound())
+	UAISystem* const AISystem = MyWorld.IsValid() ? UAISystem::GetCurrent(MyWorld.Get()) : NULL;
+	if (AISystem)
 	{
-		AIController->ReceiveMoveCompleted.RemoveDynamic(this, &UAIAsyncTaskBlueprintProxy::OnMoveCompleted);
+		AISystem->RemoveReferenceToProxyObject(this);
 	}
-
 	Super::BeginDestroy();
 }
 
@@ -72,7 +89,8 @@ class UAIAsyncTaskBlueprintProxy* UAIBlueprintHelperLibrary::CreateMoveToProxyOb
 	AAIController* AIController = Cast<AAIController>(Pawn->GetController());
 	if (AIController)
 	{
-		MyObj = NewObject<UAIAsyncTaskBlueprintProxy>();
+		UWorld* World = GEngine->GetWorldFromContextObject( WorldContextObject );
+		MyObj = NewObject<UAIAsyncTaskBlueprintProxy>(World);
 		FNavPathSharedPtr Path = TargetActor ? AIController->FindPath(TargetActor, true) : AIController->FindPath(Destination, true);
 		if (Path.IsValid())
 		{
@@ -82,7 +100,6 @@ class UAIAsyncTaskBlueprintProxy* UAIBlueprintHelperLibrary::CreateMoveToProxyOb
 		}
 		else
 		{
-			UWorld* World = GEngine->GetWorldFromContextObject( WorldContextObject );
 			World->GetTimerManager().SetTimer(MyObj, &UAIAsyncTaskBlueprintProxy::OnNoPath, 0.1, false);
 		}
 	}
