Fixed CIS compilation - forgot about one file.

diff --git a/Engine/Source/Runtime/AIModule/Classes/Blueprint/AIAsyncTaskBlueprintProxy.h b/Engine/Source/Runtime/AIModule/Classes/Blueprint/AIAsyncTaskBlueprintProxy.h
--- a/Engine/Source/Runtime/AIModule/Classes/Blueprint/AIAsyncTaskBlueprintProxy.h
+++ b/Engine/Source/Runtime/AIModule/Classes/Blueprint/AIAsyncTaskBlueprintProxy.h
@@ -5,6 +5,9 @@
 #include "Navigation/PathFollowingComponent.h"
 #include "AIAsyncTaskBlueprintProxy.generated.h"
 
+class UWorld;
+class AAIController;
+
 DECLARE_DYNAMIC_MULTICAST_DELEGATE_OneParam(FOAISimpleDelegate, EPathFollowingResult::Type, MovementResult);
 
 UCLASS(MinimalAPI)
@@ -28,6 +31,7 @@ public:
 	virtual void BeginDestroy() override;
 	// End UObject interface
 
-	TWeakObjectPtr<class AAIController> AIController;
+	TWeakObjectPtr<AAIController> AIController;
 	FAIRequestID MoveRequestId;
+	TWeakObjectPtr<UWorld> MyWorld;
 };
