Enhanced debug information data for single item in EQS Debugger (GameplayDebugger feature) #ue4
- Fixed few compilation issues with disabled USE_EQS_DEBUGGER flag
- Fixed crash in EQSRenderingComponent
- Fixed EQS debug data for sorted EQS itesm (it's slower way to sort items but only with active USE_EQS_DEBUGGER flag)
Change integrated from Dev branch, from CL #2282639

diff --git a/Engine/Source/Developer/GameplayDebugger/Private/GameplayDebuggingComponent.cpp b/Engine/Source/Developer/GameplayDebugger/Private/GameplayDebuggingComponent.cpp
--- a/Engine/Source/Developer/GameplayDebugger/Private/GameplayDebuggingComponent.cpp
+++ b/Engine/Source/Developer/GameplayDebugger/Private/GameplayDebuggingComponent.cpp
@@ -5,6 +5,7 @@
 #include "AI/Navigation/RecastNavMeshGenerator.h"
 #include "NavMeshRenderingHelpers.h"
 #include "AI/Navigation/NavigationSystem.h"
+#include "EnvironmentQuery/EnvQueryTypes.h"
 #include "EnvironmentQuery/EnvQueryManager.h"
 #include "EnvironmentQuery/EQSRenderingComponent.h"
 #include "EnvironmentQuery/EnvQueryTest.h"
diff --git a/Engine/Source/Developer/GameplayDebugger/Private/GameplayDebuggingHUDComponent.cpp b/Engine/Source/Developer/GameplayDebugger/Private/GameplayDebuggingHUDComponent.cpp
--- a/Engine/Source/Developer/GameplayDebugger/Private/GameplayDebuggingHUDComponent.cpp
+++ b/Engine/Source/Developer/GameplayDebugger/Private/GameplayDebuggingHUDComponent.cpp
@@ -381,19 +381,84 @@ void AGameplayDebuggingHUDComponent::DrawEQSData(APlayerController* PC, class UG
 	PrintString(DefaultContext, TEXT("\n"));
 
 	auto& CurrentLocalData = DebugComponent->EQSLocalData[DebugComponent->CurrentEQSIndex];
-	if (CurrentLocalData.NumValidItems > 0)
+
+	/** find and draw item selection */
+	int32 BestItemIndex = INDEX_NONE;
+	{
+		APlayerController* const MyPC = Cast<APlayerController>(PlayerOwner);
+		FVector CamLocation;
+		FRotator CamRotation;
+		check(MyPC->PlayerCameraManager != NULL);
+		MyPC->PlayerCameraManager->GetCameraViewPoint(CamLocation, CamRotation);
+		FVector FireDir = CamRotation.Vector();
+
+		FVector out_Location;
+		FRotator out_Rotation;
+		MyPC->GetPlayerViewPoint(out_Location, out_Rotation);
+
+		float bestAim = 0;
+		for (int32 Index = 0; Index < CurrentLocalData.RenderDebugHelpers.Num(); ++Index)
+		{
+			auto& CurrentItem = CurrentLocalData.RenderDebugHelpers[Index];
+
+			// look for best controlled pawn target
+			const FVector AimDir = CurrentItem.Location - out_Location;
+			float FireDist = AimDir.SizeSquared();
+			// only find targets which are < 25000 units away
+
+			FireDist = FMath::Sqrt(FireDist);
+			float newAim = FireDir | AimDir;
+			newAim = newAim / FireDist;
+			if (newAim > bestAim)
+			{
+				BestItemIndex = Index;
+				bestAim = newAim;
+			}
+		}
+
+		if (BestItemIndex != INDEX_NONE)
+		{
+			DrawDebugSphere(World, CurrentLocalData.RenderDebugHelpers[BestItemIndex].Location, CurrentLocalData.RenderDebugHelpers[BestItemIndex].Radius, 8, FColor::Red, false);
+			int32 FailedTestIndex = CurrentLocalData.RenderDebugHelpers[BestItemIndex].FailedTestIndex;
+			if (FailedTestIndex != INDEX_NONE)
+			{
+				PrintString(DefaultContext, FString::Printf(TEXT("{red}Selected item failed with test %d: {yellow}%s {LightBlue}(%s)\n")
+					, FailedTestIndex
+					, *CurrentLocalData.Tests[FailedTestIndex].ShortName
+					, *CurrentLocalData.Tests[FailedTestIndex].Detailed
+					));
+				PrintString(DefaultContext, FString::Printf(TEXT("{white}'%s' with score %3.3f\n\n"), *CurrentLocalData.RenderDebugHelpers[BestItemIndex].AdditionalInformation, CurrentLocalData.RenderDebugHelpers[BestItemIndex].FailedScore));
+			}
+		}
+	}
+
+	PrintString(DefaultContext, FString::Printf(TEXT("{white}Timestamp: {yellow}%.3f (%.2fs ago)\n")
+		, CurrentLocalData.Timestamp, PC->GetWorld()->GetTimeSeconds() - CurrentLocalData.Timestamp
+		));
+	PrintString(DefaultContext, FString::Printf(TEXT("{white}Query ID: {yellow}%d\n")
+		, CurrentLocalData.Id
+		));
+
+	PrintString(DefaultContext, FString::Printf(TEXT("{white}Query contains %d options: "), CurrentLocalData.Options.Num()));
+	for (int32 OptionIndex = 0; OptionIndex < CurrentLocalData.Options.Num(); ++OptionIndex)
 	{
-		PrintString(DefaultContext, FString::Printf(TEXT("{white}Timestamp: {yellow}%.3f (%.2fs ago)\n")
-			, CurrentLocalData.Timestamp, PC->GetWorld()->GetTimeSeconds() - CurrentLocalData.Timestamp
-			));
-		PrintString(DefaultContext, FString::Printf(TEXT("{white}Query ID: {yellow}%d\n")
-			, CurrentLocalData.Id
-			));
+		if (OptionIndex == CurrentLocalData.UsedOption)
+		{
+			PrintString(DefaultContext, FString::Printf(TEXT("{green}%s, "), *CurrentLocalData.Options[OptionIndex]));
+		}
+		else
+		{
+			PrintString(DefaultContext, FString::Printf(TEXT("{yellow}%s, "), *CurrentLocalData.Options[OptionIndex]));
+		}
+	}
+	PrintString(DefaultContext, TEXT("\n"));
 
+	const float RowHeight = 20.0f;
+	const int32 NumTests = CurrentLocalData.Tests.Num();
+	if (CurrentLocalData.NumValidItems > 0)
+	{
 		// draw test weights for best X items
 		const int32 NumItems = CurrentLocalData.Items.Num();
-		const int32 NumTests = CurrentLocalData.Tests.Num();
-		const float RowHeight = 20.0f;
 
 		FCanvasTileItem TileItem(FVector2D(0, 0), GWhiteTexture, FVector2D(Canvas->SizeX, RowHeight), FLinearColor::Black);
 		FLinearColor ColorOdd(0, 0, 0, 0.6f);
@@ -432,18 +497,20 @@ void AGameplayDebuggingHUDComponent::DrawEQSData(APlayerController* PC, class UG
 			DrawEQSItemDetails(Idx, DebugComponent);
 			DefaultContext.CursorY += RowHeight;
 		}
-
-		// test description
 		DefaultContext.CursorY += RowHeight;
-		for (int32 TestIdx = 0; TestIdx < NumTests; TestIdx++)
-		{
-			FString TestDesc = FString::Printf(TEXT("{white}Test %d = {yellow}%s {LightBlue}(%s)\n"), TestIdx,
-				*CurrentLocalData.Tests[TestIdx].ShortName,
-				*CurrentLocalData.Tests[TestIdx].Detailed);
+	}
 
-			PrintString(DefaultContext, TestDesc);
-		}
+	// test description
+	PrintString(DefaultContext, TEXT("All tests from used option:\n"));
+	for (int32 TestIdx = 0; TestIdx < NumTests; TestIdx++)
+	{
+		FString TestDesc = FString::Printf(TEXT("{white}Test %d = {yellow}%s {LightBlue}(%s)\n"), TestIdx,
+			*CurrentLocalData.Tests[TestIdx].ShortName,
+			*CurrentLocalData.Tests[TestIdx].Detailed);
+
+		PrintString(DefaultContext, TestDesc);
 	}
+
 #endif //!(UE_BUILD_SHIPPING || UE_BUILD_TEST)
 }
 
diff --git a/Engine/Source/Editor/EnvironmentQueryEditor/Private/EnvironmentQueryGraph.cpp b/Engine/Source/Editor/EnvironmentQueryEditor/Private/EnvironmentQueryGraph.cpp
--- a/Engine/Source/Editor/EnvironmentQueryEditor/Private/EnvironmentQueryGraph.cpp
+++ b/Engine/Source/Editor/EnvironmentQueryEditor/Private/EnvironmentQueryGraph.cpp
@@ -86,7 +86,9 @@ void UEnvironmentQueryGraph::UpdateAsset()
 	}
 
 	RemoveOrphanedNodes();
+#if USE_EQS_DEBUGGER
 	UEnvQueryManager::NotifyAssetUpdate(Query);
+#endif
 }
 
 void UEnvironmentQueryGraph::CalculateAllWeights()
diff --git a/Engine/Source/Runtime/AIModule/Classes/EnvironmentQuery/EQSRenderingComponent.h b/Engine/Source/Runtime/AIModule/Classes/EnvironmentQuery/EQSRenderingComponent.h
--- a/Engine/Source/Runtime/AIModule/Classes/EnvironmentQuery/EQSRenderingComponent.h
+++ b/Engine/Source/Runtime/AIModule/Classes/EnvironmentQuery/EQSRenderingComponent.h
@@ -15,11 +15,12 @@ public:
 
 	virtual void DrawDebugLabels(UCanvas* Canvas, APlayerController*) override;
 	
-	//virtual void DrawStaticElements(FStaticPrimitiveDrawInterface* PDI) override;
 	virtual FPrimitiveViewRelevance GetViewRelevance(const FSceneView* View) override;
 
-	static void CollectEQSData(const UPrimitiveComponent* InComponent, const class IEQSQueryResultSourceInterface* QueryDataSource, TArray<FSphere>& Spheres, TArray<FText3d>& Texts);
-	static void CollectEQSData(const struct FEnvQueryResult* ResultItems, const struct FEnvQueryInstance* QueryInstance, TArray<FSphere>& Spheres, TArray<FText3d>& Texts, bool ShouldDrawFailedItems);
+#if  USE_EQS_DEBUGGER 
+	static void CollectEQSData(const UPrimitiveComponent* InComponent, const class IEQSQueryResultSourceInterface* QueryDataSource, TArray<FSphere>& Spheres, TArray<FText3d>& Texts, TArray<EQSDebug::FDebugHelper>& DebugItems);
+	static void CollectEQSData(const struct FEnvQueryResult* ResultItems, const struct FEnvQueryInstance* QueryInstance, TArray<FSphere>& Spheres, TArray<FText3d>& Texts, bool ShouldDrawFailedItems, TArray<EQSDebug::FDebugHelper>& DebugItems);
+#endif
 private:
 	FEnvQueryResult QueryResult;	
 	// can be 0
diff --git a/Engine/Source/Runtime/AIModule/Classes/EnvironmentQuery/EnvQueryDebugHelpers.h b/Engine/Source/Runtime/AIModule/Classes/EnvironmentQuery/EnvQueryDebugHelpers.h
--- a/Engine/Source/Runtime/AIModule/Classes/EnvironmentQuery/EnvQueryDebugHelpers.h
+++ b/Engine/Source/Runtime/AIModule/Classes/EnvironmentQuery/EnvQueryDebugHelpers.h
@@ -1,6 +1,7 @@
 // Copyright 1998-2014 Epic Games, Inc. All Rights Reserved.
 
 #pragma once
+#include "EnvQueryTypes.h"
 #include "DebugRenderSceneProxy.h"
 #include "VisualLoggerExtension.h"
 #include "EnvQueryDebugHelpers.generated.h"
@@ -25,12 +26,29 @@ namespace EQSDebug
 		FString Detailed;
 	};
 
+	// struct filled while collecting data (to store additional debug data needed to display per rendered item)
+	struct FDebugHelper
+	{
+		FDebugHelper() : Location(FVector::ZeroVector), Radius(0), FailedTestIndex(INDEX_NONE){}
+		FDebugHelper(FVector Loc, float R) : Location(Loc), Radius(R), FailedTestIndex(INDEX_NONE) {}
+		FDebugHelper(FVector Loc, float R, const FString& Desc) : Location(Loc), Radius(R), FailedTestIndex(INDEX_NONE), AdditionalInformation(Desc) {}
+
+		FVector Location;
+		float Radius;
+		int32 FailedTestIndex;
+		float FailedScore;
+		FString AdditionalInformation;
+	};
+
 	struct FQueryData
 	{
 		TArray<FItemData> Items;
 		TArray<FTestData> Tests;
 		TArray<FDebugRenderSceneProxy::FSphere> SolidSpheres;
 		TArray<FDebugRenderSceneProxy::FText3d> Texts;
+		TArray<FDebugHelper> RenderDebugHelpers;
+		TArray<FString>	Options;
+		int32 UsedOption;
 		int32 NumValidItems;
 		int32 Id;
 		FString Name;
@@ -38,6 +56,8 @@ namespace EQSDebug
 
 		void Reset()
 		{
+			UsedOption = 0;
+			Options.Reset();
 			NumValidItems = 0;
 			Id = INDEX_NONE;
 			Name.Empty();
@@ -46,6 +66,7 @@ namespace EQSDebug
 			SolidSpheres.Reset();
 			Texts.Reset();
 			Timestamp = 0;
+			RenderDebugHelpers.Reset();
 		}
 	};
 }
@@ -87,6 +108,16 @@ FArchive& operator<<(FArchive& Ar, EQSDebug::FTestData& Data)
 	return Ar;
 }
 
+FORCEINLINE
+FArchive& operator<<(FArchive& Ar, EQSDebug::FDebugHelper& Data)
+{
+	Ar << Data.Location;
+	Ar << Data.Radius;
+	Ar << Data.AdditionalInformation;
+	Ar << Data.FailedTestIndex;
+	return Ar;
+}
+
 FORCEINLINE
 FArchive& operator<<(FArchive& Ar, EQSDebug::FQueryData& Data)
 {
@@ -98,12 +129,15 @@ FArchive& operator<<(FArchive& Ar, EQSDebug::FQueryData& Data)
 	Ar << Data.Id;
 	Ar << Data.Name;
 	Ar << Data.Timestamp;
+	Ar << Data.RenderDebugHelpers;
+	Ar << Data.Options;
+	Ar << Data.UsedOption;
 	return Ar;
 }
 
 #endif //USE_EQS_DEBUGGER || ENABLE_VISUAL_LOG
 
-#if ENABLE_VISUAL_LOG
+#if ENABLE_VISUAL_LOG && USE_EQS_DEBUGGER
 
 #define UE_VLOG_EQS(Query, CategoryName, Verbosity) \
 { \
@@ -124,7 +158,7 @@ FArchive& operator<<(FArchive& Ar, EQSDebug::FQueryData& Data)
 #else
 #define UE_VLOG_EQS(Query, CategoryName, Verbosity)
 
-#endif //ENABLE_VISUAL_LOG
+#endif //ENABLE_VISUAL_LOG && USE_EQS_DEBUGGER
 
 UCLASS(Abstract, CustomConstructor)
 class AIMODULE_API UEnvQueryDebugHelpers : public UObject
@@ -132,7 +166,7 @@ class AIMODULE_API UEnvQueryDebugHelpers : public UObject
 	GENERATED_UCLASS_BODY()
 
 	UEnvQueryDebugHelpers(const class FPostConstructInitializeProperties& PCIP) : Super(PCIP) {}
-#if USE_EQS_DEBUGGER || ENABLE_VISUAL_LOG
+#if USE_EQS_DEBUGGER
 	static void QueryToDebugData(struct FEnvQueryInstance* Query, EQSDebug::FQueryData& EQSLocalData);
 	static void QueryToBlobArray(struct FEnvQueryInstance* Query, TArray<uint8>& BlobArray, bool bUseCompression = false);
 	static void BlobArrayToDebugData(const TArray<uint8>& BlobArray, EQSDebug::FQueryData& EQSLocalData, bool bUseCompression = false);
diff --git a/Engine/Source/Runtime/AIModule/Classes/EnvironmentQuery/EnvQueryTypes.h b/Engine/Source/Runtime/AIModule/Classes/EnvironmentQuery/EnvQueryTypes.h
--- a/Engine/Source/Runtime/AIModule/Classes/EnvironmentQuery/EnvQueryTypes.h
+++ b/Engine/Source/Runtime/AIModule/Classes/EnvironmentQuery/EnvQueryTypes.h
@@ -473,6 +473,7 @@ struct AIMODULE_API FEnvQueryItemDetails
 
 	int32 FailedTestIndex;
 	int32 ItemIndex;
+	FString FailedDescription;
 #endif // USE_EQS_DEBUGGER
 
 	FEnvQueryItemDetails() {}
@@ -608,8 +609,6 @@ public:
 #if USE_EQS_DEBUGGER
 	/** set to true to store additional debug info */
 	uint8 bStoreDebugInfo : 1;
-	/** set to true when debug data has been stored before items sorting. This is time sensetive and gets reset to false every query step */
-	uint8 bDebugStoredBeforeSorting : 1;
 #endif // USE_EQS_DEBUGGER
 
 	/** run mode */
@@ -627,7 +626,6 @@ public:
 	FEnvQueryInstance() : World(NULL), CurrentTest(-1), NumValidItems(0), bFoundSingleResult(false), bPassOnSingleResult(false)
 #if USE_EQS_DEBUGGER
 		, bStoreDebugInfo(bDebuggingInfoEnabled) 
-		, bDebugStoredBeforeSorting(false)
 #endif // USE_EQS_DEBUGGER
 	{ IncStats(); }
 	FEnvQueryInstance(const FEnvQueryInstance& Other) { *this = Other; IncStats(); }
@@ -742,6 +740,18 @@ public:
 #if !NO_LOGGING
 	void Log(const FString Msg) const;
 #endif // #if !NO_LOGGING
+	
+#if USE_EQS_DEBUGGER
+#	define  UE_EQS_DBGMSG(Format, ...) \
+					Instance->ItemDetails[CurrentItem].FailedDescription = FString::Printf(Format, ##__VA_ARGS__)
+
+#	define UE_EQS_LOG(CategoryName, Verbosity, Format, ...) \
+					UE_LOG(CategoryName, Verbosity, Format, ##__VA_ARGS__); \
+					UE_EQS_DBGMSG(Format, ##__VA_ARGS__); 
+#else
+#	define UE_EQS_DBGMSG(Format, ...)
+#	define UE_EQS_LOG(CategoryName, Verbosity, Format, ...) UE_LOG(CategoryName, Verbosity, Format, ##__VA_ARGS__); 
+#endif
 
 #if CPP || UE_BUILD_DOCS
 	struct AIMODULE_API ItemIterator
@@ -764,6 +774,7 @@ public:
 					case EEnvTestFilterType::Maximum:
 						if (Score > Max)
 						{
+							UE_EQS_DBGMSG(TEXT("Value %f is above maximum value set to %f"), Score, Max);
 							bPassedTest = false;
 						}
 						break;
@@ -771,6 +782,7 @@ public:
 					case EEnvTestFilterType::Minimum:
 						if (Score < Min)
 						{
+							UE_EQS_DBGMSG(TEXT("Value %f is below minimum value set to %f"), Score, Min);
 							bPassedTest = false;
 						}
 						break;
@@ -778,17 +790,18 @@ public:
 					case EEnvTestFilterType::Range:
 						if ((Score < Min) || (Score > Max))
 						{
+							UE_EQS_DBGMSG(TEXT("Value %f is out of range set to (%f, %f)"), Score, Min, Max);
 							bPassedTest = false;
 						}
 						break;
 
 					case EEnvTestFilterType::Match:
-						UE_LOG(LogEQS, Error, TEXT("Filtering Type set to 'Match' for floating point test.  Will consider test as failed in all cases."));
+						UE_EQS_LOG(LogEQS, Error, TEXT("Filtering Type set to 'Match' for floating point test.  Will consider test as failed in all cases."));
 						bPassedTest = false;
 						break;
 
 					default:
-						UE_LOG(LogEQS, Error, TEXT("Filtering Type set to invalid value for floating point test.  Will consider test as failed in all cases."));
+						UE_EQS_LOG(LogEQS, Error, TEXT("Filtering Type set to invalid value for floating point test.  Will consider test as failed in all cases."));
 						bPassedTest = false;
 						break;
 				}
@@ -812,25 +825,31 @@ public:
 			{
 				case EEnvTestFilterType::Match:
 					bPassedTest = (bScore == bExpected);
+#if USE_EQS_DEBUGGER
+					if (!bPassedTest)
+					{
+						UE_EQS_DBGMSG(TEXT("Boolean score don't mach (expected %s and got %s)"), bExpected ? TEXT("TRUE") : TEXT("FALSE"), bScore ? TEXT("TRUE") : TEXT("FALSE"));
+					}
+#endif
 					break;
 
 				case EEnvTestFilterType::Maximum:
-					UE_LOG(LogEQS, Error, TEXT("Filtering Type set to 'Maximum' for boolean test.  Will consider test as failed in all cases."));
+					UE_EQS_LOG(LogEQS, Error, TEXT("Filtering Type set to 'Maximum' for boolean test.  Will consider test as failed in all cases."));
 					bPassedTest = false;
 					break;
 
 				case EEnvTestFilterType::Minimum:
-					UE_LOG(LogEQS, Error, TEXT("Filtering Type set to 'Minimum' for boolean test.  Will consider test as failed in all cases."));
+					UE_EQS_LOG(LogEQS, Error, TEXT("Filtering Type set to 'Minimum' for boolean test.  Will consider test as failed in all cases."));
 					bPassedTest = false;
 					break;
 
 				case EEnvTestFilterType::Range:
-					UE_LOG(LogEQS, Error, TEXT("Filtering Type set to 'Range' for boolean test.  Will consider test as failed in all cases."));
+					UE_EQS_LOG(LogEQS, Error, TEXT("Filtering Type set to 'Range' for boolean test.  Will consider test as failed in all cases."));
 					bPassedTest = false;
 					break;
 
 				default:
-					UE_LOG(LogEQS, Error, TEXT("Filtering Type set to invalid value for boolean test.  Will consider test as failed in all cases."));
+					UE_EQS_LOG(LogEQS, Error, TEXT("Filtering Type set to invalid value for boolean test.  Will consider test as failed in all cases."));
 					bPassedTest = false;
 					break;
 			}
@@ -920,10 +939,14 @@ public:
 	};
 #endif
 
+#undef UE_EQS_LOG
+#undef UE_EQS_DBGMSG
+
 #if USE_EQS_DEBUGGER
 	FEQSQueryDebugData DebugData;
 	static bool bDebuggingInfoEnabled;
 #endif // USE_EQS_DEBUGGER
+
 	FBox GetBoundingBox() const;
 };
 
diff --git a/Engine/Source/Runtime/AIModule/Private/EnvironmentQuery/EQSRenderingComponent.cpp b/Engine/Source/Runtime/AIModule/Private/EnvironmentQuery/EQSRenderingComponent.cpp
--- a/Engine/Source/Runtime/AIModule/Private/EnvironmentQuery/EQSRenderingComponent.cpp
+++ b/Engine/Source/Runtime/AIModule/Private/EnvironmentQuery/EQSRenderingComponent.cpp
@@ -53,7 +53,10 @@ FEQSSceneProxy::FEQSSceneProxy(const UPrimitiveComponent* InComponent, const FSt
 		}
 	}
 
-	FEQSSceneProxy::CollectEQSData(InComponent, QueryDataSource, SolidSpheres, Texts);
+#if  USE_EQS_DEBUGGER 
+	TArray<EQSDebug::FDebugHelper> DebugItems;
+	FEQSSceneProxy::CollectEQSData(InComponent, QueryDataSource, SolidSpheres, Texts, DebugItems);
+#endif
 }
 
 FEQSSceneProxy::FEQSSceneProxy(const UPrimitiveComponent* InComponent, const FString& InViewFlagName, bool bInDrawOnlyWhenSelected, const TArray<FSphere>& Spheres, const TArray<FText3d>& InTexts)
@@ -86,7 +89,8 @@ FEQSSceneProxy::FEQSSceneProxy(const UPrimitiveComponent* InComponent, const FSt
 	}
 }
 
-void FEQSSceneProxy::CollectEQSData(const UPrimitiveComponent* InComponent, const class IEQSQueryResultSourceInterface* InQueryDataSource, TArray<FSphere>& Spheres, TArray<FText3d>& Texts)
+#if  USE_EQS_DEBUGGER 
+void FEQSSceneProxy::CollectEQSData(const UPrimitiveComponent* InComponent, const class IEQSQueryResultSourceInterface* InQueryDataSource, TArray<FSphere>& Spheres, TArray<FText3d>& Texts, TArray<EQSDebug::FDebugHelper>& DebugItems)
 {
 	AActor* ActorOwner = InComponent ? InComponent->GetOwner() : NULL;
 	IEQSQueryResultSourceInterface* QueryDataSource = const_cast<IEQSQueryResultSourceInterface*>(InQueryDataSource);
@@ -106,7 +110,7 @@ void FEQSSceneProxy::CollectEQSData(const UPrimitiveComponent* InComponent, cons
 	const FEnvQueryResult* ResultItems = QueryDataSource->GetQueryResult();
 	const FEnvQueryInstance* QueryInstance = QueryDataSource->GetQueryInstance();
 
-	FEQSSceneProxy::CollectEQSData(ResultItems, QueryInstance, Spheres, Texts, QueryDataSource->GetShouldDrawFailedItems());
+	FEQSSceneProxy::CollectEQSData(ResultItems, QueryInstance, Spheres, Texts, QueryDataSource->GetShouldDrawFailedItems(), DebugItems);
 
 	if (ActorOwner && Spheres.Num() > EQSMaxItemsDrawn)
 	{
@@ -117,7 +121,7 @@ void FEQSSceneProxy::CollectEQSData(const UPrimitiveComponent* InComponent, cons
 	}
 }
 
-void FEQSSceneProxy::CollectEQSData(const struct FEnvQueryResult* ResultItems, const struct FEnvQueryInstance* QueryInstance, TArray<FSphere>& Spheres, TArray<FText3d>& Texts, bool ShouldDrawFailedItems)
+void FEQSSceneProxy::CollectEQSData(const struct FEnvQueryResult* ResultItems, const struct FEnvQueryInstance* QueryInstance, TArray<FSphere>& Spheres, TArray<FText3d>& Texts, bool ShouldDrawFailedItems, TArray<EQSDebug::FDebugHelper>& DebugItems)
 {
 	if (ResultItems == NULL)
 	{
@@ -131,7 +135,6 @@ void FEQSSceneProxy::CollectEQSData(const struct FEnvQueryResult* ResultItems, c
 		}
 	}
 
-#if USE_EQS_DEBUGGER
 	// using "mid-results" requires manual normalization
 	const bool bUseMidResults = QueryInstance && (QueryInstance->Items.Num() < QueryInstance->DebugData.DebugItems.Num());
 	const TArray<FEnvQueryItem>& Items = bUseMidResults ? QueryInstance->DebugData.DebugItems : ResultItems->Items;
@@ -141,12 +144,6 @@ void FEQSSceneProxy::CollectEQSData(const struct FEnvQueryResult* ResultItems, c
 	
 	const bool bSingleItemResult = QueryInstance != NULL && QueryInstance->DebugData.bSingleItemResult;
 
-	//if (!bUseMidResults && Items.Num() > 0 && Items[0].IsValid())
-	//{
-	//	const FVector Loc = FEQSRenderingHelper::ExtractLocation(ResultItems->ItemType, ResultItems->RawData, Items, 0);
-	//	Cylinders.Add(FWireCylinder(Loc, ItemDrawRadius.X, 100, FColor::White));
-	//}
-
 	float MinScore = 0.f;
 	float MaxScore = -BIG_NUMBER;
 	if (bUseMidResults)
@@ -175,6 +172,8 @@ void FEQSSceneProxy::CollectEQSData(const struct FEnvQueryResult* ResultItems, c
 					? FLinearColor(FColor::MakeRedToGreenColorFromScalar(Score)) 
 					: FLinearColor(0.2, 1.0, 1.0, 1)));
 
+				DebugItems.Add(EQSDebug::FDebugHelper(Loc, ItemDrawRadius.X));
+
 				const FString Label = bNoTestsPerformed ? TEXT("") : FString::Printf(TEXT("%.2f"), Score);
 				Texts.Add(FText3d(Label, Loc, FLinearColor::White));
 			}
@@ -188,6 +187,8 @@ void FEQSSceneProxy::CollectEQSData(const struct FEnvQueryResult* ResultItems, c
 			const FVector Loc = FEQSRenderingHelper::ExtractLocation(ResultItems->ItemType, RawData, Items, 0);
 			Spheres.Add(FSphere(ItemDrawRadius.X, Loc, FLinearColor(0.0, 1.0, 0.12, 1)));
 
+			DebugItems.Add(EQSDebug::FDebugHelper(Loc, ItemDrawRadius.X));
+
 			const FString Label = FString::Printf(TEXT("Winner %.2f"), Score);
 			Texts.Add(FText3d(Label, Loc, FLinearColor::White));
 		}
@@ -200,6 +201,8 @@ void FEQSSceneProxy::CollectEQSData(const struct FEnvQueryResult* ResultItems, c
 				const FVector Loc = FEQSRenderingHelper::ExtractLocation(ResultItems->ItemType, RawData, Items, ItemIndex);
 				Spheres.Add(FSphere(ItemDrawRadius.X, Loc, FLinearColor(0.0, 0.2, 0.025, 1)));
 
+				DebugItems.Add(EQSDebug::FDebugHelper(Loc, ItemDrawRadius.X));
+				
 				const FString Label = bNoTestsPerformed ? TEXT("") : FString::Printf(TEXT("%.2f"), Score);
 				Texts.Add(FText3d(Label, Loc, FLinearColor::White));
 			}
@@ -225,13 +228,20 @@ void FEQSSceneProxy::CollectEQSData(const struct FEnvQueryResult* ResultItems, c
 			const FVector Loc = FEQSRenderingHelper::ExtractLocation(QueryInstance->ItemType, InstanceDebugData.RawData, DebugQueryItems, ItemIndex);
 			Spheres.Add(FSphere(ItemDrawRadius.X, Loc, FLinearColor(0.0, 0.0, 0.6, 0.6)));
 
-			const FString Label = InstanceDebugData.PerformedTestNames[Details[ItemIndex].FailedTestIndex];
-			Texts.Add(FText3d(Label, Loc, FLinearColor::White));
+			auto& DebugHelper = DebugItems[DebugItems.Add(EQSDebug::FDebugHelper(Loc, ItemDrawRadius.X))];
+			DebugHelper.AdditionalInformation = Details[ItemIndex].FailedDescription;
+			if (Details[ItemIndex].FailedTestIndex != INDEX_NONE)
+			{
+				DebugHelper.FailedTestIndex = Details[ItemIndex].FailedTestIndex;
+				DebugHelper.FailedScore = Details[ItemIndex].TestResults[DebugHelper.FailedTestIndex];
+
+				const FString Label = InstanceDebugData.PerformedTestNames[DebugHelper.FailedTestIndex] + FString::Printf(TEXT("(%d)"), DebugHelper.FailedTestIndex);
+				Texts.Add(FText3d(Label, Loc, FLinearColor::White));
+			}
 		}
 	}
-#endif // USE_EQS_DEBUGGER
 }
-
+#endif //USE_EQS_DEBUGGER 
 void FEQSSceneProxy::DrawDebugLabels(UCanvas* Canvas, APlayerController* PC)
 {
 	if ( !ActorOwner
diff --git a/Engine/Source/Runtime/AIModule/Private/EnvironmentQuery/EnvQueryDebugHelpers.cpp b/Engine/Source/Runtime/AIModule/Private/EnvironmentQuery/EnvQueryDebugHelpers.cpp
--- a/Engine/Source/Runtime/AIModule/Private/EnvironmentQuery/EnvQueryDebugHelpers.cpp
+++ b/Engine/Source/Runtime/AIModule/Private/EnvironmentQuery/EnvQueryDebugHelpers.cpp
@@ -6,7 +6,7 @@
 #include "EnvironmentQuery/EnvQueryDebugHelpers.h"
 #include "EnvironmentQuery/EQSQueryResultSourceInterface.h"
 
-#if USE_EQS_DEBUGGER || ENABLE_VISUAL_LOG
+#if USE_EQS_DEBUGGER
 void UEnvQueryDebugHelpers::QueryToBlobArray(struct FEnvQueryInstance* Query, TArray<uint8>& BlobArray, bool bUseCompression)
 {
 	EQSDebug::FQueryData EQSLocalData;
@@ -43,7 +43,7 @@ void UEnvQueryDebugHelpers::QueryToDebugData(struct FEnvQueryInstance* Query, EQ
 	// step 1: data for rendering component
 	EQSLocalData.Reset();
 
-	FEQSSceneProxy::CollectEQSData(Query, Query, EQSLocalData.SolidSpheres, EQSLocalData.Texts, true);
+	FEQSSceneProxy::CollectEQSData(Query, Query, EQSLocalData.SolidSpheres, EQSLocalData.Texts, true, EQSLocalData.RenderDebugHelpers);
 
 	// step 2: detailed scoring data for HUD
 	const int32 MaxDetailedItems = 10;
@@ -82,7 +82,8 @@ void UEnvQueryDebugHelpers::QueryToDebugData(struct FEnvQueryInstance* Query, EQ
 		EQSLocalData.Items.Add(ItemInfo);
 	}
 
-	for (int32 TestIdx = 0; TestIdx < NumTests; TestIdx++)
+	const int32 NumAllTests = Query->Options[Query->OptionIndex].TestsToPerform.Num();
+	for (int32 TestIdx = 0; TestIdx < NumAllTests; TestIdx++)
 	{
 		EQSDebug::FTestData TestInfo;
 
@@ -92,6 +93,14 @@ void UEnvQueryDebugHelpers::QueryToDebugData(struct FEnvQueryInstance* Query, EQ
 
 		EQSLocalData.Tests.Add(TestInfo);
 	}
+
+	EQSLocalData.UsedOption = Query->OptionIndex;
+	const int32 NumOptions = Query->Options.Num();
+	for (int32 OptionIndex = 0; OptionIndex < NumOptions; ++OptionIndex)
+	{
+		const FString UserName = Query->Options[OptionIndex].Generator->OptionName;
+		EQSLocalData.Options.Add(UserName.Len() > 0 ? UserName : Query->Options[OptionIndex].Generator->GetName());
+	}
 }
 
 void  UEnvQueryDebugHelpers::BlobArrayToDebugData(const TArray<uint8>& BlobArray, EQSDebug::FQueryData& EQSLocalData, bool bUseCompression)
diff --git a/Engine/Source/Runtime/AIModule/Private/EnvironmentQuery/EnvQueryGenerator.cpp b/Engine/Source/Runtime/AIModule/Private/EnvironmentQuery/EnvQueryGenerator.cpp
--- a/Engine/Source/Runtime/AIModule/Private/EnvironmentQuery/EnvQueryGenerator.cpp
+++ b/Engine/Source/Runtime/AIModule/Private/EnvironmentQuery/EnvQueryGenerator.cpp
@@ -22,6 +22,8 @@ FText UEnvQueryGenerator::GetDescriptionDetails() const
 void UEnvQueryGenerator::PostEditChangeProperty(FPropertyChangedEvent& PropertyChangedEvent) 
 {
 	Super::PostEditChangeProperty(PropertyChangedEvent);
+#if USE_EQS_DEBUGGER
 	UEnvQueryManager::NotifyAssetUpdate(NULL);
+#endif
 }
 #endif //WITH_EDITOR && USE_EQS_DEBUGGER
diff --git a/Engine/Source/Runtime/AIModule/Private/EnvironmentQuery/EnvQueryInstance.cpp b/Engine/Source/Runtime/AIModule/Private/EnvironmentQuery/EnvQueryInstance.cpp
--- a/Engine/Source/Runtime/AIModule/Private/EnvironmentQuery/EnvQueryInstance.cpp
+++ b/Engine/Source/Runtime/AIModule/Private/EnvironmentQuery/EnvQueryInstance.cpp
@@ -188,9 +188,6 @@ void FEnvQueryInstance::ExecuteOneStep(double InTimeLimit)
 	const bool bDoingLastTest = (CurrentTest == OptionItem.TestsToPerform.Num() - 1);
 	bool bStepDone = true;
 	TimeLimit = InTimeLimit;
-#if USE_EQS_DEBUGGER
-	bDebugStoredBeforeSorting = false;
-#endif
 
 	if (CurrentTest < 0)
 	{
@@ -231,7 +228,7 @@ void FEnvQueryInstance::ExecuteOneStep(double InTimeLimit)
 	if (bStepDone)
 	{
 #if USE_EQS_DEBUGGER
-		if (bStoreDebugInfo && !bDebugStoredBeforeSorting)
+		if (bStoreDebugInfo)
 		{
 			DebugData.Store(this);
 		}
@@ -379,12 +376,13 @@ void FEnvQueryInstance::NormalizeScores()
 
 void FEnvQueryInstance::SortScores()
 {
+	const int32 NumItems = Items.Num();
 	if (Options[OptionIndex].bShuffleItems)
 	{
-		for (int32 ItemIndex = 0; ItemIndex < Items.Num(); ItemIndex++)
+		for (int32 ItemIndex = 0; ItemIndex < NumItems; ItemIndex++)
 		{
-			const int32 Idx1 = FMath::RandHelper(Items.Num());
-			const int32 Idx2 = FMath::RandHelper(Items.Num());
+			const int32 Idx1 = FMath::RandHelper(NumItems);
+			const int32 Idx2 = FMath::RandHelper(NumItems);
 			Items.Swap(Idx1, Idx2);
 
 #if USE_EQS_DEBUGGER
@@ -393,7 +391,34 @@ void FEnvQueryInstance::SortScores()
 		}
 	}
 
+#if USE_EQS_DEBUGGER
+	struct FSortHelperForDebugData
+	{
+		FEnvQueryItem	Item;
+		struct FEnvQueryItemDetails ItemDetails;
+
+		FSortHelperForDebugData(const FEnvQueryItem&	InItem, struct FEnvQueryItemDetails& InDebugDetails) : Item(InItem), ItemDetails(InDebugDetails) {}
+		bool operator<(const FSortHelperForDebugData& Other) const
+		{
+			return Item < Other.Item;
+		}
+	};
+	TArray<FSortHelperForDebugData> AllData;
+	AllData.Reserve(NumItems);
+	for (int32 Index = 0; Index < NumItems; ++Index)
+	{
+		AllData.Add(FSortHelperForDebugData(Items[Index], ItemDetails[Index]));
+	}
+	AllData.Sort(TGreater<FSortHelperForDebugData>());
+
+	for (int32 Index = 0; Index < NumItems; ++Index)
+	{
+		Items[Index] = AllData[Index].Item;
+		ItemDetails[Index] = AllData[Index].ItemDetails;
+	}
+#else
 	Items.Sort(TGreater<FEnvQueryItem>());
+#endif
 }
 
 void FEnvQueryInstance::StripRedundantData()
diff --git a/Engine/Source/Runtime/AIModule/Private/EnvironmentQuery/EnvQueryManager.cpp b/Engine/Source/Runtime/AIModule/Private/EnvironmentQuery/EnvQueryManager.cpp
--- a/Engine/Source/Runtime/AIModule/Private/EnvironmentQuery/EnvQueryManager.cpp
+++ b/Engine/Source/Runtime/AIModule/Private/EnvironmentQuery/EnvQueryManager.cpp
@@ -104,9 +104,10 @@ UEnvQueryManager* UEnvQueryManager::GetCurrent(UObject* WorldContextObject)
 	return World ? UAISystem::GetCurrentEQSManager(World) : NULL;
 }
 
-#if WITH_EDITOR && USE_EQS_DEBUGGER
+#if USE_EQS_DEBUGGER
 void UEnvQueryManager::NotifyAssetUpdate(UEnvQuery* Query)
 {
+#if WITH_EDITOR
 	if (GEditor == NULL)
 	{
 		return;
@@ -137,8 +138,9 @@ void UEnvQueryManager::NotifyAssetUpdate(UEnvQuery* Query)
 			}
 		}
 	}
+#endif //WITH_EDITOR
 }
-#endif // WITH_EDITOR && USE_EQS_DEBUGGER
+#endif // USE_EQS_DEBUGGER
 
 TStatId UEnvQueryManager::GetStatId() const
 {
diff --git a/Engine/Source/Runtime/AIModule/Private/EnvironmentQuery/EnvQueryTest.cpp b/Engine/Source/Runtime/AIModule/Private/EnvironmentQuery/EnvQueryTest.cpp
--- a/Engine/Source/Runtime/AIModule/Private/EnvironmentQuery/EnvQueryTest.cpp
+++ b/Engine/Source/Runtime/AIModule/Private/EnvironmentQuery/EnvQueryTest.cpp
@@ -302,7 +302,9 @@ void UEnvQueryTest::PostLoad()
 void UEnvQueryTest::PostEditChangeProperty(FPropertyChangedEvent& PropertyChangedEvent) 
 {
 	Super::PostEditChangeProperty(PropertyChangedEvent);
+#if USE_EQS_DEBUGGER
 	UEnvQueryManager::NotifyAssetUpdate(NULL);
+#endif
 }
 #endif //WITH_EDITOR && USE_EQS_DEBUGGER
 
diff --git a/Engine/Source/Runtime/AIModule/Private/VisualLoggerExtension.cpp b/Engine/Source/Runtime/AIModule/Private/VisualLoggerExtension.cpp
--- a/Engine/Source/Runtime/AIModule/Private/VisualLoggerExtension.cpp
+++ b/Engine/Source/Runtime/AIModule/Private/VisualLoggerExtension.cpp
@@ -13,24 +13,29 @@ FVisualLoggerExtension::FVisualLoggerExtension()
 
 void FVisualLoggerExtension::OnTimestampChange(float Timestamp, class UWorld* InWorld, class AActor* HelperActor )
 {
+#if USE_EQS_DEBUGGER
 	if (CurrentTimestamp != Timestamp)
 	{
 		CurrentTimestamp = Timestamp;
 		CachedEQSId = INDEX_NONE;
 		DisableEQSRendering(HelperActor);
 	}
+#endif
 }
 
 void FVisualLoggerExtension::DisableDrawingForData(class UWorld* InWorld, class UCanvas* Canvas, class AActor* HelperActor, const FName& TagName, const FVisLogEntry::FDataBlock& DataBlock, float Timestamp)
 {
+#if USE_EQS_DEBUGGER
 	if (TagName == *EVisLogTags::TAG_EQS && CurrentTimestamp == Timestamp)
 	{
 		DisableEQSRendering(HelperActor);
 	}
+#endif
 }
 
 void FVisualLoggerExtension::DisableEQSRendering(class AActor* HelperActor)
 {
+#if USE_EQS_DEBUGGER
 	if (HelperActor)
 	{
 		UEQSRenderingComponent* EQSRenderComp = HelperActor->FindComponentByClass<UEQSRenderingComponent>();
@@ -42,10 +47,12 @@ void FVisualLoggerExtension::DisableEQSRendering(class AActor* HelperActor)
 			EQSRenderComp->MarkRenderStateDirty();
 		}
 	}
+#endif
 }
 
 void FVisualLoggerExtension::DrawData(class UWorld* InWorld, class UCanvas* Canvas, class AActor* HelperActor, const FName& TagName, const FVisLogEntry::FDataBlock& DataBlock, float Timestamp)
 {
+#if USE_EQS_DEBUGGER
 	if (TagName == *EVisLogTags::TAG_EQS && HelperActor)
 	{
 		UEQSRenderingComponent* EQSRenderComp = HelperActor->FindComponentByClass<UEQSRenderingComponent>();
@@ -68,5 +75,6 @@ void FVisualLoggerExtension::DrawData(class UWorld* InWorld, class UCanvas* Canv
 			EQSRenderComp->MarkRenderStateDirty();
 		}
 	}
+#endif
 }
-#endif //ENABLE_VISUAL_LOG
\ No newline at end of file
+#endif //ENABLE_VISUAL_LOG
