Fixed CIS compilation for Monolith builds

diff --git a/Engine/Source/Runtime/AIModule/Classes/EnvironmentQuery/EnvQueryManager.h b/Engine/Source/Runtime/AIModule/Classes/EnvironmentQuery/EnvQueryManager.h
--- a/Engine/Source/Runtime/AIModule/Classes/EnvironmentQuery/EnvQueryManager.h
+++ b/Engine/Source/Runtime/AIModule/Classes/EnvironmentQuery/EnvQueryManager.h
@@ -146,15 +146,13 @@ class AIMODULE_API UEnvQueryManager : public UObject, public FTickableGameObject
 	static UEnvQueryManager* GetCurrent(UObject* WorldContextObject);
 
 #if USE_EQS_DEBUGGER
-#if WITH_EDITOR
 	static void NotifyAssetUpdate(UEnvQuery* Query);
-#endif //WITH_EDITOR
 
 	FEQSDebugger& GetDebugger() { return EQSDebugger; }
 
 protected:
 	FEQSDebugger EQSDebugger;
-#endif // WITH_EDITOR && USE_EQS_DEBUGGER
+#endif // USE_EQS_DEBUGGER
 
 protected:
 
