Changed default label of BTDecorator_Blackboard to more explanatory "Blackboard Based Condition" #UE4

diff --git a/Engine/Source/Runtime/AIModule/Private/BehaviorTree/Decorators/BTDecorator_Blackboard.cpp b/Engine/Source/Runtime/AIModule/Private/BehaviorTree/Decorators/BTDecorator_Blackboard.cpp
--- a/Engine/Source/Runtime/AIModule/Private/BehaviorTree/Decorators/BTDecorator_Blackboard.cpp
+++ b/Engine/Source/Runtime/AIModule/Private/BehaviorTree/Decorators/BTDecorator_Blackboard.cpp
@@ -7,7 +7,7 @@
 
 UBTDecorator_Blackboard::UBTDecorator_Blackboard(const class FPostConstructInitializeProperties& PCIP) : Super(PCIP)
 {
-	NodeName = "Blackboard";
+	NodeName = "Blackboard Based Condition";
 	NotifyObserver = EBTBlackboardRestart::ResultChange;
 }
 
