fixed infinite loop decorator hang

diff --git a/Engine/Source/Runtime/AIModule/Classes/BehaviorTree/BehaviorTreeTypes.h b/Engine/Source/Runtime/AIModule/Classes/BehaviorTree/BehaviorTreeTypes.h
--- a/Engine/Source/Runtime/AIModule/Classes/BehaviorTree/BehaviorTreeTypes.h
+++ b/Engine/Source/Runtime/AIModule/Classes/BehaviorTree/BehaviorTreeTypes.h
@@ -354,8 +354,18 @@ struct FBehaviorTreeSearchData
 	/** first node allowed in search */
 	struct FBTNodeIndex SearchStart;
 
+	/** search unique number */
+	int32 SearchId;
+
 	/** adds update info to PendingUpdates array, removing all previous updates for this node */
 	void AddUniqueUpdate(const FBehaviorTreeSearchUpdate& UpdateInfo);
+
+	/** assign unique Id number */
+	void AssignSearchId();
+
+private:
+
+	static int32 NextSearchId;
 };
 
 /** property block in blueprint defined nodes */
diff --git a/Engine/Source/Runtime/AIModule/Classes/BehaviorTree/Decorators/BTDecorator_Loop.h b/Engine/Source/Runtime/AIModule/Classes/BehaviorTree/Decorators/BTDecorator_Loop.h
--- a/Engine/Source/Runtime/AIModule/Classes/BehaviorTree/Decorators/BTDecorator_Loop.h
+++ b/Engine/Source/Runtime/AIModule/Classes/BehaviorTree/Decorators/BTDecorator_Loop.h
@@ -5,6 +5,7 @@
 
 struct FBTLoopDecoratorMemory
 {
+	int32 SearchId;
 	uint8 RemainingExecutions;
 };
 
diff --git a/Engine/Source/Runtime/AIModule/Private/BehaviorTree/BehaviorTreeComponent.cpp b/Engine/Source/Runtime/AIModule/Private/BehaviorTree/BehaviorTreeComponent.cpp
--- a/Engine/Source/Runtime/AIModule/Private/BehaviorTree/BehaviorTreeComponent.cpp
+++ b/Engine/Source/Runtime/AIModule/Private/BehaviorTree/BehaviorTreeComponent.cpp
@@ -854,6 +854,7 @@ void UBehaviorTreeComponent::ProcessExecutionRequest()
 
 	FBehaviorTreeInstance& ActiveInstance = InstanceStack[ActiveInstanceIdx];
 	UBTCompositeNode* TestNode = ExecutionRequest.ExecuteNode;
+	SearchData.AssignSearchId();
 
 	// activate root node if needed (can't be handled by parent composite...)
 	if (ActiveInstance.ActiveNode == NULL)
diff --git a/Engine/Source/Runtime/AIModule/Private/BehaviorTree/BehaviorTreeTypes.cpp b/Engine/Source/Runtime/AIModule/Private/BehaviorTree/BehaviorTreeTypes.cpp
--- a/Engine/Source/Runtime/AIModule/Private/BehaviorTree/BehaviorTreeTypes.cpp
+++ b/Engine/Source/Runtime/AIModule/Private/BehaviorTree/BehaviorTreeTypes.cpp
@@ -142,6 +142,9 @@ bool FBTNodeIndex::TakesPriorityOver(const FBTNodeIndex& Other) const
 //----------------------------------------------------------------------//
 // FBehaviorTreeSearchData
 //----------------------------------------------------------------------//
+
+int32 FBehaviorTreeSearchData::NextSearchId = 1;
+
 void FBehaviorTreeSearchData::AddUniqueUpdate(const FBehaviorTreeSearchUpdate& UpdateInfo)
 {
 	UE_VLOG(OwnerComp->GetOwner(), LogBehaviorTree, Verbose, TEXT("Search node update[%s]: %s"),
@@ -177,6 +180,12 @@ void FBehaviorTreeSearchData::AddUniqueUpdate(const FBehaviorTreeSearchUpdate& U
 	}
 }
 
+void FBehaviorTreeSearchData::AssignSearchId()
+{
+	SearchId = NextSearchId;
+	NextSearchId++;
+}
+
 //----------------------------------------------------------------------//
 // FBlackboardKeySelector
 //----------------------------------------------------------------------//
diff --git a/Engine/Source/Runtime/AIModule/Private/BehaviorTree/Decorators/BTDecorator_Loop.cpp b/Engine/Source/Runtime/AIModule/Private/BehaviorTree/Decorators/BTDecorator_Loop.cpp
--- a/Engine/Source/Runtime/AIModule/Private/BehaviorTree/Decorators/BTDecorator_Loop.cpp
+++ b/Engine/Source/Runtime/AIModule/Private/BehaviorTree/Decorators/BTDecorator_Loop.cpp
@@ -25,13 +25,21 @@ void UBTDecorator_Loop::OnNodeActivation(FBehaviorTreeSearchData& SearchData)
 		DecoratorMemory->RemainingExecutions = NumLoops;
 	}
 
-	if (!bInfiniteLoop)
+	bool bShouldLoop = false;
+	if (bInfiniteLoop)
+	{
+		// protect from truly infinite loop within single search
+		bShouldLoop = SearchData.SearchId != DecoratorMemory->SearchId;
+		DecoratorMemory->SearchId = SearchData.SearchId;
+	}
+	else
 	{
 		DecoratorMemory->RemainingExecutions--;
+		bShouldLoop = DecoratorMemory->RemainingExecutions > 0;
 	}
 
+
 	// set child selection overrides
-	const bool bShouldLoop = bInfiniteLoop || (DecoratorMemory->RemainingExecutions > 0);
 	if (bShouldLoop)
 	{
 		GetParentNode()->SetChildOverride(SearchData, ChildIndex);
