Refactored PathFollowing's MinAgentRadiusPct and MinAgentHeightPct use to be more clear #UE4

Also refactored MinAgentHeightPct to MinAgentHalfHeightPct

diff --git a/Engine/Source/Runtime/AIModule/Classes/Navigation/PathFollowingComponent.h b/Engine/Source/Runtime/AIModule/Classes/Navigation/PathFollowingComponent.h
--- a/Engine/Source/Runtime/AIModule/Classes/Navigation/PathFollowingComponent.h
+++ b/Engine/Source/Runtime/AIModule/Classes/Navigation/PathFollowingComponent.h
@@ -176,7 +176,7 @@ class AIMODULE_API UPathFollowingComponent : public UActorComponent, public IAIR
 	void SetBlockDetection(float DistanceThreshold, float Interval, int32 NumSamples);
 
 	/** set threshold for precise reach tests in intermediate goals (minimal test radius)  */
-	void SetPreciseReachThreshold(float AgentRadiusMultiplier, float AgentHeightMultiplier);
+	void SetPreciseReachThreshold(float AgentRadiusMultiplier, float AgentHalfHeightMultiplier);
 
 	/** set status of last requested move */
 	void SetLastMoveAtGoal(bool bFinishedAtGoal);
@@ -289,7 +289,7 @@ protected:
 	float MinAgentRadiusPct;
 
 	/** part of agent height used as min acceptable height difference */
-	float MinAgentHeightPct;
+	float MinAgentHalfHeightPct;
 
 	/** game specific data */
 	FCustomMoveSharedPtr GameData;
diff --git a/Engine/Source/Runtime/AIModule/Private/Navigation/PathFollowingComponent.cpp b/Engine/Source/Runtime/AIModule/Private/Navigation/PathFollowingComponent.cpp
--- a/Engine/Source/Runtime/AIModule/Private/Navigation/PathFollowingComponent.cpp
+++ b/Engine/Source/Runtime/AIModule/Private/Navigation/PathFollowingComponent.cpp
@@ -30,8 +30,8 @@ UPathFollowingComponent::UPathFollowingComponent(const class FPostConstructIniti
 	PrimaryComponentTick.bCanEverTick = true;
 	bAutoActivate = true;
 
-	MinAgentRadiusPct = 0.1f;
-	MinAgentHeightPct = 0.05f;
+	MinAgentRadiusPct = 1.1f;
+	MinAgentHalfHeightPct = 1.05f;
 	BlockDetectionDistance = 10.0f;
 	BlockDetectionInterval = 0.5f;
 	BlockDetectionSampleCount = 10;
@@ -787,7 +787,7 @@ void UPathFollowingComponent::UpdatePathSegment()
 		{
 			if (Path->GetPathPoints().IsValidIndex(MoveSegmentEndIndex) && Path->GetPathPoints().IsValidIndex(MoveSegmentStartIndex))
 			{
-				LogBlockHelper(GetOwner(), MovementComp, MinAgentRadiusPct, MinAgentHeightPct,
+				LogBlockHelper(GetOwner(), MovementComp, MinAgentRadiusPct, MinAgentHalfHeightPct,
 					Path->GetPathPoints()[MoveSegmentStartIndex].Location, Path->GetPathPoints()[MoveSegmentEndIndex].Location);
 			}
 			else
@@ -940,14 +940,14 @@ bool UPathFollowingComponent::HasReachedInternal(const FVector& Goal, float Goal
 	const FVector ToGoal = Goal - AgentLocation;
 
 	const float Dist2DSq = ToGoal.SizeSquared2D();
-	const float UseRadius = GoalRadius + RadiusThreshold + (bUseAgentRadius ? AgentRadius * (MinAgentRadiusPct + 1.f) : 0.0f);
+	const float UseRadius = GoalRadius + RadiusThreshold + (bUseAgentRadius ? AgentRadius * MinAgentRadiusPct : 0.0f);
 	if (Dist2DSq > FMath::Square(UseRadius))
 	{
 		return false;
 	}
 
 	const float ZDiff = FMath::Abs(ToGoal.Z);
-	const float UseHeight = GoalHalfHeight + AgentHalfHeight + (AgentHalfHeight * MinAgentHeightPct);
+	const float UseHeight = GoalHalfHeight + (AgentHalfHeight * MinAgentHalfHeightPct);
 	if (ZDiff > UseHeight)
 	{
 		return false;
@@ -1008,11 +1008,11 @@ void UPathFollowingComponent::DebugReachTest(float& CurrentDot, float& CurrentDi
 	MovingAgent->GetSimpleCollisionCylinder(AgentRadius, AgentHalfHeight);
 
 	CurrentDistance = ToGoal.Size2D();
-	const float UseRadius = GoalRadius + RadiusThreshold + (bUseAgentRadius ? AgentRadius * (MinAgentRadiusPct + 1.f) : 0.0f);
+	const float UseRadius = GoalRadius + RadiusThreshold + (bUseAgentRadius ? AgentRadius * MinAgentRadiusPct : 0.0f);
 	bDistanceFailed = (CurrentDistance > UseRadius) ? 1 : 0;
 
 	CurrentHeight = FMath::Abs(ToGoal.Z);
-	const float UseHeight = GoalHalfHeight + AgentHalfHeight + (AgentHalfHeight * MinAgentHeightPct);
+	const float UseHeight = GoalHalfHeight + (AgentHalfHeight * MinAgentHalfHeightPct);
 	bHeightFailed = (CurrentHeight > UseHeight) ? 1 : 0;
 }
 
@@ -1206,10 +1206,10 @@ void UPathFollowingComponent::UpdateMoveFocus()
 	}
 }
 
-void UPathFollowingComponent::SetPreciseReachThreshold(float AgentRadiusMultiplier, float AgentHeightMultiplier)
+void UPathFollowingComponent::SetPreciseReachThreshold(float AgentRadiusMultiplier, float AgentHalfHeightMultiplier)
 {
 	MinAgentRadiusPct = AgentRadiusMultiplier;
-	MinAgentHeightPct = AgentHeightMultiplier;
+	MinAgentHalfHeightPct = AgentHalfHeightMultiplier;
 }
 
 float UPathFollowingComponent::GetRemainingPathCost() const
