Removed escape sequences in comments for functions in BTTask_BlueprintBase.h. Shouldn't be necessary.

diff --git a/Engine/Source/Runtime/AIModule/Classes/BehaviorTree/Tasks/BTTask_BlueprintBase.h b/Engine/Source/Runtime/AIModule/Classes/BehaviorTree/Tasks/BTTask_BlueprintBase.h
--- a/Engine/Source/Runtime/AIModule/Classes/BehaviorTree/Tasks/BTTask_BlueprintBase.h
+++ b/Engine/Source/Runtime/AIModule/Classes/BehaviorTree/Tasks/BTTask_BlueprintBase.h
@@ -75,11 +75,11 @@ protected:
 	UFUNCTION(BlueprintCallable, Category="AI|BehaviorTree")
 	void FinishAbort();
 
-	/** task execution will be finished (with result \'Success\') after receiving specified message */
+	/** task execution will be finished (with result 'Success') after receiving specified message */
 	UFUNCTION(BlueprintCallable, Category="AI|BehaviorTree")
 	void SetFinishOnMessage(FName MessageName);
 
-	/** task execution will be finished (with result \'Success\') after receiving specified message with indicated ID */
+	/** task execution will be finished (with result 'Success') after receiving specified message with indicated ID */
 	UFUNCTION(BlueprintCallable, Category="AI|BehaviorTree")
 	void SetFinishOnMessageWithId(FName MessageName, int32 RequestID = -1);
 
