fixed crash on behavior tree's cleanup when AI persist through map/world change

diff --git a/Engine/Source/Runtime/AIModule/Classes/BehaviorTree/BehaviorTreeComponent.h b/Engine/Source/Runtime/AIModule/Classes/BehaviorTree/BehaviorTreeComponent.h
--- a/Engine/Source/Runtime/AIModule/Classes/BehaviorTree/BehaviorTreeComponent.h
+++ b/Engine/Source/Runtime/AIModule/Classes/BehaviorTree/BehaviorTreeComponent.h
@@ -118,6 +118,9 @@ public:
 	/** schedule execution flow update in next tick */
 	void ScheduleExecutionUpdate();
 
+	/** remove all runtime data, used on map change */
+	void Cleanup();
+
 	/** tries to find behavior tree instance in context */
 	int32 FindInstanceContainingNode(const UBTNode* Node) const;
 
diff --git a/Engine/Source/Runtime/AIModule/Classes/BehaviorTree/BehaviorTreeManager.h b/Engine/Source/Runtime/AIModule/Classes/BehaviorTree/BehaviorTreeManager.h
--- a/Engine/Source/Runtime/AIModule/Classes/BehaviorTree/BehaviorTreeManager.h
+++ b/Engine/Source/Runtime/AIModule/Classes/BehaviorTree/BehaviorTreeManager.h
@@ -3,6 +3,8 @@
 #pragma once
 #include "BehaviorTreeManager.generated.h"
 
+class UBehaviorTreeComponent;
+
 USTRUCT()
 struct FBehaviorTreeTemplateInfo
 {
@@ -43,9 +45,21 @@ class AIMODULE_API UBehaviorTreeManager : public UObject
 
 	void DumpUsageStats() const;
 
+	/** register new behavior tree component for tracking */
+	void AddActiveComponent(UBehaviorTreeComponent* Component);
+
+	/** unregister behavior tree component from tracking */
+	void RemoveActiveComponent(UBehaviorTreeComponent* Component);
+
+	static UBehaviorTreeManager* GetCurrent(UWorld* World);
+	static UBehaviorTreeManager* GetCurrent(UObject* WorldContextObject);
+
 protected:
 
 	/** initialized tree templates */
 	UPROPERTY()
 	TArray<FBehaviorTreeTemplateInfo> LoadedTemplates;
+
+	UPROPERTY()
+	TArray<UBehaviorTreeComponent*> ActiveComponents;
 };
diff --git a/Engine/Source/Runtime/AIModule/Private/BehaviorTree/BehaviorTreeComponent.cpp b/Engine/Source/Runtime/AIModule/Private/BehaviorTree/BehaviorTreeComponent.cpp
--- a/Engine/Source/Runtime/AIModule/Private/BehaviorTree/BehaviorTreeComponent.cpp
+++ b/Engine/Source/Runtime/AIModule/Private/BehaviorTree/BehaviorTreeComponent.cpp
@@ -29,6 +29,12 @@ UBehaviorTreeComponent::UBehaviorTreeComponent(const class FPostConstructInitial
 
 void UBehaviorTreeComponent::BeginDestroy()
 {
+	UBehaviorTreeManager* BTManager = UBehaviorTreeManager::GetCurrent(GetWorld());
+	if (BTManager)
+	{
+		BTManager->RemoveActiveComponent(this);
+	}
+
 	RemoveAllInstances();
 	Super::BeginDestroy();
 }
@@ -124,6 +130,12 @@ bool UBehaviorTreeComponent::StartTree(class UBehaviorTree* TreeAsset, EBTExecut
 #if USE_BEHAVIORTREE_DEBUGGER
 	DebuggerSteps.Reset();
 #endif
+	
+	UBehaviorTreeManager* BTManager = UBehaviorTreeManager::GetCurrent(GetWorld());
+	if (BTManager)
+	{
+		BTManager->AddActiveComponent(this);
+	}
 
 	// push new instance
 	const bool bPushed = PushInstance(TreeAsset);
@@ -176,6 +188,16 @@ void UBehaviorTreeComponent::RestartTree()
 	}
 }
 
+void UBehaviorTreeComponent::Cleanup()
+{
+	StopTree();
+	RemoveAllInstances();
+
+	KnownInstances.Reset();
+	InstanceStack.Reset();
+	NodeInstances.Reset();
+}
+
 void UBehaviorTreeComponent::OnTaskFinished(const class UBTTaskNode* TaskNode, EBTNodeResult::Type TaskResult)
 {
 	if (TaskNode == NULL || InstanceStack.Num() == 0 || IsPendingKill())
diff --git a/Engine/Source/Runtime/AIModule/Private/BehaviorTree/BehaviorTreeManager.cpp b/Engine/Source/Runtime/AIModule/Private/BehaviorTree/BehaviorTreeManager.cpp
--- a/Engine/Source/Runtime/AIModule/Private/BehaviorTree/BehaviorTreeManager.cpp
+++ b/Engine/Source/Runtime/AIModule/Private/BehaviorTree/BehaviorTreeManager.cpp
@@ -26,6 +26,12 @@ void UBehaviorTreeManager::FinishDestroy()
 {
 	SET_DWORD_STAT(STAT_AI_BehaviorTree_NumTemplates, 0);
 
+	for (int32 Idx = 0; Idx < ActiveComponents.Num(); Idx++)
+	{
+		ActiveComponents[Idx]->Cleanup();
+	}
+
+	ActiveComponents.Reset();
 	Super::FinishDestroy();
 }
 
@@ -352,4 +358,27 @@ void UBehaviorTreeManager::DumpUsageStats() const
 	
 	UE_LOG(LogBehaviorTree, Display, TEXT("--- Total Nodes class usage:"));
 	AllNodesCounter.Print(TEXT(","));
-}
\ No newline at end of file
+}
+
+void UBehaviorTreeManager::AddActiveComponent(UBehaviorTreeComponent* Component)
+{
+	ActiveComponents.AddUnique(Component);
+}
+
+void UBehaviorTreeManager::RemoveActiveComponent(UBehaviorTreeComponent* Component)
+{
+	ActiveComponents.Remove(Component);
+}
+
+UBehaviorTreeManager* UBehaviorTreeManager::GetCurrent(UWorld* World)
+{
+	UAISystem* AISys = UAISystem::GetCurrent(World);
+	return AISys ? AISys->GetBehaviorTreeManager() : NULL;
+}
+
+UBehaviorTreeManager* UBehaviorTreeManager::GetCurrent(UObject* WorldContextObject)
+{
+	UWorld* World = GEngine->GetWorldFromContextObject(WorldContextObject, false);
+	UAISystem* AISys = World ? UAISystem::GetCurrent(World) : NULL;
+	return AISys ? AISys->GetBehaviorTreeManager() : NULL;
+}
