fixed assert on accessing BT manager from non existing world

diff --git a/Engine/Source/Runtime/AIModule/Private/BehaviorTree/BehaviorTreeManager.cpp b/Engine/Source/Runtime/AIModule/Private/BehaviorTree/BehaviorTreeManager.cpp
--- a/Engine/Source/Runtime/AIModule/Private/BehaviorTree/BehaviorTreeManager.cpp
+++ b/Engine/Source/Runtime/AIModule/Private/BehaviorTree/BehaviorTreeManager.cpp
@@ -372,7 +372,7 @@ void UBehaviorTreeManager::RemoveActiveComponent(UBehaviorTreeComponent* Compone
 
 UBehaviorTreeManager* UBehaviorTreeManager::GetCurrent(UWorld* World)
 {
-	UAISystem* AISys = UAISystem::GetCurrent(World);
+	UAISystem* AISys = World ? UAISystem::GetCurrent(World) : NULL;
 	return AISys ? AISys->GetBehaviorTreeManager() : NULL;
 }
 
