better fix for accessing AI system managers

diff --git a/Engine/Source/Runtime/AIModule/Classes/AISystem.h b/Engine/Source/Runtime/AIModule/Classes/AISystem.h
--- a/Engine/Source/Runtime/AIModule/Classes/AISystem.h
+++ b/Engine/Source/Runtime/AIModule/Classes/AISystem.h
@@ -54,28 +54,15 @@ public:
 	FORCEINLINE UAIPerceptionSystem* GetPerceptionSystem() { return PerceptionSystem; }
 	FORCEINLINE const UAIPerceptionSystem* GetPerceptionSystem() const { return PerceptionSystem; }
 
-	FORCEINLINE static UAISystem* GetCurrent(UWorld* World) 
+	FORCEINLINE static UAISystem* GetCurrent(UWorld* World, bool bChecked = true) 
 	{ 
-		check(World != NULL)
-		check(World->GetAISystem() == NULL || Cast<UAISystem>(World->GetAISystem()) != NULL);
-		return (UAISystem*)(World->GetAISystem());
-	}
-
-	FORCEINLINE static UBehaviorTreeManager* GetCurrentBTManager(UWorld* World)
-	{
-		UAISystem* AISys = GetCurrent(World);
-		// if AI system is NULL you're probably running your AI code on a client
-		// or did something horribly wrong
-		check(AISys != NULL);
-		return AISys->GetBehaviorTreeManager();
-	}
+		if (!bChecked && World == NULL)
+		{
+			return NULL;
+		}
 
-	FORCEINLINE static UEnvQueryManager* GetCurrentEQSManager(UWorld* World)
-	{
-		UAISystem* AISys = GetCurrent(World);
-		// if AI system is NULL you're probably running your AI code on a client
-		check(AISys != NULL);
-		return AISys->GetEnvironmentQueryManager();
+		check(World);
+		return Cast<UAISystem>(World->GetAISystem());
 	}
 
 	FORCEINLINE UWorld* GetOuterWorld() const { return Cast<UWorld>(GetOuter()); }
diff --git a/Engine/Source/Runtime/AIModule/Private/BehaviorTree/BehaviorTreeComponent.cpp b/Engine/Source/Runtime/AIModule/Private/BehaviorTree/BehaviorTreeComponent.cpp
--- a/Engine/Source/Runtime/AIModule/Private/BehaviorTree/BehaviorTreeComponent.cpp
+++ b/Engine/Source/Runtime/AIModule/Private/BehaviorTree/BehaviorTreeComponent.cpp
@@ -1266,7 +1266,7 @@ bool UBehaviorTreeComponent::PushInstance(class UBehaviorTree* TreeAsset)
 		return false;
 	}
 
-	UBehaviorTreeManager* BTManager = UAISystem::GetCurrentBTManager(GetWorld());
+	UBehaviorTreeManager* BTManager = UBehaviorTreeManager::GetCurrent(GetWorld());
 	if (BTManager == NULL)
 	{
 		UE_VLOG(GetOwner(), LogBehaviorTree, Warning, TEXT("Failed to execute tree %s: behavior tree manager not found!"), *GetNameSafe(TreeAsset));
diff --git a/Engine/Source/Runtime/AIModule/Private/BehaviorTree/BehaviorTreeManager.cpp b/Engine/Source/Runtime/AIModule/Private/BehaviorTree/BehaviorTreeManager.cpp
--- a/Engine/Source/Runtime/AIModule/Private/BehaviorTree/BehaviorTreeManager.cpp
+++ b/Engine/Source/Runtime/AIModule/Private/BehaviorTree/BehaviorTreeManager.cpp
@@ -372,13 +372,13 @@ void UBehaviorTreeManager::RemoveActiveComponent(UBehaviorTreeComponent* Compone
 
 UBehaviorTreeManager* UBehaviorTreeManager::GetCurrent(UWorld* World)
 {
-	UAISystem* AISys = World ? UAISystem::GetCurrent(World) : NULL;
+	UAISystem* AISys = UAISystem::GetCurrent(World, false);
 	return AISys ? AISys->GetBehaviorTreeManager() : NULL;
 }
 
 UBehaviorTreeManager* UBehaviorTreeManager::GetCurrent(UObject* WorldContextObject)
 {
 	UWorld* World = GEngine->GetWorldFromContextObject(WorldContextObject, false);
-	UAISystem* AISys = World ? UAISystem::GetCurrent(World) : NULL;
+	UAISystem* AISys = UAISystem::GetCurrent(World, false);
 	return AISys ? AISys->GetBehaviorTreeManager() : NULL;
 }
diff --git a/Engine/Source/Runtime/AIModule/Private/BehaviorTree/Tasks/BTTask_RunEQSQuery.cpp b/Engine/Source/Runtime/AIModule/Private/BehaviorTree/Tasks/BTTask_RunEQSQuery.cpp
--- a/Engine/Source/Runtime/AIModule/Private/BehaviorTree/Tasks/BTTask_RunEQSQuery.cpp
+++ b/Engine/Source/Runtime/AIModule/Private/BehaviorTree/Tasks/BTTask_RunEQSQuery.cpp
@@ -45,7 +45,7 @@ EBTNodeResult::Type UBTTask_RunEQSQuery::ExecuteTask(class UBehaviorTreeComponen
 EBTNodeResult::Type UBTTask_RunEQSQuery::AbortTask(class UBehaviorTreeComponent* OwnerComp, uint8* NodeMemory)
 {
 	UWorld* MyWorld = OwnerComp->GetWorld();
-	UEnvQueryManager* QueryManager = UAISystem::GetCurrentEQSManager(MyWorld);
+	UEnvQueryManager* QueryManager = UEnvQueryManager::GetCurrent(MyWorld);
 	
 	if (QueryManager)
 	{
diff --git a/Engine/Source/Runtime/AIModule/Private/EnvironmentQuery/EQSTestingPawn.cpp b/Engine/Source/Runtime/AIModule/Private/EnvironmentQuery/EQSTestingPawn.cpp
--- a/Engine/Source/Runtime/AIModule/Private/EnvironmentQuery/EQSTestingPawn.cpp
+++ b/Engine/Source/Runtime/AIModule/Private/EnvironmentQuery/EQSTestingPawn.cpp
@@ -168,7 +168,7 @@ void AEQSTestingPawn::Reset()
 
 void AEQSTestingPawn::MakeOneStep()
 {
-	UEnvQueryManager* EQS = UAISystem::GetCurrentEQSManager(GetWorld());
+	UEnvQueryManager* EQS = UEnvQueryManager::GetCurrent(GetWorld());
 	if (EQS == NULL)
 	{
 		return;
diff --git a/Engine/Source/Runtime/AIModule/Private/EnvironmentQuery/EnvQueryManager.cpp b/Engine/Source/Runtime/AIModule/Private/EnvironmentQuery/EnvQueryManager.cpp
--- a/Engine/Source/Runtime/AIModule/Private/EnvironmentQuery/EnvQueryManager.cpp
+++ b/Engine/Source/Runtime/AIModule/Private/EnvironmentQuery/EnvQueryManager.cpp
@@ -61,7 +61,7 @@ int32 FEnvQueryRequest::Execute(EEnvQueryRunMode::Type RunMode, FQueryFinishedSi
 		}
 	}
 
-	UEnvQueryManager* EnvQueryManager = UAISystem::GetCurrentEQSManager(World);
+	UEnvQueryManager* EnvQueryManager = UEnvQueryManager::GetCurrent(World);
 	if (EnvQueryManager == NULL)
 	{
 		UE_LOG(LogEQS, Warning, TEXT("Missing EQS manager!"));
@@ -95,13 +95,15 @@ void UEnvQueryManager::FinishDestroy()
 
 UEnvQueryManager* UEnvQueryManager::GetCurrent(UWorld* World)
 {
-	return World ? UAISystem::GetCurrentEQSManager(World) : NULL;
+	UAISystem* AISys = UAISystem::GetCurrent(World, false);
+	return AISys ? AISys->GetEnvironmentQueryManager() : NULL;
 }
 
 UEnvQueryManager* UEnvQueryManager::GetCurrent(UObject* WorldContextObject)
 {
-	UWorld* World = WorldContextObject ? GEngine->GetWorldFromContextObject(WorldContextObject) : NULL;
-	return World ? UAISystem::GetCurrentEQSManager(World) : NULL;
+	UWorld* World = GEngine->GetWorldFromContextObject(WorldContextObject, false);
+	UAISystem* AISys = UAISystem::GetCurrent(World, false);
+	return AISys ? AISys->GetEnvironmentQueryManager() : NULL;
 }
 
 #if USE_EQS_DEBUGGER
@@ -116,7 +118,7 @@ void UEnvQueryManager::NotifyAssetUpdate(UEnvQuery* Query)
 	UWorld* World = GEditor->GetEditorWorldContext().World();
 	if (World)
 	{
-		UEnvQueryManager* EQS = UAISystem::GetCurrentEQSManager(World);
+		UEnvQueryManager* EQS = UEnvQueryManager::GetCurrent(World);
 		if (EQS)
 		{
 			EQS->InstanceCache.Reset();
