fixed crash in behavior tree's cleanup

diff --git a/Engine/Source/Runtime/AIModule/Private/BehaviorTree/BehaviorTreeComponent.cpp b/Engine/Source/Runtime/AIModule/Private/BehaviorTree/BehaviorTreeComponent.cpp
--- a/Engine/Source/Runtime/AIModule/Private/BehaviorTree/BehaviorTreeComponent.cpp
+++ b/Engine/Source/Runtime/AIModule/Private/BehaviorTree/BehaviorTreeComponent.cpp
@@ -1469,6 +1469,7 @@ void UBehaviorTreeComponent::RemoveAllInstances()
 		DummyInstance.InstanceMemory.Reset();
 	}
 
+	KnownInstances.Reset();
 	NodeInstances.Reset();
 }
 
diff --git a/Engine/Source/Runtime/AIModule/Private/BehaviorTree/BehaviorTreeManager.cpp b/Engine/Source/Runtime/AIModule/Private/BehaviorTree/BehaviorTreeManager.cpp
--- a/Engine/Source/Runtime/AIModule/Private/BehaviorTree/BehaviorTreeManager.cpp
+++ b/Engine/Source/Runtime/AIModule/Private/BehaviorTree/BehaviorTreeManager.cpp
@@ -28,7 +28,10 @@ void UBehaviorTreeManager::FinishDestroy()
 
 	for (int32 Idx = 0; Idx < ActiveComponents.Num(); Idx++)
 	{
-		ActiveComponents[Idx]->Cleanup();
+		if (ActiveComponents[Idx] && !ActiveComponents[Idx]->HasAnyFlags(RF_BeginDestroyed))
+		{
+			ActiveComponents[Idx]->Cleanup();
+		}
 	}
 
 	ActiveComponents.Reset();
diff --git a/Engine/Source/Runtime/AIModule/Private/BehaviorTree/BehaviorTreeTypes.cpp b/Engine/Source/Runtime/AIModule/Private/BehaviorTree/BehaviorTreeTypes.cpp
--- a/Engine/Source/Runtime/AIModule/Private/BehaviorTree/BehaviorTreeTypes.cpp
+++ b/Engine/Source/Runtime/AIModule/Private/BehaviorTree/BehaviorTreeTypes.cpp
@@ -74,9 +74,10 @@ void FBehaviorTreeInstance::Cleanup(class UBehaviorTreeComponent* OwnerComp, EBT
 	FBehaviorTreeInstanceId& Info = OwnerComp->KnownInstances[InstanceIdIndex];
 	if (Info.FirstNodeInstance >= 0)
 	{
+		const int32 MaxAllowedIdx = OwnerComp->NodeInstances.Num();
 		const int32 LastNodeIdx = OwnerComp->KnownInstances.IsValidIndex(InstanceIdIndex + 1) ?
-			OwnerComp->KnownInstances[InstanceIdIndex + 1].FirstNodeInstance :
-			OwnerComp->NodeInstances.Num();
+			FMath::Min(OwnerComp->KnownInstances[InstanceIdIndex + 1].FirstNodeInstance, MaxAllowedIdx) :
+			MaxAllowedIdx;
 
 		for (int32 Idx = Info.FirstNodeInstance; Idx < LastNodeIdx; Idx++)
 		{
