Marked a bunch of BTNodes as AIMODULE_API #UE4

diff --git a/Engine/Source/Runtime/AIModule/Classes/BehaviorTree/Blackboard/BlackboardKeyType_Bool.h b/Engine/Source/Runtime/AIModule/Classes/BehaviorTree/Blackboard/BlackboardKeyType_Bool.h
--- a/Engine/Source/Runtime/AIModule/Classes/BehaviorTree/Blackboard/BlackboardKeyType_Bool.h
+++ b/Engine/Source/Runtime/AIModule/Classes/BehaviorTree/Blackboard/BlackboardKeyType_Bool.h
@@ -3,8 +3,8 @@
 #include "BehaviorTree/Blackboard/BlackboardKeyType.h"
 #include "BlackboardKeyType_Bool.generated.h"
 
-UCLASS(EditInlineNew, MinimalAPI, meta=(DisplayName="Bool"))
-class UBlackboardKeyType_Bool : public UBlackboardKeyType
+UCLASS(EditInlineNew, meta=(DisplayName="Bool"))
+class AIMODULE_API UBlackboardKeyType_Bool : public UBlackboardKeyType
 {
 	GENERATED_UCLASS_BODY()
 
diff --git a/Engine/Source/Runtime/AIModule/Classes/BehaviorTree/Blackboard/BlackboardKeyType_Class.h b/Engine/Source/Runtime/AIModule/Classes/BehaviorTree/Blackboard/BlackboardKeyType_Class.h
--- a/Engine/Source/Runtime/AIModule/Classes/BehaviorTree/Blackboard/BlackboardKeyType_Class.h
+++ b/Engine/Source/Runtime/AIModule/Classes/BehaviorTree/Blackboard/BlackboardKeyType_Class.h
@@ -3,8 +3,8 @@
 #include "BehaviorTree/Blackboard/BlackboardKeyType.h"
 #include "BlackboardKeyType_Class.generated.h"
 
-UCLASS(EditInlineNew, MinimalAPI, meta=(DisplayName="Class"))
-class UBlackboardKeyType_Class : public UBlackboardKeyType
+UCLASS(EditInlineNew, meta=(DisplayName="Class"))
+class AIMODULE_API UBlackboardKeyType_Class : public UBlackboardKeyType
 {
 	GENERATED_UCLASS_BODY()
 
diff --git a/Engine/Source/Runtime/AIModule/Classes/BehaviorTree/Blackboard/BlackboardKeyType_Enum.h b/Engine/Source/Runtime/AIModule/Classes/BehaviorTree/Blackboard/BlackboardKeyType_Enum.h
--- a/Engine/Source/Runtime/AIModule/Classes/BehaviorTree/Blackboard/BlackboardKeyType_Enum.h
+++ b/Engine/Source/Runtime/AIModule/Classes/BehaviorTree/Blackboard/BlackboardKeyType_Enum.h
@@ -3,8 +3,8 @@
 #include "BehaviorTree/Blackboard/BlackboardKeyType.h"
 #include "BlackboardKeyType_Enum.generated.h"
 
-UCLASS(EditInlineNew, MinimalAPI, meta=(DisplayName="Enum"))
-class UBlackboardKeyType_Enum : public UBlackboardKeyType
+UCLASS(EditInlineNew, meta=(DisplayName="Enum"))
+class AIMODULE_API UBlackboardKeyType_Enum : public UBlackboardKeyType
 {
 	GENERATED_UCLASS_BODY()
 
diff --git a/Engine/Source/Runtime/AIModule/Classes/BehaviorTree/Blackboard/BlackboardKeyType_Float.h b/Engine/Source/Runtime/AIModule/Classes/BehaviorTree/Blackboard/BlackboardKeyType_Float.h
--- a/Engine/Source/Runtime/AIModule/Classes/BehaviorTree/Blackboard/BlackboardKeyType_Float.h
+++ b/Engine/Source/Runtime/AIModule/Classes/BehaviorTree/Blackboard/BlackboardKeyType_Float.h
@@ -3,8 +3,8 @@
 #include "BehaviorTree/Blackboard/BlackboardKeyType.h"
 #include "BlackboardKeyType_Float.generated.h"
 
-UCLASS(EditInlineNew, MinimalAPI, meta=(DisplayName="Float"))
-class UBlackboardKeyType_Float : public UBlackboardKeyType
+UCLASS(EditInlineNew, meta=(DisplayName="Float"))
+class AIMODULE_API UBlackboardKeyType_Float : public UBlackboardKeyType
 {
 	GENERATED_UCLASS_BODY()
 
diff --git a/Engine/Source/Runtime/AIModule/Classes/BehaviorTree/Blackboard/BlackboardKeyType_Int.h b/Engine/Source/Runtime/AIModule/Classes/BehaviorTree/Blackboard/BlackboardKeyType_Int.h
--- a/Engine/Source/Runtime/AIModule/Classes/BehaviorTree/Blackboard/BlackboardKeyType_Int.h
+++ b/Engine/Source/Runtime/AIModule/Classes/BehaviorTree/Blackboard/BlackboardKeyType_Int.h
@@ -3,8 +3,8 @@
 #include "BehaviorTree/Blackboard/BlackboardKeyType.h"
 #include "BlackboardKeyType_Int.generated.h"
 
-UCLASS(EditInlineNew, MinimalAPI, meta=(DisplayName="Int"))
-class UBlackboardKeyType_Int : public UBlackboardKeyType
+UCLASS(EditInlineNew, meta=(DisplayName="Int"))
+class AIMODULE_API UBlackboardKeyType_Int : public UBlackboardKeyType
 {
 	GENERATED_UCLASS_BODY()
 
diff --git a/Engine/Source/Runtime/AIModule/Classes/BehaviorTree/Blackboard/BlackboardKeyType_Name.h b/Engine/Source/Runtime/AIModule/Classes/BehaviorTree/Blackboard/BlackboardKeyType_Name.h
--- a/Engine/Source/Runtime/AIModule/Classes/BehaviorTree/Blackboard/BlackboardKeyType_Name.h
+++ b/Engine/Source/Runtime/AIModule/Classes/BehaviorTree/Blackboard/BlackboardKeyType_Name.h
@@ -3,8 +3,8 @@
 #include "BehaviorTree/Blackboard/BlackboardKeyType.h"
 #include "BlackboardKeyType_Name.generated.h"
 
-UCLASS(EditInlineNew, MinimalAPI, meta=(DisplayName="Name"))
-class UBlackboardKeyType_Name : public UBlackboardKeyType
+UCLASS(EditInlineNew, meta=(DisplayName="Name"))
+class AIMODULE_API UBlackboardKeyType_Name : public UBlackboardKeyType
 {
 	GENERATED_UCLASS_BODY()
 
diff --git a/Engine/Source/Runtime/AIModule/Classes/BehaviorTree/Blackboard/BlackboardKeyType_NativeEnum.h b/Engine/Source/Runtime/AIModule/Classes/BehaviorTree/Blackboard/BlackboardKeyType_NativeEnum.h
--- a/Engine/Source/Runtime/AIModule/Classes/BehaviorTree/Blackboard/BlackboardKeyType_NativeEnum.h
+++ b/Engine/Source/Runtime/AIModule/Classes/BehaviorTree/Blackboard/BlackboardKeyType_NativeEnum.h
@@ -4,8 +4,8 @@
 #include "BehaviorTree/Blackboard/BlackboardKeyType.h"
 #include "BlackboardKeyType_NativeEnum.generated.h"
 
-UCLASS(EditInlineNew, MinimalAPI, meta=(DisplayName="Native Enum"))
-class UBlackboardKeyType_NativeEnum : public UBlackboardKeyType
+UCLASS(EditInlineNew, meta=(DisplayName="Native Enum"))
+class AIMODULE_API UBlackboardKeyType_NativeEnum : public UBlackboardKeyType
 {
 	GENERATED_UCLASS_BODY()
 
diff --git a/Engine/Source/Runtime/AIModule/Classes/BehaviorTree/Blackboard/BlackboardKeyType_Object.h b/Engine/Source/Runtime/AIModule/Classes/BehaviorTree/Blackboard/BlackboardKeyType_Object.h
--- a/Engine/Source/Runtime/AIModule/Classes/BehaviorTree/Blackboard/BlackboardKeyType_Object.h
+++ b/Engine/Source/Runtime/AIModule/Classes/BehaviorTree/Blackboard/BlackboardKeyType_Object.h
@@ -3,8 +3,8 @@
 #include "BehaviorTree/Blackboard/BlackboardKeyType.h"
 #include "BlackboardKeyType_Object.generated.h"
 
-UCLASS(EditInlineNew, MinimalAPI, meta=(DisplayName="Object"))
-class UBlackboardKeyType_Object : public UBlackboardKeyType
+UCLASS(EditInlineNew, meta=(DisplayName="Object"))
+class AIMODULE_API UBlackboardKeyType_Object : public UBlackboardKeyType
 {
 	GENERATED_UCLASS_BODY()
 
diff --git a/Engine/Source/Runtime/AIModule/Classes/BehaviorTree/Blackboard/BlackboardKeyType_Rotator.h b/Engine/Source/Runtime/AIModule/Classes/BehaviorTree/Blackboard/BlackboardKeyType_Rotator.h
--- a/Engine/Source/Runtime/AIModule/Classes/BehaviorTree/Blackboard/BlackboardKeyType_Rotator.h
+++ b/Engine/Source/Runtime/AIModule/Classes/BehaviorTree/Blackboard/BlackboardKeyType_Rotator.h
@@ -2,8 +2,8 @@
 #pragma once
 #include "BlackboardKeyType_Rotator.generated.h"
 
-UCLASS(EditInlineNew, MinimalAPI, meta=(DisplayName="Rotator"))
-class UBlackboardKeyType_Rotator : public UBlackboardKeyType
+UCLASS(EditInlineNew, meta=(DisplayName="Rotator"))
+class AIMODULE_API UBlackboardKeyType_Rotator : public UBlackboardKeyType
 {
 	GENERATED_UCLASS_BODY()
 
diff --git a/Engine/Source/Runtime/AIModule/Classes/BehaviorTree/Blackboard/BlackboardKeyType_String.h b/Engine/Source/Runtime/AIModule/Classes/BehaviorTree/Blackboard/BlackboardKeyType_String.h
--- a/Engine/Source/Runtime/AIModule/Classes/BehaviorTree/Blackboard/BlackboardKeyType_String.h
+++ b/Engine/Source/Runtime/AIModule/Classes/BehaviorTree/Blackboard/BlackboardKeyType_String.h
@@ -3,8 +3,8 @@
 #include "BehaviorTree/Blackboard/BlackboardKeyType.h"
 #include "BlackboardKeyType_String.generated.h"
 
-UCLASS(EditInlineNew, MinimalAPI, meta=(DisplayName="String"))
-class UBlackboardKeyType_String : public UBlackboardKeyType
+UCLASS(EditInlineNew, meta=(DisplayName="String"))
+class AIMODULE_API UBlackboardKeyType_String : public UBlackboardKeyType
 {
 	GENERATED_UCLASS_BODY()
 
diff --git a/Engine/Source/Runtime/AIModule/Classes/BehaviorTree/Blackboard/BlackboardKeyType_Vector.h b/Engine/Source/Runtime/AIModule/Classes/BehaviorTree/Blackboard/BlackboardKeyType_Vector.h
--- a/Engine/Source/Runtime/AIModule/Classes/BehaviorTree/Blackboard/BlackboardKeyType_Vector.h
+++ b/Engine/Source/Runtime/AIModule/Classes/BehaviorTree/Blackboard/BlackboardKeyType_Vector.h
@@ -3,8 +3,8 @@
 #include "BehaviorTree/Blackboard/BlackboardKeyType.h"
 #include "BlackboardKeyType_Vector.generated.h"
 
-UCLASS(EditInlineNew, MinimalAPI, meta=(DisplayName="Vector"))
-class UBlackboardKeyType_Vector : public UBlackboardKeyType
+UCLASS(EditInlineNew, meta=(DisplayName="Vector"))
+class AIMODULE_API UBlackboardKeyType_Vector : public UBlackboardKeyType
 {
 	GENERATED_UCLASS_BODY()
 
diff --git a/Engine/Source/Runtime/AIModule/Classes/BehaviorTree/Composites/BTComposite_Selector.h b/Engine/Source/Runtime/AIModule/Classes/BehaviorTree/Composites/BTComposite_Selector.h
--- a/Engine/Source/Runtime/AIModule/Classes/BehaviorTree/Composites/BTComposite_Selector.h
+++ b/Engine/Source/Runtime/AIModule/Classes/BehaviorTree/Composites/BTComposite_Selector.h
@@ -9,7 +9,7 @@
  * A selector node runs each child in turn until the first child node succeeds, in which case the selector fails.
  */
 UCLASS()
-class UBTComposite_Selector : public UBTCompositeNode
+class AIMODULE_API UBTComposite_Selector: public UBTCompositeNode
 {
 	GENERATED_UCLASS_BODY()
 
diff --git a/Engine/Source/Runtime/AIModule/Classes/BehaviorTree/Composites/BTComposite_Sequence.h b/Engine/Source/Runtime/AIModule/Classes/BehaviorTree/Composites/BTComposite_Sequence.h
--- a/Engine/Source/Runtime/AIModule/Classes/BehaviorTree/Composites/BTComposite_Sequence.h
+++ b/Engine/Source/Runtime/AIModule/Classes/BehaviorTree/Composites/BTComposite_Sequence.h
@@ -9,7 +9,7 @@
  * A sequence node runs each child in turn until one fails, in which case the sequence succeeds.
  */
 UCLASS()
-class UBTComposite_Sequence : public UBTCompositeNode
+class AIMODULE_API UBTComposite_Sequence : public UBTCompositeNode
 {
 	GENERATED_UCLASS_BODY()
 
diff --git a/Engine/Source/Runtime/AIModule/Classes/BehaviorTree/Composites/BTComposite_SimpleParallel.h b/Engine/Source/Runtime/AIModule/Classes/BehaviorTree/Composites/BTComposite_SimpleParallel.h
--- a/Engine/Source/Runtime/AIModule/Classes/BehaviorTree/Composites/BTComposite_SimpleParallel.h
+++ b/Engine/Source/Runtime/AIModule/Classes/BehaviorTree/Composites/BTComposite_SimpleParallel.h
@@ -42,7 +42,7 @@ struct FBTParallelMemory : public FBTCompositeMemory
  * Allows for running two children: one which must be a single task node (with optional decorators), and the other of which can be a complete subtree.
  */
 UCLASS(MinimalAPI)
-class UBTComposite_SimpleParallel : public UBTCompositeNode
+class AIMODULE_API UBTComposite_SimpleParallel : public UBTCompositeNode
 {
 	GENERATED_UCLASS_BODY()
 
diff --git a/Engine/Source/Runtime/AIModule/Classes/BehaviorTree/Decorators/BTDecorator_ForceSuccess.h b/Engine/Source/Runtime/AIModule/Classes/BehaviorTree/Decorators/BTDecorator_ForceSuccess.h
--- a/Engine/Source/Runtime/AIModule/Classes/BehaviorTree/Decorators/BTDecorator_ForceSuccess.h
+++ b/Engine/Source/Runtime/AIModule/Classes/BehaviorTree/Decorators/BTDecorator_ForceSuccess.h
@@ -14,7 +14,7 @@
  */
 
 UCLASS(HideCategories=(Condition))
-class UBTDecorator_ForceSuccess : public UBTDecorator
+class AIMODULE_API UBTDecorator_ForceSuccess : public UBTDecorator
 {
 	GENERATED_UCLASS_BODY()
 
diff --git a/Engine/Source/Runtime/AIModule/Classes/BehaviorTree/Services/BTService_DefaultFocus.h b/Engine/Source/Runtime/AIModule/Classes/BehaviorTree/Services/BTService_DefaultFocus.h
--- a/Engine/Source/Runtime/AIModule/Classes/BehaviorTree/Services/BTService_DefaultFocus.h
+++ b/Engine/Source/Runtime/AIModule/Classes/BehaviorTree/Services/BTService_DefaultFocus.h
@@ -21,7 +21,7 @@ struct FBTFocusMemory
  * A service node that automatically sets the AI controller's focus when it becomes active.
  */
 UCLASS(hidecategories=(Service))
-class UBTService_DefaultFocus : public UBTService_BlackboardBase
+class AIMODULE_API UBTService_DefaultFocus : public UBTService_BlackboardBase
 {
 	GENERATED_UCLASS_BODY()
 
diff --git a/Engine/Source/Runtime/AIModule/Classes/BehaviorTree/Tasks/BTTask_MakeNoise.h b/Engine/Source/Runtime/AIModule/Classes/BehaviorTree/Tasks/BTTask_MakeNoise.h
--- a/Engine/Source/Runtime/AIModule/Classes/BehaviorTree/Tasks/BTTask_MakeNoise.h
+++ b/Engine/Source/Runtime/AIModule/Classes/BehaviorTree/Tasks/BTTask_MakeNoise.h
@@ -8,7 +8,7 @@
  * A task node that calls MakeNoise() on this Pawn when executed.
  */
 UCLASS()
-class UBTTask_MakeNoise : public UBTTaskNode
+class AIMODULE_API UBTTask_MakeNoise : public UBTTaskNode
 {
 	GENERATED_UCLASS_BODY()
 
diff --git a/Engine/Source/Runtime/AIModule/Classes/BehaviorTree/Tasks/BTTask_PlaySound.h b/Engine/Source/Runtime/AIModule/Classes/BehaviorTree/Tasks/BTTask_PlaySound.h
--- a/Engine/Source/Runtime/AIModule/Classes/BehaviorTree/Tasks/BTTask_PlaySound.h
+++ b/Engine/Source/Runtime/AIModule/Classes/BehaviorTree/Tasks/BTTask_PlaySound.h
@@ -11,7 +11,7 @@ class USoundCue;
  * Plays the specified sound when executed.
  */
 UCLASS()
-class UBTTask_PlaySound : public UBTTaskNode
+class AIMODULE_API UBTTask_PlaySound : public UBTTaskNode
 {
 	GENERATED_UCLASS_BODY()
 
diff --git a/Engine/Source/Runtime/AIModule/Classes/BehaviorTree/Tasks/BTTask_RunEQSQuery.h b/Engine/Source/Runtime/AIModule/Classes/BehaviorTree/Tasks/BTTask_RunEQSQuery.h
--- a/Engine/Source/Runtime/AIModule/Classes/BehaviorTree/Tasks/BTTask_RunEQSQuery.h
+++ b/Engine/Source/Runtime/AIModule/Classes/BehaviorTree/Tasks/BTTask_RunEQSQuery.h
@@ -14,8 +14,8 @@ struct FBTEnvQueryTaskMemory
  * Run Environment Query System Query task node.
  * Runs the specified environment query when executed.
  */
-UCLASS(MinimalAPI)
-class UBTTask_RunEQSQuery : public UBTTask_BlackboardBase
+UCLASS()
+class AIMODULE_API UBTTask_RunEQSQuery : public UBTTask_BlackboardBase
 {
 	GENERATED_UCLASS_BODY()
 
diff --git a/Engine/Source/Runtime/AIModule/Classes/BehaviorTree/Tasks/BTTask_Wait.h b/Engine/Source/Runtime/AIModule/Classes/BehaviorTree/Tasks/BTTask_Wait.h
--- a/Engine/Source/Runtime/AIModule/Classes/BehaviorTree/Tasks/BTTask_Wait.h
+++ b/Engine/Source/Runtime/AIModule/Classes/BehaviorTree/Tasks/BTTask_Wait.h
@@ -14,7 +14,7 @@ struct FBTWaitTaskMemory
  * Wait for the specified time when executed.
  */
 UCLASS()
-class UBTTask_Wait : public UBTTaskNode
+class AIMODULE_API UBTTask_Wait : public UBTTaskNode
 {
 	GENERATED_UCLASS_BODY()
 
