Compilation issues fix #UE4

diff --git a/Engine/Source/Runtime/AIModule/Classes/BehaviorTree/Composites/BTComposite_SimpleParallel.h b/Engine/Source/Runtime/AIModule/Classes/BehaviorTree/Composites/BTComposite_SimpleParallel.h
--- a/Engine/Source/Runtime/AIModule/Classes/BehaviorTree/Composites/BTComposite_SimpleParallel.h
+++ b/Engine/Source/Runtime/AIModule/Classes/BehaviorTree/Composites/BTComposite_SimpleParallel.h
@@ -41,7 +41,7 @@ struct FBTParallelMemory : public FBTCompositeMemory
  * Simple Parallel composite node.
  * Allows for running two children: one which must be a single task node (with optional decorators), and the other of which can be a complete subtree.
  */
-UCLASS(MinimalAPI)
+UCLASS()
 class AIMODULE_API UBTComposite_SimpleParallel : public UBTCompositeNode
 {
 	GENERATED_UCLASS_BODY()
