Refactored some BT API functions taking UBehaviorTree pointer to accept reference instead #UE4

diff --git a/Engine/Source/Developer/AITestSuite/Classes/MockAI_BT.h b/Engine/Source/Developer/AITestSuite/Classes/MockAI_BT.h
--- a/Engine/Source/Developer/AITestSuite/Classes/MockAI_BT.h
+++ b/Engine/Source/Developer/AITestSuite/Classes/MockAI_BT.h
@@ -16,5 +16,5 @@ class UMockAI_BT : public UMockAI
 	TArray<int32> ExpectedResult;
 
 	bool IsRunning() const;
-	bool RunBT(UBehaviorTree* BTAsset, EBTExecutionMode::Type RunType = EBTExecutionMode::SingleRun);
+	bool RunBT(UBehaviorTree& BTAsset, EBTExecutionMode::Type RunType = EBTExecutionMode::SingleRun);
 };
diff --git a/Engine/Source/Developer/AITestSuite/Private/MockAI/MockAI_BT.cpp b/Engine/Source/Developer/AITestSuite/Private/MockAI/MockAI_BT.cpp
--- a/Engine/Source/Developer/AITestSuite/Private/MockAI/MockAI_BT.cpp
+++ b/Engine/Source/Developer/AITestSuite/Private/MockAI/MockAI_BT.cpp
@@ -25,9 +25,9 @@ bool UMockAI_BT::IsRunning() const
 	return BTComp && BTComp->IsRunning() && BTComp->GetRootTree();
 }
 
-bool UMockAI_BT::RunBT(UBehaviorTree* BTAsset, EBTExecutionMode::Type RunType)
+bool UMockAI_BT::RunBT(UBehaviorTree& BTAsset, EBTExecutionMode::Type RunType)
 {
-	BBComp->InitializeBlackboard(BTAsset->BlackboardAsset);
+	BBComp->InitializeBlackboard(BTAsset.BlackboardAsset);
 	BBComp->CacheBrainComponent(BTComp);
 	BTComp->CacheBlackboardComponent(BBComp);
 
@@ -37,4 +37,4 @@ bool UMockAI_BT::RunBT(UBehaviorTree* BTAsset, EBTExecutionMode::Type RunType)
 	BTComp->RegisterComponentWithWorld(World);
 
 	return BTComp->StartTree(BTAsset, RunType);
-}
\ No newline at end of file
+}
diff --git a/Engine/Source/Developer/AITestSuite/Private/Tests/AITestsCommon.cpp b/Engine/Source/Developer/AITestSuite/Private/Tests/AITestsCommon.cpp
--- a/Engine/Source/Developer/AITestSuite/Private/Tests/AITestsCommon.cpp
+++ b/Engine/Source/Developer/AITestSuite/Private/Tests/AITestsCommon.cpp
@@ -86,9 +86,9 @@ FAITest_SimpleBT::~FAITest_SimpleBT()
 
 void FAITest_SimpleBT::SetUp()
 {
-	if (AIBTUser)
+	if (AIBTUser && BTAsset)
 	{
-		AIBTUser->RunBT(BTAsset, EBTExecutionMode::SingleRun);
+		AIBTUser->RunBT(*BTAsset, EBTExecutionMode::SingleRun);
 		AIBTUser->SetEnableTicking(true);
 	}
 }
diff --git a/Engine/Source/Developer/AITestSuite/Public/BTBuilder.h b/Engine/Source/Developer/AITestSuite/Public/BTBuilder.h
--- a/Engine/Source/Developer/AITestSuite/Public/BTBuilder.h
+++ b/Engine/Source/Developer/AITestSuite/Public/BTBuilder.h
@@ -53,7 +53,7 @@ struct FBTBuilder
 	static UBTComposite_Selector& AddSelector(UBehaviorTree& TreeOb)
 	{
 		UBTComposite_Selector* NodeOb = NewObject<UBTComposite_Selector>(&TreeOb);
-		NodeOb->InitializeFromAsset(&TreeOb);
+		NodeOb->InitializeFromAsset(TreeOb);
 		TreeOb.RootNode = NodeOb;
 		return *NodeOb;
 	}
@@ -61,7 +61,7 @@ struct FBTBuilder
 	static UBTComposite_Selector& AddSelector(UBTCompositeNode& ParentNode)
 	{
 		UBTComposite_Selector* NodeOb = NewObject<UBTComposite_Selector>(ParentNode.GetTreeAsset());
-		NodeOb->InitializeFromAsset(ParentNode.GetTreeAsset());
+		NodeOb->InitializeFromAsset(*ParentNode.GetTreeAsset());
 
 		const int32 ChildIdx = ParentNode.Children.AddZeroed(1);
 		ParentNode.Children[ChildIdx].ChildComposite = NodeOb;
@@ -72,7 +72,7 @@ struct FBTBuilder
 	static UBTComposite_Sequence& AddSequence(UBehaviorTree& TreeOb)
 	{
 		UBTComposite_Sequence* NodeOb = NewObject<UBTComposite_Sequence>(&TreeOb);
-		NodeOb->InitializeFromAsset(&TreeOb);
+		NodeOb->InitializeFromAsset(TreeOb);
 		TreeOb.RootNode = NodeOb;
 		return *NodeOb;
 	}
@@ -80,7 +80,7 @@ struct FBTBuilder
 	static UBTComposite_Sequence& AddSequence(UBTCompositeNode& ParentNode)
 	{
 		UBTComposite_Sequence* NodeOb = NewObject<UBTComposite_Sequence>(ParentNode.GetTreeAsset());
-		NodeOb->InitializeFromAsset(ParentNode.GetTreeAsset());
+		NodeOb->InitializeFromAsset(*ParentNode.GetTreeAsset());
 
 		const int32 ChildIdx = ParentNode.Children.AddZeroed(1);
 		ParentNode.Children[ChildIdx].ChildComposite = NodeOb;
@@ -92,7 +92,7 @@ struct FBTBuilder
 	{
 		UBTComposite_SimpleParallel* NodeOb = NewObject<UBTComposite_SimpleParallel>(&TreeOb);
 		NodeOb->FinishMode = Mode;
-		NodeOb->InitializeFromAsset(&TreeOb);
+		NodeOb->InitializeFromAsset(TreeOb);
 		TreeOb.RootNode = NodeOb;
 		return *NodeOb;
 	}
@@ -101,7 +101,7 @@ struct FBTBuilder
 	{
 		UBTComposite_SimpleParallel* NodeOb = NewObject<UBTComposite_SimpleParallel>(ParentNode.GetTreeAsset());
 		NodeOb->FinishMode = Mode;
-		NodeOb->InitializeFromAsset(ParentNode.GetTreeAsset());
+		NodeOb->InitializeFromAsset(*ParentNode.GetTreeAsset());
 
 		const int32 ChildIdx = ParentNode.Children.AddZeroed(1);
 		ParentNode.Children[ChildIdx].ChildComposite = NodeOb;
diff --git a/Engine/Source/Editor/BehaviorTreeEditor/Private/BehaviorTreeDecoratorGraphNode_Decorator.cpp b/Engine/Source/Editor/BehaviorTreeEditor/Private/BehaviorTreeDecoratorGraphNode_Decorator.cpp
--- a/Engine/Source/Editor/BehaviorTreeEditor/Private/BehaviorTreeDecoratorGraphNode_Decorator.cpp
+++ b/Engine/Source/Editor/BehaviorTreeEditor/Private/BehaviorTreeDecoratorGraphNode_Decorator.cpp
@@ -22,13 +22,18 @@ void UBehaviorTreeDecoratorGraphNode_Decorator::PostPlacedNewNode()
 	if (NodeClass != NULL)
 	{
 		UBehaviorTreeGraphNode_CompositeDecorator* OwningNode = Cast<UBehaviorTreeGraphNode_CompositeDecorator>(GetDecoratorGraph()->GetOuter());
-		UBehaviorTree* BT = OwningNode ? Cast<UBehaviorTree>(OwningNode->GetOuter()->GetOuter()) : NULL;
-		
-		UBTDecorator* MyDecorator = ConstructObject<UBTDecorator>(NodeClass, BT);
-		MyDecorator->InitializeFromAsset(BT);
-		OwningNode->InitializeDecorator(MyDecorator);
-
-		NodeInstance = MyDecorator;
+		if (OwningNode)
+		{
+			UBehaviorTree* BT = Cast<UBehaviorTree>(OwningNode->GetOuter()->GetOuter());
+			if (BT)
+			{
+				UBTDecorator* MyDecorator = ConstructObject<UBTDecorator>(NodeClass, BT);
+				MyDecorator->InitializeFromAsset(*BT);
+				OwningNode->InitializeDecorator(MyDecorator);
+
+				NodeInstance = MyDecorator;
+			}
+		}
 	}
 }
 
@@ -64,13 +69,18 @@ void UBehaviorTreeDecoratorGraphNode_Decorator::PostEditImport()
 	if (NodeInstance)
 	{
 		UBehaviorTreeGraphNode_CompositeDecorator* OwningNode = Cast<UBehaviorTreeGraphNode_CompositeDecorator>(GetDecoratorGraph()->GetOuter());
-		UBehaviorTree* BT = OwningNode ? Cast<UBehaviorTree>(OwningNode->GetOuter()->GetOuter()) : NULL;
-
-		UBTDecorator* MyDecorator = (UBTDecorator*)NodeInstance;
-		MyDecorator->InitializeFromAsset(BT);
-		MyDecorator->InitializeNode(NULL, MAX_uint16, 0, 0);
-
-		OwningNode->InitializeDecorator(MyDecorator);
+		if (OwningNode)
+		{
+			UBehaviorTree* BT = Cast<UBehaviorTree>(OwningNode->GetOuter()->GetOuter());
+			if (BT)
+			{
+				UBTDecorator* MyDecorator = (UBTDecorator*)NodeInstance;
+				MyDecorator->InitializeFromAsset(*BT);
+				MyDecorator->InitializeNode(NULL, MAX_uint16, 0, 0);
+
+				OwningNode->InitializeDecorator(MyDecorator);
+			}
+		}
 	}
 }
 
diff --git a/Engine/Source/Editor/BehaviorTreeEditor/Private/BehaviorTreeGraph.cpp b/Engine/Source/Editor/BehaviorTreeEditor/Private/BehaviorTreeGraph.cpp
--- a/Engine/Source/Editor/BehaviorTreeEditor/Private/BehaviorTreeGraph.cpp
+++ b/Engine/Source/Editor/BehaviorTreeEditor/Private/BehaviorTreeGraph.cpp
@@ -18,6 +18,11 @@ UBehaviorTreeGraph::UBehaviorTreeGraph(const class FPostConstructInitializePrope
 void UBehaviorTreeGraph::UpdateBlackboardChange()
 {
 	UBehaviorTree* BTAsset = Cast<UBehaviorTree>(GetOuter());
+	if (BTAsset == nullptr)
+	{
+		return;
+	}
+
 	for (int32 Index = 0; Index < Nodes.Num(); ++Index)
 	{
 		UBehaviorTreeGraphNode* MyNode = Cast<UBehaviorTreeGraphNode>(Nodes[Index]);
@@ -27,7 +32,7 @@ void UBehaviorTreeGraph::UpdateBlackboardChange()
 			UBTNode* MyNodeInstance = Cast<UBTNode>(MyNode->NodeInstance);
 			if (MyNodeInstance)
 			{
-				MyNodeInstance->InitializeFromAsset(BTAsset);
+				MyNodeInstance->InitializeFromAsset(*BTAsset);
 			}
 		}
 
@@ -36,7 +41,7 @@ void UBehaviorTreeGraph::UpdateBlackboardChange()
 			UBTNode* MyNodeInstance = MyNode->Decorators[iDecorator] ? Cast<UBTNode>(MyNode->Decorators[iDecorator]->NodeInstance) : NULL;
 			if (MyNodeInstance)
 			{
-				MyNodeInstance->InitializeFromAsset(BTAsset);
+				MyNodeInstance->InitializeFromAsset(*BTAsset);
 			}
 
 			UBehaviorTreeGraphNode_CompositeDecorator* CompDecoratorNode = Cast<UBehaviorTreeGraphNode_CompositeDecorator>(MyNode->Decorators[iDecorator]);
@@ -51,7 +56,7 @@ void UBehaviorTreeGraph::UpdateBlackboardChange()
 			UBTNode* MyNodeInstance = MyNode->Services[iService] ? Cast<UBTNode>(MyNode->Services[iService]->NodeInstance) : NULL;
 			if (MyNodeInstance)
 			{
-				MyNodeInstance->InitializeFromAsset(BTAsset);
+				MyNodeInstance->InitializeFromAsset(*BTAsset);
 			}
 		}
 	}
diff --git a/Engine/Source/Editor/BehaviorTreeEditor/Private/BehaviorTreeGraphNode.cpp b/Engine/Source/Editor/BehaviorTreeEditor/Private/BehaviorTreeGraphNode.cpp
--- a/Engine/Source/Editor/BehaviorTreeEditor/Private/BehaviorTreeGraphNode.cpp
+++ b/Engine/Source/Editor/BehaviorTreeEditor/Private/BehaviorTreeGraphNode.cpp
@@ -50,12 +50,15 @@ void UBehaviorTreeGraphNode::PostPlacedNewNode()
 	if (NodeClass)
 	{
 		UBehaviorTree* BT = Cast<UBehaviorTree>(GetBehaviorTreeGraph()->GetOuter());
-		NodeInstance = ConstructObject<UBTNode>(NodeClass, BT);
+		if (BT)
+		{
+			NodeInstance = ConstructObject<UBTNode>(NodeClass, BT);
 
-		UBTNode* BTNode = (UBTNode*)NodeInstance;
-		BTNode->SetFlags(RF_Transactional);
-		BTNode->InitializeFromAsset(BT);
-		BTNode->InitializeNode(NULL, MAX_uint16, 0, 0);
+			UBTNode* BTNode = (UBTNode*)NodeInstance;
+			BTNode->SetFlags(RF_Transactional);
+			BTNode->InitializeFromAsset(*BT);
+			BTNode->InitializeNode(NULL, MAX_uint16, 0, 0);
+		}
 	}
 }
 
@@ -86,9 +89,12 @@ void UBehaviorTreeGraphNode::PostEditImport()
 	if (NodeInstance)
 	{
 		UBehaviorTree* BT = Cast<UBehaviorTree>(GetBehaviorTreeGraph()->GetOuter());
-		UBTNode* BTNode = (UBTNode*)NodeInstance;
-		BTNode->InitializeFromAsset(BT);
-		BTNode->InitializeNode(NULL, MAX_uint16, 0, 0);
+		if (BT)
+		{
+			UBTNode* BTNode = (UBTNode*)NodeInstance;
+			BTNode->InitializeFromAsset(*BT);
+			BTNode->InitializeNode(NULL, MAX_uint16, 0, 0);
+		}
 	}
 }
 
diff --git a/Engine/Source/Editor/BehaviorTreeEditor/Private/BehaviorTreeGraphNode_CompositeDecorator.cpp b/Engine/Source/Editor/BehaviorTreeEditor/Private/BehaviorTreeGraphNode_CompositeDecorator.cpp
--- a/Engine/Source/Editor/BehaviorTreeEditor/Private/BehaviorTreeGraphNode_CompositeDecorator.cpp
+++ b/Engine/Source/Editor/BehaviorTreeEditor/Private/BehaviorTreeGraphNode_CompositeDecorator.cpp
@@ -167,7 +167,7 @@ void UBehaviorTreeGraphNode_CompositeDecorator::OnBlackboardUpdate()
 			UBTNode* MyNodeInstance = MyNode ? Cast<UBTNode>(MyNode->NodeInstance) : NULL;
 			if (MyNodeInstance)
 			{
-				MyNodeInstance->InitializeFromAsset(BTAsset);
+				MyNodeInstance->InitializeFromAsset(*BTAsset);
 			}
 		}
 	}
diff --git a/Engine/Source/Editor/BehaviorTreeEditor/Private/BehaviorTreeGraphNode_SubtreeTask.cpp b/Engine/Source/Editor/BehaviorTreeEditor/Private/BehaviorTreeGraphNode_SubtreeTask.cpp
--- a/Engine/Source/Editor/BehaviorTreeEditor/Private/BehaviorTreeGraphNode_SubtreeTask.cpp
+++ b/Engine/Source/Editor/BehaviorTreeEditor/Private/BehaviorTreeGraphNode_SubtreeTask.cpp
@@ -60,51 +60,55 @@ bool UBehaviorTreeGraphNode_SubtreeTask::UpdateInjectedNodes()
 	if (SubRoot)
 	{
 		UBehaviorTree* BTAsset = Cast<UBehaviorTree>(GetBehaviorTreeGraph()->GetOuter());
-		for (int32 Index = 0; Index < SubRoot->Decorators.Num(); Index++)
+
+		if(BTAsset)
 		{
-			UBehaviorTreeGraphNode* SubNode = SubRoot->Decorators[Index];
-			if (SubNode)
+			for (int32 Index = 0; Index < SubRoot->Decorators.Num(); Index++)
 			{
-				SubNode->PrepareForCopying();
-
-				UBehaviorTreeGraphNode* InjectedNode = Cast<UBehaviorTreeGraphNode>(StaticDuplicateObject(SubNode, GetOuter(), TEXT("")));
+				UBehaviorTreeGraphNode* SubNode = SubRoot->Decorators[Index];
+				if (SubNode)
+				{
+					SubNode->PrepareForCopying();
 
-				SubNode->PostCopyNode();
-				InjectedNode->PostCopyNode();
+					UBehaviorTreeGraphNode* InjectedNode = Cast<UBehaviorTreeGraphNode>(StaticDuplicateObject(SubNode, GetOuter(), TEXT("")));
 
-				InjectedNode->ParentNode = this;
-				InjectedNode->bInjectedNode = true;
+					SubNode->PostCopyNode();
+					InjectedNode->PostCopyNode();
 
-				UBTDecorator* InjectedInstance = Cast<UBTDecorator>(InjectedNode->NodeInstance);
-				if (InjectedInstance)
-				{
-					InjectedInstance->InitializeFromAsset(BTAsset);
-				}
+					InjectedNode->ParentNode = this;
+					InjectedNode->bInjectedNode = true;
 
-				UBehaviorTreeGraphNode_CompositeDecorator* CompNode = Cast<UBehaviorTreeGraphNode_CompositeDecorator>(InjectedNode);
-				if (CompNode)
-				{
-					UEdGraph* SubGraph = CompNode->GetBoundGraph();
-					if (SubGraph)
+					UBTDecorator* InjectedInstance = Cast<UBTDecorator>(InjectedNode->NodeInstance);
+					if (InjectedInstance)
 					{
-						SubGraph->bEditable = false;
+						InjectedInstance->InitializeFromAsset(*BTAsset);
+					}
 
-						for (int32 SubIndex = 0; SubIndex < SubGraph->Nodes.Num(); SubIndex++)
+					UBehaviorTreeGraphNode_CompositeDecorator* CompNode = Cast<UBehaviorTreeGraphNode_CompositeDecorator>(InjectedNode);
+					if (CompNode)
+					{
+						UEdGraph* SubGraph = CompNode->GetBoundGraph();
+						if (SubGraph)
 						{
-							UBehaviorTreeDecoratorGraphNode_Decorator* InjectedDecorator = Cast<UBehaviorTreeDecoratorGraphNode_Decorator>(SubGraph->Nodes[SubIndex]);
-							if (InjectedDecorator)
+							SubGraph->bEditable = false;
+
+							for (int32 SubIndex = 0; SubIndex < SubGraph->Nodes.Num(); SubIndex++)
 							{
-								InjectedInstance = Cast<UBTDecorator>(InjectedDecorator->NodeInstance);
-								if (InjectedInstance)
+								UBehaviorTreeDecoratorGraphNode_Decorator* InjectedDecorator = Cast<UBehaviorTreeDecoratorGraphNode_Decorator>(SubGraph->Nodes[SubIndex]);
+								if (InjectedDecorator)
 								{
-									InjectedInstance->InitializeFromAsset(BTAsset);
+									InjectedInstance = Cast<UBTDecorator>(InjectedDecorator->NodeInstance);
+									if (InjectedInstance)
+									{
+										InjectedInstance->InitializeFromAsset(*BTAsset);
+									}
 								}
 							}
 						}
 					}
-				}
 
-				Decorators.Add(InjectedNode);
+					Decorators.Add(InjectedNode);
+				}
 			}
 		}
 	}
diff --git a/Engine/Source/Runtime/AIModule/Classes/BehaviorTree/BTNode.h b/Engine/Source/Runtime/AIModule/Classes/BehaviorTree/BTNode.h
--- a/Engine/Source/Runtime/AIModule/Classes/BehaviorTree/BTNode.h
+++ b/Engine/Source/Runtime/AIModule/Classes/BehaviorTree/BTNode.h
@@ -6,6 +6,8 @@
 #include "BehaviorTreeTypes.h"
 #include "BTNode.generated.h"
 
+class UBehaviorTree;
+
 struct FBTInstancedNodeMemory
 {
 	int32 NodeIdx;
@@ -22,7 +24,7 @@ class AIMODULE_API UBTNode : public UObject
 	void InitializeNode(class UBTCompositeNode* InParentNode, uint16 InExecutionIndex, uint16 InMemoryOffset, uint8 InTreeDepth);
 
 	/** initialize any asset related data */
-	virtual void InitializeFromAsset(class UBehaviorTree* Asset);
+	virtual void InitializeFromAsset(UBehaviorTree& Asset);
 	
 	/** initialize memory block */
 	virtual void InitializeMemory(class UBehaviorTreeComponent* OwnerComp, uint8* NodeMemory, EBTMemoryInit::Type InitType) const;
diff --git a/Engine/Source/Runtime/AIModule/Classes/BehaviorTree/BehaviorTree.h b/Engine/Source/Runtime/AIModule/Classes/BehaviorTree/BehaviorTree.h
--- a/Engine/Source/Runtime/AIModule/Classes/BehaviorTree/BehaviorTree.h
+++ b/Engine/Source/Runtime/AIModule/Classes/BehaviorTree/BehaviorTree.h
@@ -6,6 +6,10 @@
 #include "Engine/Blueprint.h"
 #include "BehaviorTree.generated.h"
 
+class UBTCompositeNode;
+class UBlackboardData;
+class UBTDecorator;
+
 UCLASS(BlueprintType)
 class AIMODULE_API UBehaviorTree : public UObject
 {
@@ -13,7 +17,7 @@ class AIMODULE_API UBehaviorTree : public UObject
 
 	/** root node of loaded tree */
 	UPROPERTY()
-	class UBTCompositeNode* RootNode;
+	UBTCompositeNode* RootNode;
 
 #if WITH_EDITORONLY_DATA
 
@@ -29,11 +33,11 @@ class AIMODULE_API UBehaviorTree : public UObject
 
 	/** blackboard asset for this tree */
 	UPROPERTY()
-	class UBlackboardData* BlackboardAsset;
+	UBlackboardData* BlackboardAsset;
 
 	/** root level decorators, used by subtrees */
 	UPROPERTY()
-	TArray<class UBTDecorator*> RootDecorators;
+	TArray<UBTDecorator*> RootDecorators;
 
 	/** logic operators for root level decorators, used by subtrees  */
 	UPROPERTY()
diff --git a/Engine/Source/Runtime/AIModule/Classes/BehaviorTree/BehaviorTreeComponent.h b/Engine/Source/Runtime/AIModule/Classes/BehaviorTree/BehaviorTreeComponent.h
--- a/Engine/Source/Runtime/AIModule/Classes/BehaviorTree/BehaviorTreeComponent.h
+++ b/Engine/Source/Runtime/AIModule/Classes/BehaviorTree/BehaviorTreeComponent.h
@@ -84,7 +84,7 @@ public:
 	// End UObject overrides
 
 	/** starts execution from root */
-	bool StartTree(UBehaviorTree* Asset, EBTExecutionMode::Type ExecuteMode = EBTExecutionMode::Looped);
+	bool StartTree(UBehaviorTree& Asset, EBTExecutionMode::Type ExecuteMode = EBTExecutionMode::Looped);
 
 	/** stops execution */
 	void StopTree();
@@ -231,7 +231,7 @@ protected:
 
 	/** push behavior tree instance on execution stack
 	 *	@NOTE: should never be called out-side of BT execution, meaning only BT tasks can push another BT instance! */
-	bool PushInstance(UBehaviorTree* TreeAsset);
+	bool PushInstance(UBehaviorTree& TreeAsset);
 
 	/** add unique Id of newly created subtree to KnownInstances list and return its index */
 	uint8 UpdateInstanceId(UBehaviorTree* TreeAsset, const UBTNode* OriginNode, int32 OriginInstanceIdx);
diff --git a/Engine/Source/Runtime/AIModule/Classes/BehaviorTree/BehaviorTreeManager.h b/Engine/Source/Runtime/AIModule/Classes/BehaviorTree/BehaviorTreeManager.h
--- a/Engine/Source/Runtime/AIModule/Classes/BehaviorTree/BehaviorTreeManager.h
+++ b/Engine/Source/Runtime/AIModule/Classes/BehaviorTree/BehaviorTreeManager.h
@@ -32,7 +32,7 @@ class AIMODULE_API UBehaviorTreeManager : public UObject
 	int32 MaxDebuggerSteps;
 
 	/** get behavior tree template for given blueprint */
-	bool LoadTree(class UBehaviorTree* Asset, class UBTCompositeNode*& Root, uint16& InstanceMemorySize);
+	bool LoadTree(UBehaviorTree& Asset, class UBTCompositeNode*& Root, uint16& InstanceMemorySize);
 
 	/** get aligned memory size */
 	static int32 GetAlignedDataSize(int32 Size);
diff --git a/Engine/Source/Runtime/AIModule/Classes/BehaviorTree/Decorators/BTDecorator_Blackboard.h b/Engine/Source/Runtime/AIModule/Classes/BehaviorTree/Decorators/BTDecorator_Blackboard.h
--- a/Engine/Source/Runtime/AIModule/Classes/BehaviorTree/Decorators/BTDecorator_Blackboard.h
+++ b/Engine/Source/Runtime/AIModule/Classes/BehaviorTree/Decorators/BTDecorator_Blackboard.h
@@ -79,7 +79,7 @@ protected:
 	virtual void BuildDescription();
 
 	virtual void PostEditChangeProperty(FPropertyChangedEvent& PropertyChangedEvent) override;
-	virtual void InitializeFromAsset(class UBehaviorTree* Asset) override;
+	virtual void InitializeFromAsset(UBehaviorTree& Asset) override;
 
 #endif
 
diff --git a/Engine/Source/Runtime/AIModule/Classes/BehaviorTree/Decorators/BTDecorator_BlackboardBase.h b/Engine/Source/Runtime/AIModule/Classes/BehaviorTree/Decorators/BTDecorator_BlackboardBase.h
--- a/Engine/Source/Runtime/AIModule/Classes/BehaviorTree/Decorators/BTDecorator_BlackboardBase.h
+++ b/Engine/Source/Runtime/AIModule/Classes/BehaviorTree/Decorators/BTDecorator_BlackboardBase.h
@@ -9,7 +9,7 @@ class AIMODULE_API UBTDecorator_BlackboardBase : public UBTDecorator
 	GENERATED_UCLASS_BODY()
 
 	/** initialize any asset related data */
-	virtual void InitializeFromAsset(class UBehaviorTree* Asset) override;
+	virtual void InitializeFromAsset(UBehaviorTree& Asset) override;
 
 	/** notify about change in blackboard keys */
 	virtual void OnBlackboardChange(const class UBlackboardComponent* Blackboard, FBlackboard::FKey ChangedKeyID);
diff --git a/Engine/Source/Runtime/AIModule/Classes/BehaviorTree/Decorators/BTDecorator_CompareBBEntries.h b/Engine/Source/Runtime/AIModule/Classes/BehaviorTree/Decorators/BTDecorator_CompareBBEntries.h
--- a/Engine/Source/Runtime/AIModule/Classes/BehaviorTree/Decorators/BTDecorator_CompareBBEntries.h
+++ b/Engine/Source/Runtime/AIModule/Classes/BehaviorTree/Decorators/BTDecorator_CompareBBEntries.h
@@ -19,7 +19,7 @@ namespace EBlackBoardEntryComparison
  * A decorator node that bases its condition on a comparison between two Blackboard keys.
  */
 UCLASS(HideCategories=(Condition))
-class UBTDecorator_CompareBBEntries : public UBTDecorator
+class AIMODULE_API UBTDecorator_CompareBBEntries : public UBTDecorator
 {
 	GENERATED_UCLASS_BODY()
 
@@ -46,7 +46,7 @@ public:
 	virtual bool CalculateRawConditionValue(class UBehaviorTreeComponent* OwnerComp, uint8* NodeMemory) const override;
 	virtual FString GetStaticDescription() const override;
 
-	virtual void InitializeFromAsset(class UBehaviorTree* Asset) override;
+	virtual void InitializeFromAsset(UBehaviorTree& Asset) override;
 
 	virtual void OnBlackboardChange(const class UBlackboardComponent* Blackboard, FBlackboard::FKey ChangedKeyID);
 	virtual void OnBecomeRelevant(UBehaviorTreeComponent* OwnerComp, uint8* NodeMemory) override;
diff --git a/Engine/Source/Runtime/AIModule/Classes/BehaviorTree/Decorators/BTDecorator_ConeCheck.h b/Engine/Source/Runtime/AIModule/Classes/BehaviorTree/Decorators/BTDecorator_ConeCheck.h
--- a/Engine/Source/Runtime/AIModule/Classes/BehaviorTree/Decorators/BTDecorator_ConeCheck.h
+++ b/Engine/Source/Runtime/AIModule/Classes/BehaviorTree/Decorators/BTDecorator_ConeCheck.h
@@ -38,7 +38,7 @@ class AIMODULE_API UBTDecorator_ConeCheck : public UBTDecorator
 	
 	float ConeHalfAngleDot;
 
-	virtual void InitializeFromAsset(class UBehaviorTree* Asset) override;
+	virtual void InitializeFromAsset(UBehaviorTree& Asset) override;
 	virtual uint16 GetInstanceMemorySize() const override;
 	virtual void DescribeRuntimeValues(const class UBehaviorTreeComponent* OwnerComp, uint8* NodeMemory, EBTDescriptionVerbosity::Type Verbosity, TArray<FString>& Values) const override;
 	virtual FString GetStaticDescription() const override;
diff --git a/Engine/Source/Runtime/AIModule/Classes/BehaviorTree/Decorators/BTDecorator_DoesPathExist.h b/Engine/Source/Runtime/AIModule/Classes/BehaviorTree/Decorators/BTDecorator_DoesPathExist.h
--- a/Engine/Source/Runtime/AIModule/Classes/BehaviorTree/Decorators/BTDecorator_DoesPathExist.h
+++ b/Engine/Source/Runtime/AIModule/Classes/BehaviorTree/Decorators/BTDecorator_DoesPathExist.h
@@ -49,7 +49,7 @@ public:
 
 	virtual bool CalculateRawConditionValue(class UBehaviorTreeComponent* OwnerComp, uint8* NodeMemory) const override;
 	virtual FString GetStaticDescription() const override;
-	virtual void InitializeFromAsset(class UBehaviorTree* Asset) override;
+	virtual void InitializeFromAsset(UBehaviorTree& Asset) override;
 
 #if WITH_EDITOR
 	virtual FName GetNodeIconName() const override;
diff --git a/Engine/Source/Runtime/AIModule/Classes/BehaviorTree/Decorators/BTDecorator_KeepInCone.h b/Engine/Source/Runtime/AIModule/Classes/BehaviorTree/Decorators/BTDecorator_KeepInCone.h
--- a/Engine/Source/Runtime/AIModule/Classes/BehaviorTree/Decorators/BTDecorator_KeepInCone.h
+++ b/Engine/Source/Runtime/AIModule/Classes/BehaviorTree/Decorators/BTDecorator_KeepInCone.h
@@ -42,7 +42,7 @@ class AIMODULE_API UBTDecorator_KeepInCone : public UBTDecorator
 	
 	float ConeHalfAngleDot;
 
-	virtual void InitializeFromAsset(class UBehaviorTree* Asset) override;
+	virtual void InitializeFromAsset(UBehaviorTree& Asset) override;
 	virtual uint16 GetInstanceMemorySize() const override;
 	virtual void DescribeRuntimeValues(const class UBehaviorTreeComponent* OwnerComp, uint8* NodeMemory, EBTDescriptionVerbosity::Type Verbosity, TArray<FString>& Values) const override;
 	virtual FString GetStaticDescription() const override;
diff --git a/Engine/Source/Runtime/AIModule/Classes/BehaviorTree/Services/BTService_BlackboardBase.h b/Engine/Source/Runtime/AIModule/Classes/BehaviorTree/Services/BTService_BlackboardBase.h
--- a/Engine/Source/Runtime/AIModule/Classes/BehaviorTree/Services/BTService_BlackboardBase.h
+++ b/Engine/Source/Runtime/AIModule/Classes/BehaviorTree/Services/BTService_BlackboardBase.h
@@ -9,7 +9,7 @@ class AIMODULE_API UBTService_BlackboardBase : public UBTService
 	GENERATED_UCLASS_BODY()
 
 	/** initialize any asset related data */
-	virtual void InitializeFromAsset(class UBehaviorTree* Asset) override;
+	virtual void InitializeFromAsset(UBehaviorTree& Asset) override;
 
 	/** get name of selected blackboard key */
 	FName GetSelectedBlackboardKey() const;
diff --git a/Engine/Source/Runtime/AIModule/Classes/BehaviorTree/Tasks/BTTask_BlackboardBase.h b/Engine/Source/Runtime/AIModule/Classes/BehaviorTree/Tasks/BTTask_BlackboardBase.h
--- a/Engine/Source/Runtime/AIModule/Classes/BehaviorTree/Tasks/BTTask_BlackboardBase.h
+++ b/Engine/Source/Runtime/AIModule/Classes/BehaviorTree/Tasks/BTTask_BlackboardBase.h
@@ -13,7 +13,7 @@ class AIMODULE_API UBTTask_BlackboardBase : public UBTTaskNode
 	GENERATED_UCLASS_BODY()
 
 	/** initialize any asset related data */
-	virtual void InitializeFromAsset(UBehaviorTree* Asset) override;
+	virtual void InitializeFromAsset(UBehaviorTree& Asset) override;
 
 	/** get name of selected blackboard key */
 	FName GetSelectedBlackboardKey() const;
diff --git a/Engine/Source/Runtime/AIModule/Private/AIController.cpp b/Engine/Source/Runtime/AIModule/Private/AIController.cpp
--- a/Engine/Source/Runtime/AIModule/Private/AIController.cpp
+++ b/Engine/Source/Runtime/AIModule/Private/AIController.cpp
@@ -738,7 +738,7 @@ bool AAIController::RunBehaviorTree(UBehaviorTree* BTAsset)
 		}
 
 		check(BTComp != NULL);
-		BTComp->StartTree(BTAsset, EBTExecutionMode::Looped);
+		BTComp->StartTree(*BTAsset, EBTExecutionMode::Looped);
 	}
 
 	return bSuccess;
diff --git a/Engine/Source/Runtime/AIModule/Private/BehaviorTree/BTNode.cpp b/Engine/Source/Runtime/AIModule/Private/BehaviorTree/BTNode.cpp
--- a/Engine/Source/Runtime/AIModule/Private/BehaviorTree/BTNode.cpp
+++ b/Engine/Source/Runtime/AIModule/Private/BehaviorTree/BTNode.cpp
@@ -108,9 +108,9 @@ void UBTNode::InitializeExecutionOrder(class UBTNode* NextNode)
 }
 #endif
 
-void UBTNode::InitializeFromAsset(class UBehaviorTree* Asset)
+void UBTNode::InitializeFromAsset(UBehaviorTree& Asset)
 {
-	TreeAsset = Asset;
+	TreeAsset = &Asset;
 }
 
 UBlackboardData* UBTNode::GetBlackboardAsset() const
diff --git a/Engine/Source/Runtime/AIModule/Private/BehaviorTree/BehaviorTreeComponent.cpp b/Engine/Source/Runtime/AIModule/Private/BehaviorTree/BehaviorTreeComponent.cpp
--- a/Engine/Source/Runtime/AIModule/Private/BehaviorTree/BehaviorTreeComponent.cpp
+++ b/Engine/Source/Runtime/AIModule/Private/BehaviorTree/BehaviorTreeComponent.cpp
@@ -13,7 +13,7 @@
 // UBehaviorTreeComponent
 //----------------------------------------------------------------------//
 
-UBehaviorTreeComponent::UBehaviorTreeComponent(const class FPostConstructInitializeProperties& PCIP) : Super(PCIP)
+UBehaviorTreeComponent::UBehaviorTreeComponent(const FPostConstructInitializeProperties& PCIP) : Super(PCIP)
 {
 	ActiveInstanceIdx = 0;
 	bLoopExecution = false;
@@ -105,12 +105,12 @@ bool UBehaviorTreeComponent::IsPaused() const
 	return bIsPaused;
 }
 
-bool UBehaviorTreeComponent::StartTree(class UBehaviorTree* TreeAsset, EBTExecutionMode::Type ExecuteMode)
+bool UBehaviorTreeComponent::StartTree(UBehaviorTree& Asset, EBTExecutionMode::Type ExecuteMode /*= EBTExecutionMode::Looped*/)
 {
 	// clear instance stack, start should always run new tree from root
-	class UBehaviorTree* CurrentRoot = GetRootTree();
+	UBehaviorTree* CurrentRoot = GetRootTree();
 	
-	if (CurrentRoot == TreeAsset && TreeHasBeenStarted())
+	if (CurrentRoot == &Asset && TreeHasBeenStarted())
 	{
 		UE_VLOG(GetOwner(), LogBehaviorTree, Log, TEXT("Skipping behavior start request - it's already running"));
 		return true;
@@ -118,7 +118,7 @@ bool UBehaviorTreeComponent::StartTree(class UBehaviorTree* TreeAsset, EBTExecut
 	else if (CurrentRoot)
 	{
 		UE_VLOG(GetOwner(), LogBehaviorTree, Log, TEXT("Abandoning behavior %s to start new one (%s)"),
-			*GetNameSafe(CurrentRoot), *GetNameSafe(TreeAsset));
+			*GetNameSafe(CurrentRoot), *Asset.GetName());
 	}
 
 	StopTree();
@@ -138,7 +138,7 @@ bool UBehaviorTreeComponent::StartTree(class UBehaviorTree* TreeAsset, EBTExecut
 	}
 
 	// push new instance
-	const bool bPushed = PushInstance(TreeAsset);
+	const bool bPushed = PushInstance(Asset);
 	return bPushed;
 }
 
@@ -199,7 +199,7 @@ void UBehaviorTreeComponent::Cleanup()
 	NodeInstances.Reset();
 }
 
-void UBehaviorTreeComponent::OnTaskFinished(const class UBTTaskNode* TaskNode, EBTNodeResult::Type TaskResult)
+void UBehaviorTreeComponent::OnTaskFinished(const UBTTaskNode* TaskNode, EBTNodeResult::Type TaskResult)
 {
 	if (TaskNode == NULL || InstanceStack.Num() == 0 || IsPendingKill())
 	{
@@ -285,7 +285,7 @@ void UBehaviorTreeComponent::OnTreeFinished()
 	}
 }
 
-bool UBehaviorTreeComponent::IsExecutingBranch(const class UBTNode* Node, int32 ChildIndex) const
+bool UBehaviorTreeComponent::IsExecutingBranch(const UBTNode* Node, int32 ChildIndex) const
 {
 	const int32 TestInstanceIdx = FindInstanceContainingNode(Node);
 	if (TestInstanceIdx == INDEX_NONE || InstanceStack[TestInstanceIdx].ActiveNode == NULL)
@@ -325,7 +325,7 @@ bool UBehaviorTreeComponent::IsExecutingBranch(const class UBTNode* Node, int32
 	return (ActiveExecutionIndex >= Node->GetExecutionIndex()) && (ActiveExecutionIndex < NextChildExecutionIndex);
 }
 
-bool UBehaviorTreeComponent::IsAuxNodeActive(const class UBTAuxiliaryNode* AuxNode) const
+bool UBehaviorTreeComponent::IsAuxNodeActive(const UBTAuxiliaryNode* AuxNode) const
 {
 	if (AuxNode == NULL)
 	{
@@ -363,7 +363,7 @@ bool UBehaviorTreeComponent::IsAuxNodeActive(const class UBTAuxiliaryNode* AuxNo
 	return false;
 }
 
-EBTTaskStatus::Type UBehaviorTreeComponent::GetTaskStatus(const class UBTTaskNode* TaskNode) const
+EBTTaskStatus::Type UBehaviorTreeComponent::GetTaskStatus(const UBTTaskNode* TaskNode) const
 {
 	EBTTaskStatus::Type Status = EBTTaskStatus::Inactive;
 	const int32 InstanceIdx = FindInstanceContainingNode(TaskNode);
@@ -401,7 +401,7 @@ EBTTaskStatus::Type UBehaviorTreeComponent::GetTaskStatus(const class UBTTaskNod
 	return Status;
 }
 
-void UBehaviorTreeComponent::RequestExecution(const class UBTDecorator* RequestedBy)
+void UBehaviorTreeComponent::RequestExecution(const UBTDecorator* RequestedBy)
 {
 	// search range depends on decorator's FlowAbortMode:
 	//
@@ -454,15 +454,15 @@ void UBehaviorTreeComponent::RequestExecution(EBTNodeResult::Type LastResult)
 }
 
 static void FindCommonParent(const TArray<FBehaviorTreeInstance>& Instances, const TArray<FBehaviorTreeInstanceId>& KnownInstances,
-							 class UBTCompositeNode* InNodeA, uint16 InstanceIdxA,
-							 class UBTCompositeNode* InNodeB, uint16 InstanceIdxB,
-							 class UBTCompositeNode*& CommonParentNode, uint16& CommonInstanceIdx)
+							 UBTCompositeNode* InNodeA, uint16 InstanceIdxA,
+							 UBTCompositeNode* InNodeB, uint16 InstanceIdxB,
+							 UBTCompositeNode*& CommonParentNode, uint16& CommonInstanceIdx)
 {
 	// find two nodes in the same instance (choose lower index = closer to root)
 	CommonInstanceIdx = (InstanceIdxA <= InstanceIdxB) ? InstanceIdxA : InstanceIdxB;
 
-	class UBTCompositeNode* NodeA = (CommonInstanceIdx == InstanceIdxA) ? InNodeA : Instances[CommonInstanceIdx].ActiveNode->GetParentNode();
-	class UBTCompositeNode* NodeB = (CommonInstanceIdx == InstanceIdxB) ? InNodeB : Instances[CommonInstanceIdx].ActiveNode->GetParentNode();
+	UBTCompositeNode* NodeA = (CommonInstanceIdx == InstanceIdxA) ? InNodeA : Instances[CommonInstanceIdx].ActiveNode->GetParentNode();
+	UBTCompositeNode* NodeB = (CommonInstanceIdx == InstanceIdxB) ? InNodeB : Instances[CommonInstanceIdx].ActiveNode->GetParentNode();
 
 	// special case: node was taken from CommonInstanceIdx, but it had ActiveNode set to root (no parent)
 	if (NodeA == NULL && CommonInstanceIdx != InstanceIdxA)
@@ -523,7 +523,7 @@ void UBehaviorTreeComponent::ScheduleExecutionUpdate()
 	}
 }
 
-void UBehaviorTreeComponent::RequestExecution(class UBTCompositeNode* RequestedOn, int32 InstanceIdx, const class UBTNode* RequestedBy,
+void UBehaviorTreeComponent::RequestExecution(UBTCompositeNode* RequestedOn, int32 InstanceIdx, const UBTNode* RequestedBy,
 											  int32 RequestedByChildIndex, EBTNodeResult::Type ContinueWithResult, bool bStoreForDebugger)
 {
 	SCOPE_CYCLE_COUNTER(STAT_AI_BehaviorTree_SearchTime);
@@ -775,7 +775,7 @@ void UBehaviorTreeComponent::ApplySearchUpdates(const TArray<struct FBehaviorTre
 	}
 }
 
-void UBehaviorTreeComponent::ApplySearchData(class UBTNode* NewActiveNode)
+void UBehaviorTreeComponent::ApplySearchData(UBTNode* NewActiveNode)
 {
 	const int32 NewNodeExecutionIndex = NewActiveNode ? NewActiveNode->GetExecutionIndex() : 0;
 
@@ -1059,7 +1059,7 @@ void UBehaviorTreeComponent::ProcessPendingExecution()
 	}
 }
 
-bool UBehaviorTreeComponent::DeactivateUpTo(class UBTCompositeNode* Node, uint16 NodeInstanceIdx, EBTNodeResult::Type& NodeResult)
+bool UBehaviorTreeComponent::DeactivateUpTo(UBTCompositeNode* Node, uint16 NodeInstanceIdx, EBTNodeResult::Type& NodeResult)
 {
 	UBTNode* DeactivatedChild = InstanceStack[ActiveInstanceIdx].ActiveNode;
 	bool bDeactivateRoot = true;
@@ -1187,7 +1187,7 @@ void UBehaviorTreeComponent::AbortCurrentTask()
 	}
 }
 
-void UBehaviorTreeComponent::RegisterMessageObserver(const class UBTTaskNode* TaskNode, FName MessageType)
+void UBehaviorTreeComponent::RegisterMessageObserver(const UBTTaskNode* TaskNode, FName MessageType)
 {
 	if (TaskNode)
 	{
@@ -1204,7 +1204,7 @@ void UBehaviorTreeComponent::RegisterMessageObserver(const class UBTTaskNode* Ta
 	}
 }
 
-void UBehaviorTreeComponent::RegisterMessageObserver(const class UBTTaskNode* TaskNode, FName MessageType, FAIRequestID RequestID)
+void UBehaviorTreeComponent::RegisterMessageObserver(const UBTTaskNode* TaskNode, FName MessageType, FAIRequestID RequestID)
 {
 	if (TaskNode)
 	{
@@ -1231,7 +1231,7 @@ void UBehaviorTreeComponent::UnregisterMessageObserversFrom(const FBTNodeIndex&
 	}
 }
 
-void UBehaviorTreeComponent::UnregisterMessageObserversFrom(const class UBTTaskNode* TaskNode)
+void UBehaviorTreeComponent::UnregisterMessageObserversFrom(const UBTTaskNode* TaskNode)
 {
 	if (TaskNode)
 	{
@@ -1255,7 +1255,7 @@ void UBehaviorTreeComponent::UnregisterMessageObserversFrom(const class UBTTaskN
 #define AUX_NODE_WRAPPER(cmd) cmd;
 #endif // STATS
 
-void UBehaviorTreeComponent::RegisterParallelTask(const class UBTTaskNode* TaskNode)
+void UBehaviorTreeComponent::RegisterParallelTask(const UBTTaskNode* TaskNode)
 {
 	FBehaviorTreeInstance& InstanceInfo = InstanceStack[ActiveInstanceIdx];
 	AUX_NODE_WRAPPER(InstanceInfo.ParallelTasks.Add( FBehaviorTreeParallelTask(TaskNode, EBTTaskStatus::Active) ));
@@ -1270,7 +1270,7 @@ void UBehaviorTreeComponent::RegisterParallelTask(const class UBTTaskNode* TaskN
 	}
 }
 
-void UBehaviorTreeComponent::UnregisterParallelTask(const class UBTTaskNode* TaskNode, uint16 InstanceIdx)
+void UBehaviorTreeComponent::UnregisterParallelTask(const UBTTaskNode* TaskNode, uint16 InstanceIdx)
 {
 	bool bShouldUpdate = false;
 	if (InstanceStack.IsValidIndex(InstanceIdx))
@@ -1318,18 +1318,13 @@ void UBehaviorTreeComponent::UpdateAbortingTasks()
 	}
 }
 
-bool UBehaviorTreeComponent::PushInstance(class UBehaviorTree* TreeAsset)
+bool UBehaviorTreeComponent::PushInstance(UBehaviorTree& TreeAsset)
 {
-	if (TreeAsset == NULL)
-	{
-		return false;
-	}
-
 	// check if blackboard class match
-	if (TreeAsset->BlackboardAsset && BlackboardComp && !BlackboardComp->IsCompatibleWith(TreeAsset->BlackboardAsset))
+	if (TreeAsset.BlackboardAsset && BlackboardComp && !BlackboardComp->IsCompatibleWith(TreeAsset.BlackboardAsset))
 	{
 		UE_VLOG(GetOwner(), LogBehaviorTree, Warning, TEXT("Failed to execute tree %s: blackboard %s is not compatibile with current: %s!"),
-			*GetNameSafe(TreeAsset), *GetNameSafe(TreeAsset->BlackboardAsset), *GetNameSafe(BlackboardComp->GetBlackboardAsset()));
+			*TreeAsset.GetName(), *GetNameSafe(TreeAsset.BlackboardAsset), *GetNameSafe(BlackboardComp->GetBlackboardAsset()));
 
 		return false;
 	}
@@ -1337,7 +1332,7 @@ bool UBehaviorTreeComponent::PushInstance(class UBehaviorTree* TreeAsset)
 	UBehaviorTreeManager* BTManager = UBehaviorTreeManager::GetCurrent(GetWorld());
 	if (BTManager == NULL)
 	{
-		UE_VLOG(GetOwner(), LogBehaviorTree, Warning, TEXT("Failed to execute tree %s: behavior tree manager not found!"), *GetNameSafe(TreeAsset));
+		UE_VLOG(GetOwner(), LogBehaviorTree, Warning, TEXT("Failed to execute tree %s: behavior tree manager not found!"), *TreeAsset.GetName());
 		return false;
 	}
 
@@ -1353,7 +1348,7 @@ bool UBehaviorTreeComponent::PushInstance(class UBehaviorTree* TreeAsset)
 		if (!bIsAllowed)
 		{
 			UE_VLOG(GetOwner(), LogBehaviorTree, Warning, TEXT("Failed to execute tree %s: parent of active node does not allow it! (%s)"),
-				*GetNameSafe(TreeAsset), *UBehaviorTreeTypes::DescribeNodeHelper(ActiveParent));
+				*TreeAsset.GetName(), *UBehaviorTreeTypes::DescribeNodeHelper(ActiveParent));
 			return false;
 		}
 	}
@@ -1365,7 +1360,7 @@ bool UBehaviorTreeComponent::PushInstance(class UBehaviorTree* TreeAsset)
 	if (bLoaded)
 	{
 		FBehaviorTreeInstance NewInstance;
-		NewInstance.InstanceIdIndex = UpdateInstanceId(TreeAsset, ActiveNode, InstanceStack.Num() - 1);
+		NewInstance.InstanceIdIndex = UpdateInstanceId(&TreeAsset, ActiveNode, InstanceStack.Num() - 1);
 		NewInstance.RootNode = RootNode;
 		NewInstance.ActiveNode = NULL;
 		NewInstance.ActiveNodeType = EBTActiveNode::Composite;
@@ -1398,7 +1393,7 @@ bool UBehaviorTreeComponent::PushInstance(class UBehaviorTree* TreeAsset)
 			ServiceNode->WrappedOnBecomeRelevant(this, NodeMemory);
 		}
 
-		FBehaviorTreeDelegates::OnTreeStarted.Broadcast(this, TreeAsset);
+		FBehaviorTreeDelegates::OnTreeStarted.Broadcast(this, &TreeAsset);
 
 		// start new task
 		RequestExecution(RootNode, ActiveInstanceIdx, RootNode, 0, EBTNodeResult::InProgress);
@@ -1408,7 +1403,7 @@ bool UBehaviorTreeComponent::PushInstance(class UBehaviorTree* TreeAsset)
 	return false;
 }
 
-uint8 UBehaviorTreeComponent::UpdateInstanceId(class UBehaviorTree* TreeAsset, const UBTNode* OriginNode, int32 OriginInstanceIdx)
+uint8 UBehaviorTreeComponent::UpdateInstanceId(UBehaviorTree* TreeAsset, const UBTNode* OriginNode, int32 OriginInstanceIdx)
 {
 	FBehaviorTreeInstanceId InstanceId;
 	InstanceId.TreeAsset = TreeAsset;
@@ -1475,7 +1470,7 @@ int32 UBehaviorTreeComponent::FindInstanceContainingNode(const UBTNode* Node) co
 	return InstanceIdx;
 }
 
-class UBTNode* UBehaviorTreeComponent::FindTemplateNode(const class UBTNode* Node) const
+UBTNode* UBehaviorTreeComponent::FindTemplateNode(const UBTNode* Node) const
 {
 	if (Node == NULL || !Node->IsInstanced() || Node->GetParentNode() == NULL)
 	{
@@ -1797,7 +1792,7 @@ void UBehaviorTreeComponent::StoreDebuggerRemovedInstance(uint16 InstanceIdx) co
 #endif
 }
 
-void UBehaviorTreeComponent::StoreDebuggerSearchStep(const class UBTNode* Node, uint16 InstanceIdx, EBTNodeResult::Type NodeResult) const
+void UBehaviorTreeComponent::StoreDebuggerSearchStep(const UBTNode* Node, uint16 InstanceIdx, EBTNodeResult::Type NodeResult) const
 {
 #if USE_BEHAVIORTREE_DEBUGGER
 	if (Node && NodeResult != EBTNodeResult::InProgress && NodeResult != EBTNodeResult::Aborted)
@@ -1819,7 +1814,7 @@ void UBehaviorTreeComponent::StoreDebuggerSearchStep(const class UBTNode* Node,
 #endif
 }
 
-void UBehaviorTreeComponent::StoreDebuggerSearchStep(const class UBTNode* Node, uint16 InstanceIdx, bool bPassed) const
+void UBehaviorTreeComponent::StoreDebuggerSearchStep(const UBTNode* Node, uint16 InstanceIdx, bool bPassed) const
 {
 #if USE_BEHAVIORTREE_DEBUGGER
 	if (Node && !bPassed)
@@ -1838,7 +1833,7 @@ void UBehaviorTreeComponent::StoreDebuggerSearchStep(const class UBTNode* Node,
 #endif
 }
 
-void UBehaviorTreeComponent::StoreDebuggerRestart(const class UBTNode* Node, uint16 InstanceIdx, bool bAllowed)
+void UBehaviorTreeComponent::StoreDebuggerRestart(const UBTNode* Node, uint16 InstanceIdx, bool bAllowed)
 {
 #if USE_BEHAVIORTREE_DEBUGGER
 	if (Node)
@@ -1858,7 +1853,7 @@ void UBehaviorTreeComponent::StoreDebuggerRestart(const class UBTNode* Node, uin
 #endif
 }
 
-void UBehaviorTreeComponent::StoreDebuggerRuntimeValues(TArray<FString>& RuntimeDescriptions, class UBTNode* RootNode, uint16 InstanceIdx) const
+void UBehaviorTreeComponent::StoreDebuggerRuntimeValues(TArray<FString>& RuntimeDescriptions, UBTNode* RootNode, uint16 InstanceIdx) const
 {
 #if USE_BEHAVIORTREE_DEBUGGER
 	if (!InstanceStack.IsValidIndex(InstanceIdx))
@@ -1893,7 +1888,7 @@ void UBehaviorTreeComponent::StoreDebuggerRuntimeValues(TArray<FString>& Runtime
 #endif
 }
 
-void UBehaviorTreeComponent::UpdateDebuggerAfterExecution(const class UBTTaskNode* TaskNode, uint16 InstanceIdx) const
+void UBehaviorTreeComponent::UpdateDebuggerAfterExecution(const UBTTaskNode* TaskNode, uint16 InstanceIdx) const
 {
 #if USE_BEHAVIORTREE_DEBUGGER
 	FBehaviorTreeExecutionStep& CurrentStep = DebuggerSteps.Last();
diff --git a/Engine/Source/Runtime/AIModule/Private/BehaviorTree/BehaviorTreeManager.cpp b/Engine/Source/Runtime/AIModule/Private/BehaviorTree/BehaviorTreeManager.cpp
--- a/Engine/Source/Runtime/AIModule/Private/BehaviorTree/BehaviorTreeManager.cpp
+++ b/Engine/Source/Runtime/AIModule/Private/BehaviorTree/BehaviorTreeManager.cpp
@@ -75,7 +75,7 @@ struct FNodeInitializationData
 
 static void InitializeNodeHelper(UBTCompositeNode* ParentNode, UBTNode* NodeOb,
 	uint8 TreeDepth, uint16& ExecutionIndex, TArray<FNodeInitializationData>& InitList,
-	class UBehaviorTree* TreeAsset, UObject* NodeOuter)
+	UBehaviorTree& TreeAsset, UObject* NodeOuter)
 {
 	// special case: subtrees
 	UBTTask_RunBehavior* SubtreeTask = Cast<UBTTask_RunBehavior>(NodeOb);
@@ -96,7 +96,7 @@ static void InitializeNodeHelper(UBTCompositeNode* ParentNode, UBTNode* NodeOb,
 			if (CompositeOb->Services[ServiceIndex] == NULL)
 			{
 				UE_LOG(LogBehaviorTree, Warning, TEXT("%s has missing service node! (parent: %s)"),
-					*GetNameSafe(TreeAsset), *UBehaviorTreeTypes::DescribeNodeHelper(CompositeOb));
+					*TreeAsset.GetName(), *UBehaviorTreeTypes::DescribeNodeHelper(CompositeOb));
 
 				CompositeOb->Services.RemoveAt(ServiceIndex, 1, false);
 				ServiceIndex--;
@@ -121,7 +121,7 @@ static void InitializeNodeHelper(UBTCompositeNode* ParentNode, UBTNode* NodeOb,
 				if (ChildInfo.Decorators[DecoratorIndex] == NULL)
 				{
 					UE_LOG(LogBehaviorTree, Warning, TEXT("%s has missing decorator node! (parent: %s, branch: %d)"),
-						*GetNameSafe(TreeAsset), *UBehaviorTreeTypes::DescribeNodeHelper(CompositeOb), ChildIndex);
+						*TreeAsset.GetName(), *UBehaviorTreeTypes::DescribeNodeHelper(CompositeOb), ChildIndex);
 
 					ChildInfo.Decorators.RemoveAt(DecoratorIndex, 1, false);
 					DecoratorIndex--;
@@ -162,14 +162,14 @@ static void InitializeNodeHelper(UBTCompositeNode* ParentNode, UBTNode* NodeOb,
 	}
 }
 
-bool UBehaviorTreeManager::LoadTree(class UBehaviorTree* Asset, UBTCompositeNode*& Root, uint16& InstanceMemorySize)
+bool UBehaviorTreeManager::LoadTree(UBehaviorTree& Asset, class UBTCompositeNode*& Root, uint16& InstanceMemorySize)
 {
 	SCOPE_CYCLE_COUNTER(STAT_AI_BehaviorTree_LoadTime);
 
 	for (int32 TemplateIndex = 0; TemplateIndex < LoadedTemplates.Num(); TemplateIndex++)
 	{
 		FBehaviorTreeTemplateInfo& TemplateInfo = LoadedTemplates[TemplateIndex];
-		if (TemplateInfo.Asset == Asset)
+		if (TemplateInfo.Asset == &Asset)
 		{
 			Root = TemplateInfo.Template;
 			InstanceMemorySize = TemplateInfo.InstanceMemorySize;
@@ -177,11 +177,11 @@ bool UBehaviorTreeManager::LoadTree(class UBehaviorTree* Asset, UBTCompositeNode
 		}
 	}
 
-	if (Asset->RootNode)
+	if (Asset.RootNode)
 	{
 		FBehaviorTreeTemplateInfo TemplateInfo;
-		TemplateInfo.Asset = Asset;
-		TemplateInfo.Template = Cast<UBTCompositeNode>(StaticDuplicateObject(Asset->RootNode, this, TEXT("None")));
+		TemplateInfo.Asset = &Asset;
+		TemplateInfo.Template = Cast<UBTCompositeNode>(StaticDuplicateObject(Asset.RootNode, this, TEXT("None")));
 
 		TArray<FNodeInitializationData> InitList;
 		uint16 ExecutionIndex = 0;
diff --git a/Engine/Source/Runtime/AIModule/Private/BehaviorTree/Decorators/BTDecorator_Blackboard.cpp b/Engine/Source/Runtime/AIModule/Private/BehaviorTree/Decorators/BTDecorator_Blackboard.cpp
--- a/Engine/Source/Runtime/AIModule/Private/BehaviorTree/Decorators/BTDecorator_Blackboard.cpp
+++ b/Engine/Source/Runtime/AIModule/Private/BehaviorTree/Decorators/BTDecorator_Blackboard.cpp
@@ -223,7 +223,7 @@ void UBTDecorator_Blackboard::PostEditChangeProperty(FPropertyChangedEvent& Prop
 	BuildDescription();
 }
 
-void UBTDecorator_Blackboard::InitializeFromAsset(class UBehaviorTree* Asset)
+void UBTDecorator_Blackboard::InitializeFromAsset(UBehaviorTree& Asset)
 {
 	Super::InitializeFromAsset(Asset);
 	BuildDescription();
diff --git a/Engine/Source/Runtime/AIModule/Private/BehaviorTree/Decorators/BTDecorator_BlackboardBase.cpp b/Engine/Source/Runtime/AIModule/Private/BehaviorTree/Decorators/BTDecorator_BlackboardBase.cpp
--- a/Engine/Source/Runtime/AIModule/Private/BehaviorTree/Decorators/BTDecorator_BlackboardBase.cpp
+++ b/Engine/Source/Runtime/AIModule/Private/BehaviorTree/Decorators/BTDecorator_BlackboardBase.cpp
@@ -20,7 +20,7 @@ void UBTDecorator_BlackboardBase::PostInitProperties()
 	BBKeyObserver = FOnBlackboardChange::CreateUObject(this, &UBTDecorator_BlackboardBase::OnBlackboardChange);
 }
 
-void UBTDecorator_BlackboardBase::InitializeFromAsset(class UBehaviorTree* Asset)
+void UBTDecorator_BlackboardBase::InitializeFromAsset(UBehaviorTree& Asset)
 {
 	Super::InitializeFromAsset(Asset);
 	BlackboardKey.CacheSelectedKey(GetBlackboardAsset());
diff --git a/Engine/Source/Runtime/AIModule/Private/BehaviorTree/Decorators/BTDecorator_CompareBBEntries.cpp b/Engine/Source/Runtime/AIModule/Private/BehaviorTree/Decorators/BTDecorator_CompareBBEntries.cpp
--- a/Engine/Source/Runtime/AIModule/Private/BehaviorTree/Decorators/BTDecorator_CompareBBEntries.cpp
+++ b/Engine/Source/Runtime/AIModule/Private/BehaviorTree/Decorators/BTDecorator_CompareBBEntries.cpp
@@ -19,7 +19,7 @@ void UBTDecorator_CompareBBEntries::PostInitProperties()
 	BBKeyObserver = FOnBlackboardChange::CreateUObject(this, &UBTDecorator_CompareBBEntries::OnBlackboardChange);
 }
 
-void UBTDecorator_CompareBBEntries::InitializeFromAsset(class UBehaviorTree* Asset)
+void UBTDecorator_CompareBBEntries::InitializeFromAsset(UBehaviorTree& Asset)
 {
 	Super::InitializeFromAsset(Asset);
 
diff --git a/Engine/Source/Runtime/AIModule/Private/BehaviorTree/Decorators/BTDecorator_ConeCheck.cpp b/Engine/Source/Runtime/AIModule/Private/BehaviorTree/Decorators/BTDecorator_ConeCheck.cpp
--- a/Engine/Source/Runtime/AIModule/Private/BehaviorTree/Decorators/BTDecorator_ConeCheck.cpp
+++ b/Engine/Source/Runtime/AIModule/Private/BehaviorTree/Decorators/BTDecorator_ConeCheck.cpp
@@ -26,7 +26,7 @@ UBTDecorator_ConeCheck::UBTDecorator_ConeCheck(const class FPostConstructInitial
 	ConeHalfAngle = 45.0f;
 }
 
-void UBTDecorator_ConeCheck::InitializeFromAsset(class UBehaviorTree* Asset)
+void UBTDecorator_ConeCheck::InitializeFromAsset(UBehaviorTree& Asset)
 {
 	Super::InitializeFromAsset(Asset);
 
diff --git a/Engine/Source/Runtime/AIModule/Private/BehaviorTree/Decorators/BTDecorator_DoesPathExist.cpp b/Engine/Source/Runtime/AIModule/Private/BehaviorTree/Decorators/BTDecorator_DoesPathExist.cpp
--- a/Engine/Source/Runtime/AIModule/Private/BehaviorTree/Decorators/BTDecorator_DoesPathExist.cpp
+++ b/Engine/Source/Runtime/AIModule/Private/BehaviorTree/Decorators/BTDecorator_DoesPathExist.cpp
@@ -27,7 +27,7 @@ UBTDecorator_DoesPathExist::UBTDecorator_DoesPathExist(const FPostConstructIniti
 	PathQueryType = EPathExistanceQueryType::HierarchicalQuery;
 }
 
-void UBTDecorator_DoesPathExist::InitializeFromAsset(UBehaviorTree* Asset)
+void UBTDecorator_DoesPathExist::InitializeFromAsset(UBehaviorTree& Asset)
 {
 	Super::InitializeFromAsset(Asset);
 
diff --git a/Engine/Source/Runtime/AIModule/Private/BehaviorTree/Decorators/BTDecorator_KeepInCone.cpp b/Engine/Source/Runtime/AIModule/Private/BehaviorTree/Decorators/BTDecorator_KeepInCone.cpp
--- a/Engine/Source/Runtime/AIModule/Private/BehaviorTree/Decorators/BTDecorator_KeepInCone.cpp
+++ b/Engine/Source/Runtime/AIModule/Private/BehaviorTree/Decorators/BTDecorator_KeepInCone.cpp
@@ -26,7 +26,7 @@ UBTDecorator_KeepInCone::UBTDecorator_KeepInCone(const class FPostConstructIniti
 	ConeHalfAngle = 45.0f;
 }
 
-void UBTDecorator_KeepInCone::InitializeFromAsset(class UBehaviorTree* Asset)
+void UBTDecorator_KeepInCone::InitializeFromAsset(UBehaviorTree& Asset)
 {
 	Super::InitializeFromAsset(Asset);
 
diff --git a/Engine/Source/Runtime/AIModule/Private/BehaviorTree/Services/BTService_BlackboardBase.cpp b/Engine/Source/Runtime/AIModule/Private/BehaviorTree/Services/BTService_BlackboardBase.cpp
--- a/Engine/Source/Runtime/AIModule/Private/BehaviorTree/Services/BTService_BlackboardBase.cpp
+++ b/Engine/Source/Runtime/AIModule/Private/BehaviorTree/Services/BTService_BlackboardBase.cpp
@@ -10,7 +10,7 @@ UBTService_BlackboardBase::UBTService_BlackboardBase(const class FPostConstructI
 	// empty KeySelector = allow everything
 }
 
-void UBTService_BlackboardBase::InitializeFromAsset(class UBehaviorTree* Asset)
+void UBTService_BlackboardBase::InitializeFromAsset(UBehaviorTree& Asset)
 {
 	Super::InitializeFromAsset(Asset);
 	BlackboardKey.CacheSelectedKey(GetBlackboardAsset());
diff --git a/Engine/Source/Runtime/AIModule/Private/BehaviorTree/Tasks/BTTask_BlackboardBase.cpp b/Engine/Source/Runtime/AIModule/Private/BehaviorTree/Tasks/BTTask_BlackboardBase.cpp
--- a/Engine/Source/Runtime/AIModule/Private/BehaviorTree/Tasks/BTTask_BlackboardBase.cpp
+++ b/Engine/Source/Runtime/AIModule/Private/BehaviorTree/Tasks/BTTask_BlackboardBase.cpp
@@ -10,7 +10,7 @@ UBTTask_BlackboardBase::UBTTask_BlackboardBase(const class FPostConstructInitial
 	// empty KeySelector = allow everything
 }
 
-void UBTTask_BlackboardBase::InitializeFromAsset(class UBehaviorTree* Asset)
+void UBTTask_BlackboardBase::InitializeFromAsset(UBehaviorTree& Asset)
 {
 	Super::InitializeFromAsset(Asset);
 	BlackboardKey.CacheSelectedKey(GetBlackboardAsset());
diff --git a/Engine/Source/Runtime/AIModule/Private/BehaviorTree/Tasks/BTTask_RunBehavior.cpp b/Engine/Source/Runtime/AIModule/Private/BehaviorTree/Tasks/BTTask_RunBehavior.cpp
--- a/Engine/Source/Runtime/AIModule/Private/BehaviorTree/Tasks/BTTask_RunBehavior.cpp
+++ b/Engine/Source/Runtime/AIModule/Private/BehaviorTree/Tasks/BTTask_RunBehavior.cpp
@@ -11,7 +11,7 @@ UBTTask_RunBehavior::UBTTask_RunBehavior(const class FPostConstructInitializePro
 
 EBTNodeResult::Type UBTTask_RunBehavior::ExecuteTask(UBehaviorTreeComponent* OwnerComp, uint8* NodeMemory)
 {
-	const bool bPushed = OwnerComp->PushInstance(BehaviorAsset);
+	const bool bPushed = OwnerComp->PushInstance(*BehaviorAsset);
 	return bPushed ? EBTNodeResult::InProgress : EBTNodeResult::Failed;
 }
 
@@ -154,7 +154,7 @@ void UBTTask_RunBehavior::InjectNodes(UBehaviorTreeComponent* OwnerComp, uint8*
 			for (int32 Idx = 0; Idx < NumInjectedDecorators; Idx++)
 			{
 				UBTDecorator* InstancedOb = Cast<UBTDecorator>(OwnerComp->NodeInstances[FirstNodeIdx + Idx]);
-				InstancedOb->InitializeFromAsset(BehaviorAsset);
+				InstancedOb->InitializeFromAsset(*BehaviorAsset);
 				InstancedOb->InitializeDecorator(ChildIdx);
 
 				if (!bAlreadyInjected)
