Remove unneeded includes

diff --git a/Engine/Source/Runtime/AIModule/Private/BehaviorTree/Tasks/BTTask_MakeNoise.cpp b/Engine/Source/Runtime/AIModule/Private/BehaviorTree/Tasks/BTTask_MakeNoise.cpp
--- a/Engine/Source/Runtime/AIModule/Private/BehaviorTree/Tasks/BTTask_MakeNoise.cpp
+++ b/Engine/Source/Runtime/AIModule/Private/BehaviorTree/Tasks/BTTask_MakeNoise.cpp
@@ -1,8 +1,6 @@
 // Copyright 1998-2014 Epic Games, Inc. All Rights Reserved.
 
 #include "AIModulePrivate.h"
-#include "SoundDefinitions.h"
-#include "Sound/ReverbVolume.h"
 #include "BehaviorTree/Tasks/BTTask_MakeNoise.h"
 
 UBTTask_MakeNoise::UBTTask_MakeNoise(const FPostConstructInitializeProperties& PCIP) 
@@ -27,7 +25,7 @@ EBTNodeResult::Type UBTTask_MakeNoise::ExecuteTask(UBehaviorTreeComponent* Owner
 
 	return EBTNodeResult::Failed;
 }
-
+	
 #if WITH_EDITOR
 
 FName UBTTask_MakeNoise::GetNodeIconName() const
