Removed forced loading of AITestSuite module #UE4

diff --git a/Engine/Source/Runtime/AIModule/Private/AIModule.cpp b/Engine/Source/Runtime/AIModule/Private/AIModule.cpp
--- a/Engine/Source/Runtime/AIModule/Private/AIModule.cpp
+++ b/Engine/Source/Runtime/AIModule/Private/AIModule.cpp
@@ -29,7 +29,6 @@ void FAIModule::StartupModule()
 	// This code will execute after your module is loaded into memory and after global variables initialization. We needs some place to load GameplayDebugger module so it's best place for it now.
 #if !(UE_BUILD_SHIPPING || UE_BUILD_TEST)
 	FModuleManager::LoadModulePtr< IModuleInterface >("GameplayDebugger");
-	FModuleManager::LoadModulePtr< IModuleInterface >("AITestSuite");
 #endif
 
 #if WITH_EDITOR && ENABLE_VISUAL_LOG
