Made AITestSuite a circular dependency of AIModule #UE4

- reenables AITestSuite and makes it available for all projects

diff --git a/Engine/Source/Runtime/AIModule/AIModule.Build.cs b/Engine/Source/Runtime/AIModule/AIModule.Build.cs
--- a/Engine/Source/Runtime/AIModule/AIModule.Build.cs
+++ b/Engine/Source/Runtime/AIModule/AIModule.Build.cs
@@ -43,6 +43,8 @@ namespace UnrealBuildTool.Rules
 			if (UEBuildConfiguration.bBuildEditor == true)
 			{
 				PrivateDependencyModuleNames.Add("UnrealEd");
+
+                CircularlyReferencedDependentModules.Add("AITestSuite");
 			}
 
             if (UEBuildConfiguration.bCompileRecast)
diff --git a/Engine/Source/Runtime/AIModule/Private/AIModule.cpp b/Engine/Source/Runtime/AIModule/Private/AIModule.cpp
--- a/Engine/Source/Runtime/AIModule/Private/AIModule.cpp
+++ b/Engine/Source/Runtime/AIModule/Private/AIModule.cpp
@@ -31,6 +31,10 @@ void FAIModule::StartupModule()
 	FModuleManager::LoadModulePtr< IModuleInterface >("GameplayDebugger");
 #endif
 
+#if WITH_EDITOR 
+	FModuleManager::LoadModulePtr< IModuleInterface >("AITestSuite");
+#endif // WITH_EDITOR 
+
 #if WITH_EDITOR && ENABLE_VISUAL_LOG
 	FVisualLog::Get().RegisterExtension(*EVisLogTags::TAG_EQS, &VisualLoggerExtension);
 #endif
diff --git a/Engine/Source/Runtime/Engine/Engine.Build.cs b/Engine/Source/Runtime/Engine/Engine.Build.cs
--- a/Engine/Source/Runtime/Engine/Engine.Build.cs
+++ b/Engine/Source/Runtime/Engine/Engine.Build.cs
@@ -53,7 +53,7 @@ public class Engine : ModuleRules
 				"EngineMessages",
 				"EngineSettings",
 				"SynthBenchmark",                
-                "AIModule",
+                "AIModule", // Here until all engine's dependencies of AI code are removed
 			}
 		);
 
