Fixed module dependency issue with AITestSuite #UE4

diff --git a/Engine/Source/Runtime/AIModule/AIModule.Build.cs b/Engine/Source/Runtime/AIModule/AIModule.Build.cs
--- a/Engine/Source/Runtime/AIModule/AIModule.Build.cs
+++ b/Engine/Source/Runtime/AIModule/AIModule.Build.cs
@@ -44,6 +44,7 @@ namespace UnrealBuildTool.Rules
 			{
 				PrivateDependencyModuleNames.Add("UnrealEd");
 
+                PrivateDependencyModuleNames.Add("AITestSuite");                
                 CircularlyReferencedDependentModules.Add("AITestSuite");
 			}
 
