Fixed brain and BB component crashing in unittesting due to not having Actor owner #UE4

diff --git a/Engine/Source/Runtime/AIModule/Private/BehaviorTree/BlackboardComponent.cpp b/Engine/Source/Runtime/AIModule/Private/BehaviorTree/BlackboardComponent.cpp
--- a/Engine/Source/Runtime/AIModule/Private/BehaviorTree/BlackboardComponent.cpp
+++ b/Engine/Source/Runtime/AIModule/Private/BehaviorTree/BlackboardComponent.cpp
@@ -17,11 +17,17 @@ void UBlackboardComponent::InitializeComponent()
 	Super::InitializeComponent();
 
 	// cache blackboard component if owner has one
-	BrainComp = GetOwner()->FindComponentByClass<UBrainComponent>();
-	if (BrainComp)
+	// note that it's a valid scenario for this component to not have an owner at all (at least in terms of unittesting)
+	AActor* Owner = GetOwner();
+	if (Owner)
 	{
-		BrainComp->CacheBlackboardComponent(this);
+		BrainComp = GetOwner()->FindComponentByClass<UBrainComponent>();
+		if (BrainComp)
+		{
+			BrainComp->CacheBlackboardComponent(this);
+		}
 	}
+
 	InitializeBlackboard(BlackboardAsset);
 }
 
diff --git a/Engine/Source/Runtime/AIModule/Private/BrainComponent.cpp b/Engine/Source/Runtime/AIModule/Private/BrainComponent.cpp
--- a/Engine/Source/Runtime/AIModule/Private/BrainComponent.cpp
+++ b/Engine/Source/Runtime/AIModule/Private/BrainComponent.cpp
@@ -251,10 +251,15 @@ void UBrainComponent::InitializeComponent()
 	Super::InitializeComponent();
 
 	// cache blackboard component if owner has one
-	BlackboardComp = GetOwner()->FindComponentByClass<UBlackboardComponent>();
-	if (BlackboardComp)
+	// note that it's a valid scenario for this component to not have an owner at all (at least in terms of unittesting)
+	AActor* Owner = GetOwner();
+	if (Owner)
 	{
-		BlackboardComp->CacheBrainComponent(this);
+		BlackboardComp = GetOwner()->FindComponentByClass<UBlackboardComponent>();
+		if (BlackboardComp)
+		{
+			BlackboardComp->CacheBrainComponent(this);
+		}
 	}
 }
 
