CRITICAL: Regression: All maps give the warning "Failed to find object 'Class None.'
" in the output log.

Improper use of FStringClassReference, fixed.

TTP 348408

diff --git a/Engine/Source/Runtime/AIModule/Private/AISystem.cpp b/Engine/Source/Runtime/AIModule/Private/AISystem.cpp
--- a/Engine/Source/Runtime/AIModule/Private/AISystem.cpp
+++ b/Engine/Source/Runtime/AIModule/Private/AISystem.cpp
@@ -36,7 +36,7 @@ void UAISystem::PostInitProperties()
 		UWorld* WorldOuter = Cast<UWorld>(GetOuter());
 		UObject* ManagersOuter = WorldOuter != NULL ? (UObject*)WorldOuter : (UObject*)this;
 		
-		TSubclassOf<UAIHotSpotManager> HotSpotManagerClass = LoadClass<UAIHotSpotManager>(NULL, *HotSpotManagerClassName.ToString(), NULL, LOAD_None, NULL);
+		TSubclassOf<UAIHotSpotManager> HotSpotManagerClass = HotSpotManagerClassName.ResolveClass();
 		if (HotSpotManagerClass)
 		{
 			HotSpotManager = NewObject<UAIHotSpotManager>(ManagersOuter, HotSpotManagerClass);
