added tooltip for UBTDecorator.bInverseCondition
ttp #348937

diff --git a/Engine/Source/Runtime/AIModule/Classes/BehaviorTree/BTDecorator.h b/Engine/Source/Runtime/AIModule/Classes/BehaviorTree/BTDecorator.h
--- a/Engine/Source/Runtime/AIModule/Classes/BehaviorTree/BTDecorator.h
+++ b/Engine/Source/Runtime/AIModule/Classes/BehaviorTree/BTDecorator.h
@@ -86,7 +86,8 @@ protected:
 	uint32 bShowInverseConditionDesc : 1;
 
 private:
-	UPROPERTY(Category=Condition, EditAnywhere)
+	/** if set, condition check result will be inversed */
+	UPROPERTY(Category = Condition, EditAnywhere)
 	uint32 bInverseCondition : 1;
 
 protected:
