fixed timing of first tick event for blueprint based BT service nodes

diff --git a/Engine/Source/Runtime/AIModule/Private/BehaviorTree/Services/BTService_BlueprintBase.cpp b/Engine/Source/Runtime/AIModule/Private/BehaviorTree/Services/BTService_BlueprintBase.cpp
--- a/Engine/Source/Runtime/AIModule/Private/BehaviorTree/Services/BTService_BlueprintBase.cpp
+++ b/Engine/Source/Runtime/AIModule/Private/BehaviorTree/Services/BTService_BlueprintBase.cpp
@@ -78,6 +78,10 @@ void UBTService_BlueprintBase::OnSearchStart(struct FBehaviorTreeSearchData& Sea
 	if (bImplementsReceiveSearchStart)
 	{
 		ReceiveSearchStart(SearchData.OwnerComp->GetOwner());
+
+		const float NextTickTime = FMath::FRandRange(FMath::Max(0.0f, Interval - RandomDeviation), (Interval + RandomDeviation));
+		uint8* NodeMemory = GetNodeMemory<uint8>(SearchData);
+		SetNextTickTime(NodeMemory, NextTickTime);
 	}
 	else
 	{
