Added null check to UE_VLOG_EQS #ue4

diff --git a/Engine/Source/Runtime/AIModule/Classes/EnvironmentQuery/EnvQueryDebugHelpers.h b/Engine/Source/Runtime/AIModule/Classes/EnvironmentQuery/EnvQueryDebugHelpers.h
--- a/Engine/Source/Runtime/AIModule/Classes/EnvironmentQuery/EnvQueryDebugHelpers.h
+++ b/Engine/Source/Runtime/AIModule/Classes/EnvironmentQuery/EnvQueryDebugHelpers.h
@@ -141,8 +141,10 @@ FArchive& operator<<(FArchive& Ar, EQSDebug::FQueryData& Data)
 
 #define UE_VLOG_EQS(Query, CategoryName, Verbosity) \
 { \
-	UEnvQueryDebugHelpers::LogQuery(Query, CategoryName.GetCategoryName(), ELogVerbosity::Verbosity, \
-		FString::Printf(TEXT("Executed EQS: \n - Name: '%s' (id=%d, option=%d),\n - All Items: %d,\n - ValidItems: %d"), *Query->QueryName, Query->QueryID, Query->OptionIndex, Query->ItemDetails.Num(), Query->NumValidItems)); \
+	if (Query != NULL) \
+	{ \
+		UEnvQueryDebugHelpers::LogQuery(Query, CategoryName.GetCategoryName(), ELogVerbosity::Verbosity, FString::Printf(TEXT("Executed EQS: \n - Name: '%s' (id=%d, option=%d),\n - All Items: %d,\n - ValidItems: %d"), *Query->QueryName, Query->QueryID, Query->OptionIndex, Query->ItemDetails.Num(), Query->NumValidItems)); \
+	}
 }
 
 #else
