Fixed CIS compilation #ue4

diff --git a/Engine/Source/Runtime/AIModule/Classes/EnvironmentQuery/EnvQueryDebugHelpers.h b/Engine/Source/Runtime/AIModule/Classes/EnvironmentQuery/EnvQueryDebugHelpers.h
--- a/Engine/Source/Runtime/AIModule/Classes/EnvironmentQuery/EnvQueryDebugHelpers.h
+++ b/Engine/Source/Runtime/AIModule/Classes/EnvironmentQuery/EnvQueryDebugHelpers.h
@@ -144,7 +144,7 @@ FArchive& operator<<(FArchive& Ar, EQSDebug::FQueryData& Data)
 	if (Query != NULL) \
 	{ \
 		UEnvQueryDebugHelpers::LogQuery(Query, CategoryName.GetCategoryName(), ELogVerbosity::Verbosity, FString::Printf(TEXT("Executed EQS: \n - Name: '%s' (id=%d, option=%d),\n - All Items: %d,\n - ValidItems: %d"), *Query->QueryName, Query->QueryID, Query->OptionIndex, Query->ItemDetails.Num(), Query->NumValidItems)); \
-	}
+	} \
 }
 
 #else
