UE4 - Fix CIS

diff --git a/Engine/Source/Runtime/AIModule/Classes/EnvironmentQuery/EnvQueryDebugHelpers.h b/Engine/Source/Runtime/AIModule/Classes/EnvironmentQuery/EnvQueryDebugHelpers.h
--- a/Engine/Source/Runtime/AIModule/Classes/EnvironmentQuery/EnvQueryDebugHelpers.h
+++ b/Engine/Source/Runtime/AIModule/Classes/EnvironmentQuery/EnvQueryDebugHelpers.h
@@ -170,7 +170,7 @@ class AIMODULE_API UEnvQueryDebugHelpers : public UObject
 FORCEINLINE void UEnvQueryDebugHelpers::LogQuery(struct FEnvQueryInstance* Query, const FName& CategoryName, ELogVerbosity::Type Type, const FString& AdditionalLogInfo)
 {
 #if ENABLE_VISUAL_LOG
-	if (GEngine && GEngine->bDisableAILogging || Query == NULL || FVisualLogger::Get().IsRecording() == false || (FVisualLogger::Get().IsBlockedForAllCategories() && FVisualLogger::Get().GetWhiteList().Find(CategoryName) == INDEX_NONE))
+	if ((GEngine && GEngine->bDisableAILogging) || Query == NULL || FVisualLogger::Get().IsRecording() == false || (FVisualLogger::Get().IsBlockedForAllCategories() && FVisualLogger::Get().GetWhiteList().Find(CategoryName) == INDEX_NONE))
 	{
 		return;
 	}
