Renamed UBTFunctionLibrary::GetBlackboard to UBTFunctionLibrary::GetOwnersBlackboard to avoid confusion with UAIBlueprintHelperLibrary::GetBlackboard #UE4

- also, made a little optimization to UBTFunctionLibrary's functions

diff --git a/Engine/Config/BaseEngine.ini b/Engine/Config/BaseEngine.ini
--- a/Engine/Config/BaseEngine.ini
+++ b/Engine/Config/BaseEngine.ini
@@ -586,6 +586,7 @@ bUseStreamingPause=false
 +ActiveClassRedirects=(OldClassName="PathFollowingComponent",NewClassName="/Script/AIModule.PathFollowingComponent")
 +ActiveClassRedirects=(OldClassName="PawnSensingComponent",NewClassName="/Script/AIModule.PawnSensingComponent")
 +K2FieldRedirects=(OldFieldName="BTTask_MoveDirectlyToward.bForceMoveToLocation",NewFieldName="bDisablePathUpdateOnGoalLocationChange")
++K2FieldRedirects=(OldFieldName="BTFunctionLibrary.GetBlackboard", NewFieldName="BTFunctionLibrary.GetOwnersBlackboard")
 
 ; Renamed SpriteComponent to BillboardComponent
 +ActiveClassRedirects=(OldClassName="SpriteComponent",NewClassName="BillboardComponent")
diff --git a/Engine/Source/Runtime/AIModule/Classes/BehaviorTree/BTFunctionLibrary.h b/Engine/Source/Runtime/AIModule/Classes/BehaviorTree/BTFunctionLibrary.h
--- a/Engine/Source/Runtime/AIModule/Classes/BehaviorTree/BTFunctionLibrary.h
+++ b/Engine/Source/Runtime/AIModule/Classes/BehaviorTree/BTFunctionLibrary.h
@@ -3,13 +3,16 @@
 #pragma once
 #include "BTFunctionLibrary.generated.h"
 
+class UBlackboardComponent;
+
 UCLASS(meta=(RestrictedToClasses="BTNode"))
-class UBTFunctionLibrary : public UBlueprintFunctionLibrary
+class AIMODULE_API UBTFunctionLibrary : public UBlueprintFunctionLibrary
 {
-	GENERATED_UCLASS_BODY()
+	GENERATED_BODY()
 
+public:
 	UFUNCTION(BlueprintPure, Category="AI|BehaviorTree", Meta=(HidePin="NodeOwner", DefaultToSelf="NodeOwner"))
-	static class UBlackboardComponent* GetBlackboard(UBTNode* NodeOwner);
+	static UBlackboardComponent* GetOwnersBlackboard(UBTNode* NodeOwner);
 
 	UFUNCTION(BlueprintPure, Category="AI|BehaviorTree", Meta=(HidePin="NodeOwner", DefaultToSelf="NodeOwner"))
 	static UObject* GetBlackboardValueAsObject(UBTNode* NodeOwner, const struct FBlackboardKeySelector& Key);
diff --git a/Engine/Source/Runtime/AIModule/Private/BehaviorTree/BTFunctionLibrary.cpp b/Engine/Source/Runtime/AIModule/Private/BehaviorTree/BTFunctionLibrary.cpp
--- a/Engine/Source/Runtime/AIModule/Private/BehaviorTree/BTFunctionLibrary.cpp
+++ b/Engine/Source/Runtime/AIModule/Private/BehaviorTree/BTFunctionLibrary.cpp
@@ -6,79 +6,98 @@
 #include "BehaviorTree/Blackboard/BlackboardKeyAllTypes.h"
 #include "BehaviorTree/BTFunctionLibrary.h"
 
-UBTFunctionLibrary::UBTFunctionLibrary(const FObjectInitializer& ObjectInitializer) : Super(ObjectInitializer)
+namespace
 {
+	FORCEINLINE UBlackboardComponent* GetBlackboard(UBTNode& BTNode)
+	{
+		check(BTNode.GetOuter());
+		check(BTNode.GetClass()->HasAnyClassFlags(CLASS_CompiledFromBlueprint) && "This function call is valid only for BP-implemented BT nodes");
+		check(Cast<UBehaviorTreeComponent>(BTNode.GetOuter()));
+
+		return ((UBehaviorTreeComponent*)BTNode.GetOuter())->GetBlackboardComponent();
+	}
 }
 
-UBlackboardComponent* UBTFunctionLibrary::GetBlackboard(UBTNode* NodeOwner)
+UBlackboardComponent* UBTFunctionLibrary::GetOwnersBlackboard(UBTNode* NodeOwner)
 {
-	UBehaviorTreeComponent* BTComp = NodeOwner ? Cast<UBehaviorTreeComponent>(NodeOwner->GetOuter()) : NULL;
-	return BTComp ? BTComp->GetBlackboardComponent() : NULL;
+	check(NodeOwner != NULL);
+	return GetBlackboard(*NodeOwner);
 }
 
 UObject* UBTFunctionLibrary::GetBlackboardValueAsObject(UBTNode* NodeOwner, const struct FBlackboardKeySelector& Key)
 {
-	const UBlackboardComponent* BlackboardComp = GetBlackboard(NodeOwner);
+	check(NodeOwner != NULL);
+	const UBlackboardComponent* BlackboardComp = GetBlackboard(*NodeOwner);
 	return BlackboardComp ? BlackboardComp->GetValueAsObject(Key.SelectedKeyName) : NULL;
 }
 
 AActor* UBTFunctionLibrary::GetBlackboardValueAsActor(UBTNode* NodeOwner, const struct FBlackboardKeySelector& Key)
 {
-	const UBlackboardComponent* BlackboardComp = GetBlackboard(NodeOwner);
+	check(NodeOwner != NULL);
+	const UBlackboardComponent* BlackboardComp = GetBlackboard(*NodeOwner);
 	return BlackboardComp ? Cast<AActor>(BlackboardComp->GetValueAsObject(Key.SelectedKeyName)) : NULL;
 }
 
 UClass* UBTFunctionLibrary::GetBlackboardValueAsClass(UBTNode* NodeOwner, const struct FBlackboardKeySelector& Key)
 {
-	const UBlackboardComponent* BlackboardComp = GetBlackboard(NodeOwner);
+	check(NodeOwner != NULL);
+	const UBlackboardComponent* BlackboardComp = GetBlackboard(*NodeOwner);
 	return BlackboardComp ? BlackboardComp->GetValueAsClass(Key.SelectedKeyName) : NULL;
 }
 
 uint8 UBTFunctionLibrary::GetBlackboardValueAsEnum(UBTNode* NodeOwner, const struct FBlackboardKeySelector& Key)
 {
-	const UBlackboardComponent* BlackboardComp = GetBlackboard(NodeOwner);
+	check(NodeOwner != NULL);
+	const UBlackboardComponent* BlackboardComp = GetBlackboard(*NodeOwner);
 	return BlackboardComp ? BlackboardComp->GetValueAsEnum(Key.SelectedKeyName) : 0;
 }
 
 int32 UBTFunctionLibrary::GetBlackboardValueAsInt(UBTNode* NodeOwner, const struct FBlackboardKeySelector& Key)
 {
-	const UBlackboardComponent* BlackboardComp = GetBlackboard(NodeOwner);
+	check(NodeOwner != NULL);
+	const UBlackboardComponent* BlackboardComp = GetBlackboard(*NodeOwner);
 	return BlackboardComp ? BlackboardComp->GetValueAsInt(Key.SelectedKeyName) : 0;
 }
 
 float UBTFunctionLibrary::GetBlackboardValueAsFloat(UBTNode* NodeOwner, const struct FBlackboardKeySelector& Key)
 {
-	const UBlackboardComponent* BlackboardComp = GetBlackboard(NodeOwner);
+	check(NodeOwner != NULL);
+	const UBlackboardComponent* BlackboardComp = GetBlackboard(*NodeOwner);
 	return BlackboardComp ? BlackboardComp->GetValueAsFloat(Key.SelectedKeyName) : 0.0f;
 }
 
 bool UBTFunctionLibrary::GetBlackboardValueAsBool(UBTNode* NodeOwner, const struct FBlackboardKeySelector& Key)
 {
-	const UBlackboardComponent* BlackboardComp = GetBlackboard(NodeOwner);
+	check(NodeOwner != NULL);
+	const UBlackboardComponent* BlackboardComp = GetBlackboard(*NodeOwner);
 	return BlackboardComp ? BlackboardComp->GetValueAsBool(Key.SelectedKeyName) : false;
 }
 
 FString UBTFunctionLibrary::GetBlackboardValueAsString(UBTNode* NodeOwner, const struct FBlackboardKeySelector& Key)
 {
-	const UBlackboardComponent* BlackboardComp = GetBlackboard(NodeOwner);
+	check(NodeOwner != NULL);
+	const UBlackboardComponent* BlackboardComp = GetBlackboard(*NodeOwner);
 	return BlackboardComp ? BlackboardComp->GetValueAsString(Key.SelectedKeyName) : FString();
 }
 
 FName UBTFunctionLibrary::GetBlackboardValueAsName(UBTNode* NodeOwner, const struct FBlackboardKeySelector& Key)
 {
-	const UBlackboardComponent* BlackboardComp = GetBlackboard(NodeOwner);
+	check(NodeOwner != NULL);
+	const UBlackboardComponent* BlackboardComp = GetBlackboard(*NodeOwner);
 	return BlackboardComp ? BlackboardComp->GetValueAsName(Key.SelectedKeyName) : NAME_None;
 }
 
 FVector UBTFunctionLibrary::GetBlackboardValueAsVector(UBTNode* NodeOwner, const struct FBlackboardKeySelector& Key)
 {
-	const UBlackboardComponent* BlackboardComp = GetBlackboard(NodeOwner);
+	check(NodeOwner != NULL);
+	const UBlackboardComponent* BlackboardComp = GetBlackboard(*NodeOwner);
 	return BlackboardComp ? BlackboardComp->GetValueAsVector(Key.SelectedKeyName) : FVector::ZeroVector;
 }
 
 void UBTFunctionLibrary::SetBlackboardValueAsObject(UBTNode* NodeOwner, const struct FBlackboardKeySelector& Key, UObject* Value)
 {
-	UBlackboardComponent* BlackboardComp = GetBlackboard(NodeOwner);
+	check(NodeOwner != NULL);
+	UBlackboardComponent* BlackboardComp = GetBlackboard(*NodeOwner);
 	if (BlackboardComp != NULL)
 	{
 		BlackboardComp->SetValueAsObject(Key.SelectedKeyName, Value);
@@ -87,7 +106,8 @@ void UBTFunctionLibrary::SetBlackboardValueAsObject(UBTNode* NodeOwner, const st
 
 void UBTFunctionLibrary::SetBlackboardValueAsClass(UBTNode* NodeOwner, const struct FBlackboardKeySelector& Key, UClass* Value)
 {
-	UBlackboardComponent* BlackboardComp = GetBlackboard(NodeOwner);
+	check(NodeOwner != NULL);
+	UBlackboardComponent* BlackboardComp = GetBlackboard(*NodeOwner);
 	if (BlackboardComp != NULL)
 	{
 		BlackboardComp->SetValueAsClass(Key.SelectedKeyName, Value);
@@ -96,7 +116,8 @@ void UBTFunctionLibrary::SetBlackboardValueAsClass(UBTNode* NodeOwner, const str
 
 void UBTFunctionLibrary::SetBlackboardValueAsEnum(UBTNode* NodeOwner, const struct FBlackboardKeySelector& Key, uint8 Value)
 {
-	UBlackboardComponent* BlackboardComp = GetBlackboard(NodeOwner);
+	check(NodeOwner != NULL);
+	UBlackboardComponent* BlackboardComp = GetBlackboard(*NodeOwner);
 	if (BlackboardComp != NULL)
 	{
 		BlackboardComp->SetValueAsEnum(Key.SelectedKeyName, Value);
@@ -105,7 +126,8 @@ void UBTFunctionLibrary::SetBlackboardValueAsEnum(UBTNode* NodeOwner, const stru
 
 void UBTFunctionLibrary::SetBlackboardValueAsInt(UBTNode* NodeOwner, const struct FBlackboardKeySelector& Key, int32 Value)
 {
-	UBlackboardComponent* BlackboardComp = GetBlackboard(NodeOwner);
+	check(NodeOwner != NULL);
+	UBlackboardComponent* BlackboardComp = GetBlackboard(*NodeOwner);
 	if (BlackboardComp != NULL)
 	{
 		BlackboardComp->SetValueAsInt(Key.SelectedKeyName, Value);
@@ -114,7 +136,8 @@ void UBTFunctionLibrary::SetBlackboardValueAsInt(UBTNode* NodeOwner, const struc
 
 void UBTFunctionLibrary::SetBlackboardValueAsFloat(UBTNode* NodeOwner, const struct FBlackboardKeySelector& Key, float Value)
 {
-	UBlackboardComponent* BlackboardComp = GetBlackboard(NodeOwner);
+	check(NodeOwner != NULL);
+	UBlackboardComponent* BlackboardComp = GetBlackboard(*NodeOwner);
 	if (BlackboardComp != NULL)
 	{
 		BlackboardComp->SetValueAsFloat(Key.SelectedKeyName, Value);
@@ -123,7 +146,8 @@ void UBTFunctionLibrary::SetBlackboardValueAsFloat(UBTNode* NodeOwner, const str
 
 void UBTFunctionLibrary::SetBlackboardValueAsBool(UBTNode* NodeOwner, const struct FBlackboardKeySelector& Key, bool Value)
 {
-	UBlackboardComponent* BlackboardComp = GetBlackboard(NodeOwner);
+	check(NodeOwner != NULL);
+	UBlackboardComponent* BlackboardComp = GetBlackboard(*NodeOwner);
 	if (BlackboardComp != NULL)
 	{
 		BlackboardComp->SetValueAsBool(Key.SelectedKeyName, Value);
@@ -132,7 +156,8 @@ void UBTFunctionLibrary::SetBlackboardValueAsBool(UBTNode* NodeOwner, const stru
 
 void UBTFunctionLibrary::SetBlackboardValueAsString(UBTNode* NodeOwner, const struct FBlackboardKeySelector& Key, FString Value)
 {
-	UBlackboardComponent* BlackboardComp = GetBlackboard(NodeOwner);
+	check(NodeOwner != NULL);
+	UBlackboardComponent* BlackboardComp = GetBlackboard(*NodeOwner);
 	if (BlackboardComp != NULL)
 	{
 		BlackboardComp->SetValueAsString(Key.SelectedKeyName, Value);
@@ -141,7 +166,8 @@ void UBTFunctionLibrary::SetBlackboardValueAsString(UBTNode* NodeOwner, const st
 
 void UBTFunctionLibrary::SetBlackboardValueAsName(UBTNode* NodeOwner, const struct FBlackboardKeySelector& Key, FName Value)
 {
-	UBlackboardComponent* BlackboardComp = GetBlackboard(NodeOwner);
+	check(NodeOwner != NULL);
+	UBlackboardComponent* BlackboardComp = GetBlackboard(*NodeOwner);
 	if (BlackboardComp != NULL)
 	{
 		BlackboardComp->SetValueAsName(Key.SelectedKeyName, Value);
@@ -150,7 +176,8 @@ void UBTFunctionLibrary::SetBlackboardValueAsName(UBTNode* NodeOwner, const stru
 
 void UBTFunctionLibrary::SetBlackboardValueAsVector(UBTNode* NodeOwner, const struct FBlackboardKeySelector& Key, FVector Value)
 {
-	UBlackboardComponent* BlackboardComp = GetBlackboard(NodeOwner);
+	check(NodeOwner != NULL);
+	UBlackboardComponent* BlackboardComp = GetBlackboard(*NodeOwner);
 	if (BlackboardComp != NULL)
 	{
 		BlackboardComp->SetValueAsVector(Key.SelectedKeyName, Value);
@@ -159,7 +186,8 @@ void UBTFunctionLibrary::SetBlackboardValueAsVector(UBTNode* NodeOwner, const st
 
 void UBTFunctionLibrary::ClearBlackboardValueAsVector(UBTNode* NodeOwner, const struct FBlackboardKeySelector& Key)
 {
-	UBlackboardComponent* BlackboardComp = GetBlackboard(NodeOwner);
+	check(NodeOwner != NULL);
+	UBlackboardComponent* BlackboardComp = GetBlackboard(*NodeOwner);
 	if (BlackboardComp != NULL)
 	{
 		BlackboardComp->ClearValueAsVector(Key.SelectedKeyName);
diff --git a/Engine/Source/Runtime/AIModule/Private/Blueprint/AIBlueprintHelperLibrary.cpp b/Engine/Source/Runtime/AIModule/Private/Blueprint/AIBlueprintHelperLibrary.cpp
--- a/Engine/Source/Runtime/AIModule/Private/Blueprint/AIBlueprintHelperLibrary.cpp
+++ b/Engine/Source/Runtime/AIModule/Private/Blueprint/AIBlueprintHelperLibrary.cpp
@@ -159,17 +159,20 @@ APawn* UAIBlueprintHelperLibrary::SpawnAI(class UObject* WorldContextObject, UBl
 
 UBlackboardComponent* UAIBlueprintHelperLibrary::GetBlackboard(AActor* Target)
 {
-	UBlackboardComponent* BlackboardComp = NULL;
+	UBlackboardComponent* BlackboardComp = nullptr;
 
-	APawn* TargetPawn = Cast<APawn>(Target);
-	if (TargetPawn && TargetPawn->GetController())
+	if (Target != nullptr)
 	{
-		BlackboardComp = TargetPawn->GetController()->FindComponentByClass<UBlackboardComponent>();
-	}
+		APawn* TargetPawn = Cast<APawn>(Target);
+		if (TargetPawn && TargetPawn->GetController())
+		{
+			BlackboardComp = TargetPawn->GetController()->FindComponentByClass<UBlackboardComponent>();
+		}
 
-	if (BlackboardComp == NULL && Target)
-	{
-		BlackboardComp = Target->FindComponentByClass<UBlackboardComponent>();
+		if (BlackboardComp == nullptr)
+		{
+			BlackboardComp = Target->FindComponentByClass<UBlackboardComponent>();
+		}
 	}
 
 	return BlackboardComp;
