Post TSubobjectPtr<> removal cleanup

diff --git a/Engine/Source/Runtime/AIModule/Classes/AIController.h b/Engine/Source/Runtime/AIModule/Classes/AIController.h
--- a/Engine/Source/Runtime/AIModule/Classes/AIController.h
+++ b/Engine/Source/Runtime/AIModule/Classes/AIController.h
@@ -94,12 +94,11 @@ private:
 	/** Component used for pathfinding and querying environment's navigation. */
 	UPROPERTY()
 	UNavigationComponent* NavComponent;
-public:
 
-private:
 	/** Component used for moving along a path. */
 	UPROPERTY()
 	UPathFollowingComponent* PathFollowingComponent;
+
 public:
 
 	/** Component responsible for behaviors. */
@@ -108,14 +107,12 @@ public:
 
 	UPROPERTY()
 	UAIPerceptionComponent* PerceptionComponent;
-	
-public:
 
 private:
 	UPROPERTY(BlueprintReadOnly, Category = AI, meta = (AllowPrivateAccess = "true"))
 	UPawnActionsComponent* ActionsComp;
-public:
 
+public:
 	/** Event called when PossessedPawn is possesed by this controller. */
 	UFUNCTION(BlueprintImplementableEvent, Category = "AI")
 	void OnPossess(APawn* PossessedPawn);
diff --git a/Engine/Source/Runtime/Engine/Classes/Atmosphere/AtmosphericFog.h b/Engine/Source/Runtime/Engine/Classes/Atmosphere/AtmosphericFog.h
--- a/Engine/Source/Runtime/Engine/Classes/Atmosphere/AtmosphericFog.h
+++ b/Engine/Source/Runtime/Engine/Classes/Atmosphere/AtmosphericFog.h
@@ -30,7 +30,6 @@ public:
 	virtual void PostActorCreated() override;
 #endif
 
-public:
 	/** Returns AtmosphericFogComponent subobject **/
 	FORCEINLINE class UAtmosphericFogComponent* GetAtmosphericFogComponent() { return AtmosphericFogComponent; }
 #if WITH_EDITORONLY_DATA
diff --git a/Engine/Source/Runtime/Engine/Classes/Camera/CameraActor.h b/Engine/Source/Runtime/Engine/Classes/Camera/CameraActor.h
--- a/Engine/Source/Runtime/Engine/Classes/Camera/CameraActor.h
+++ b/Engine/Source/Runtime/Engine/Classes/Camera/CameraActor.h
@@ -18,11 +18,10 @@ private:
 	UPROPERTY(Category="AutoPlayerActivation", EditAnywhere)
 	TEnumAsByte<EAutoReceiveInput::Type> AutoActivateForPlayer;
 
-public:
-private:
 	/** The camera component for this camera */
 	UPROPERTY(Category = CameraActor, VisibleAnywhere, BlueprintReadOnly, meta = (AllowPrivateAccess = "true"))
 	class UCameraComponent* CameraComponent;
+
 public:
 
 	/** If this CameraActor is being used to preview a CameraAnim in the editor, this is the anim being previewed. */
@@ -63,7 +62,6 @@ public:
 	ENGINE_API virtual void BeginPlay() override;
 	// End AActor interface
 
-public:
 	/** Returns CameraComponent subobject **/
 	FORCEINLINE class UCameraComponent* GetCameraComponent() const { return CameraComponent; }
 };
diff --git a/Engine/Source/Runtime/Engine/Classes/Camera/CameraAnimInst.h b/Engine/Source/Runtime/Engine/Classes/Camera/CameraAnimInst.h
--- a/Engine/Source/Runtime/Engine/Classes/Camera/CameraAnimInst.h
+++ b/Engine/Source/Runtime/Engine/Classes/Camera/CameraAnimInst.h
@@ -16,12 +16,10 @@ class ENGINE_API UCameraAnimInst : public UObject
 	UPROPERTY()
 	class UCameraAnim* CamAnim;
 
-protected:
 private:
 	/** the UInterpGroupInst used to do the interpolation */
 	UPROPERTY(instanced)
 	class UInterpGroupInst* InterpGroupInst;
-protected:
 
 public:
 	/** Current time for the animation */
diff --git a/Engine/Source/Runtime/Engine/Classes/Engine/ExponentialHeightFog.h b/Engine/Source/Runtime/Engine/Classes/Engine/ExponentialHeightFog.h
--- a/Engine/Source/Runtime/Engine/Classes/Engine/ExponentialHeightFog.h
+++ b/Engine/Source/Runtime/Engine/Classes/Engine/ExponentialHeightFog.h
@@ -28,7 +28,6 @@ public:
 	virtual void PostInitializeComponents() override;
 	//End AActor Interface
 
-public:
 	/** Returns Component subobject **/
 	FORCEINLINE class UExponentialHeightFogComponent* GetComponent() const { return Component; }
 };
diff --git a/Engine/Source/Runtime/Engine/Classes/Engine/Light.h b/Engine/Source/Runtime/Engine/Classes/Engine/Light.h
--- a/Engine/Source/Runtime/Engine/Classes/Engine/Light.h
+++ b/Engine/Source/Runtime/Engine/Classes/Engine/Light.h
@@ -73,7 +73,6 @@ public:
 	virtual bool IsLevelBoundsRelevant() const override { return false; }
 	// End AActor interface.
 
-public:
 	/** Returns LightComponent subobject **/
 	FORCEINLINE class ULightComponent* GetLightComponent() const { return LightComponent; }
 };
diff --git a/Engine/Source/Runtime/Engine/Classes/Engine/NavigationObjectBase.h b/Engine/Source/Runtime/Engine/Classes/Engine/NavigationObjectBase.h
--- a/Engine/Source/Runtime/Engine/Classes/Engine/NavigationObjectBase.h
+++ b/Engine/Source/Runtime/Engine/Classes/Engine/NavigationObjectBase.h
@@ -12,15 +12,11 @@ class ENGINE_API ANavigationObjectBase : public AActor, public INavAgentInterfac
 private:
 	UPROPERTY()
 	class UCapsuleComponent* CapsuleComponent;
-public:
 
-private:
 	/** Normal editor sprite. */
 	UPROPERTY()
 	class UBillboardComponent* GoodSprite;
-public:
 
-private:
 	/** Used to draw bad collision intersection in editor. */
 	UPROPERTY()
 	class UBillboardComponent* BadSprite;
diff --git a/Engine/Source/Runtime/Engine/Classes/Engine/Note.h b/Engine/Source/Runtime/Engine/Classes/Engine/Note.h
--- a/Engine/Source/Runtime/Engine/Classes/Engine/Note.h
+++ b/Engine/Source/Runtime/Engine/Classes/Engine/Note.h
@@ -35,7 +35,6 @@ public:
 	// End AActor Interface
 #endif
 
-public:
 #if WITH_EDITORONLY_DATA
 	/** Returns SpriteComponent subobject **/
 	FORCEINLINE class UBillboardComponent* GetSpriteComponent() const { return SpriteComponent; }
diff --git a/Engine/Source/Runtime/Engine/Classes/Engine/ReflectionCapture.h b/Engine/Source/Runtime/Engine/Classes/Engine/ReflectionCapture.h
--- a/Engine/Source/Runtime/Engine/Classes/Engine/ReflectionCapture.h
+++ b/Engine/Source/Runtime/Engine/Classes/Engine/ReflectionCapture.h
@@ -26,8 +26,6 @@ public:
 	virtual void PostEditMove(bool bFinished) override;
 #endif // WITH_EDITOR
 
-
-public:
 	/** Returns CaptureComponent subobject **/
 	FORCEINLINE class UReflectionCaptureComponent* GetCaptureComponent() const { return CaptureComponent; }
 #if WITH_EDITORONLY_DATA
diff --git a/Engine/Source/Runtime/Engine/Classes/Engine/SceneCapture2D.h b/Engine/Source/Runtime/Engine/Classes/Engine/SceneCapture2D.h
--- a/Engine/Source/Runtime/Engine/Classes/Engine/SceneCapture2D.h
+++ b/Engine/Source/Runtime/Engine/Classes/Engine/SceneCapture2D.h
@@ -33,7 +33,6 @@ public:
 	UFUNCTION(BlueprintCallable, Category="Rendering")
 	void OnInterpToggle(bool bEnable);
 
-public:
 	/** Returns CaptureComponent2D subobject **/
 	FORCEINLINE class USceneCaptureComponent2D* GetCaptureComponent2D() const { return CaptureComponent2D; }
 	/** Returns DrawFrustum subobject **/
diff --git a/Engine/Source/Runtime/Engine/Classes/Engine/SceneCaptureCube.h b/Engine/Source/Runtime/Engine/Classes/Engine/SceneCaptureCube.h
--- a/Engine/Source/Runtime/Engine/Classes/Engine/SceneCaptureCube.h
+++ b/Engine/Source/Runtime/Engine/Classes/Engine/SceneCaptureCube.h
@@ -37,7 +37,6 @@ public:
 	UFUNCTION(BlueprintCallable, Category = "Rendering")
 	void OnInterpToggle(bool bEnable);
 
-public:
 	/** Returns CaptureComponentCube subobject **/
 	FORCEINLINE class USceneCaptureComponentCube* GetCaptureComponentCube() const { return CaptureComponentCube; }
 	/** Returns DrawFrustum subobject **/
diff --git a/Engine/Source/Runtime/Engine/Classes/GameFramework/WheeledVehicle.h b/Engine/Source/Runtime/Engine/Classes/GameFramework/WheeledVehicle.h
--- a/Engine/Source/Runtime/Engine/Classes/GameFramework/WheeledVehicle.h
+++ b/Engine/Source/Runtime/Engine/Classes/GameFramework/WheeledVehicle.h
@@ -37,7 +37,6 @@ public:
 	virtual void DisplayDebug(class UCanvas* Canvas, const FDebugDisplayInfo& DebugDisplay, float& YL, float& YPos) override;
 	// End Actor interface
 
-public:
 	/** Returns Mesh subobject **/
 	FORCEINLINE class USkeletalMeshComponent* GetMesh() const { return Mesh; }
 	/** Returns VehicleMovement subobject **/
diff --git a/Engine/Source/Runtime/Engine/Classes/Materials/MaterialInstanceActor.h b/Engine/Source/Runtime/Engine/Classes/Materials/MaterialInstanceActor.h
--- a/Engine/Source/Runtime/Engine/Classes/Materials/MaterialInstanceActor.h
+++ b/Engine/Source/Runtime/Engine/Classes/Materials/MaterialInstanceActor.h
@@ -32,7 +32,6 @@ public:
 #endif // WITH_EDITOR
 	// End UObject Interface
 
-public:
 #if WITH_EDITORONLY_DATA
 	/** Returns SpriteComponent subobject **/
 	FORCEINLINE UBillboardComponent* GetSpriteComponent() const { return SpriteComponent; }
diff --git a/Engine/Source/Runtime/Engine/Classes/PhysicsEngine/PhysicsThruster.h b/Engine/Source/Runtime/Engine/Classes/PhysicsEngine/PhysicsThruster.h
--- a/Engine/Source/Runtime/Engine/Classes/PhysicsEngine/PhysicsThruster.h
+++ b/Engine/Source/Runtime/Engine/Classes/PhysicsEngine/PhysicsThruster.h
@@ -26,10 +26,8 @@ private:
 
 	UPROPERTY()
 	class UBillboardComponent* SpriteComponent;
-public:
 #endif
 
-
 public:
 	/** Returns ThrusterComponent subobject **/
 	FORCEINLINE class UPhysicsThrusterComponent* GetThrusterComponent() const { return ThrusterComponent; }
diff --git a/Engine/Source/Runtime/GameplayAbilities/Public/Abilities/GameplayAbilityWorldReticle_ActorVisualization.h b/Engine/Source/Runtime/GameplayAbilities/Public/Abilities/GameplayAbilityWorldReticle_ActorVisualization.h
--- a/Engine/Source/Runtime/GameplayAbilities/Public/Abilities/GameplayAbilityWorldReticle_ActorVisualization.h
+++ b/Engine/Source/Runtime/GameplayAbilities/Public/Abilities/GameplayAbilityWorldReticle_ActorVisualization.h
@@ -29,9 +29,6 @@ public:
 	/** Overridable function called whenever this actor is being removed from a level */
 	virtual void EndPlay(const EEndPlayReason::Type EndPlayReason);
 
-protected:
-
-public:
 	/** Returns CollisionComponent subobject **/
 	FORCEINLINE class UCapsuleComponent* GetCollisionComponent() { return CollisionComponent; }
 };
diff --git a/Engine/Source/Runtime/GameplayAbilities/Public/AbilitySystemTestPawn.h b/Engine/Source/Runtime/GameplayAbilities/Public/AbilitySystemTestPawn.h
--- a/Engine/Source/Runtime/GameplayAbilities/Public/AbilitySystemTestPawn.h
+++ b/Engine/Source/Runtime/GameplayAbilities/Public/AbilitySystemTestPawn.h
@@ -29,7 +29,6 @@ public:
 
 	static FName AbilitySystemComponentName;
 
-public:
 	/** Returns AbilitySystemComponent subobject **/
 	FORCEINLINE class UAbilitySystemComponent* GetAbilitySystemComponent() { return AbilitySystemComponent; }
 };
