Changed GameplayDebugger to use AutoStartup functionality (bIsAutoStartupModule flag) #ue4
Moved AutoStartup initialization to FEngineLoop::Init() function (ELoadingPhase::PostEngineInit phase) - after talk with Jaroslaw Palczynski about it. #ue4

diff --git a/Engine/Source/Developer/GameplayDebugger/GameplayDebugger.Build.cs b/Engine/Source/Developer/GameplayDebugger/GameplayDebugger.Build.cs
--- a/Engine/Source/Developer/GameplayDebugger/GameplayDebugger.Build.cs
+++ b/Engine/Source/Developer/GameplayDebugger/GameplayDebugger.Build.cs
@@ -6,6 +6,8 @@ namespace UnrealBuildTool.Rules
     {
         public GameplayDebugger(TargetInfo Target)
         {
+            bIsAutoStartupModule = true;
+
             PublicIncludePaths.AddRange(
                 new string[] {
 				    "Developer/GameplayDebugger/Public",
diff --git a/Engine/Source/Runtime/AIModule/Private/AIModule.cpp b/Engine/Source/Runtime/AIModule/Private/AIModule.cpp
--- a/Engine/Source/Runtime/AIModule/Private/AIModule.cpp
+++ b/Engine/Source/Runtime/AIModule/Private/AIModule.cpp
@@ -26,11 +26,6 @@ IMPLEMENT_MODULE(FAIModule, AIModule)
 
 void FAIModule::StartupModule()
 { 
-	// This code will execute after your module is loaded into memory and after global variables initialization. We needs some place to load GameplayDebugger module so it's best place for it now.
-#if !(UE_BUILD_SHIPPING || UE_BUILD_TEST)
-	FModuleManager::LoadModulePtr< IModuleInterface >("GameplayDebugger");
-#endif
-
 #if WITH_EDITOR 
 	FModuleManager::LoadModulePtr< IModuleInterface >("AITestSuite");
 #endif // WITH_EDITOR 
