Minor tweakes for EnvQueryDebugHelpers header file.

diff --git a/Engine/Source/Runtime/AIModule/Classes/EnvironmentQuery/EnvQueryDebugHelpers.h b/Engine/Source/Runtime/AIModule/Classes/EnvironmentQuery/EnvQueryDebugHelpers.h
--- a/Engine/Source/Runtime/AIModule/Classes/EnvironmentQuery/EnvQueryDebugHelpers.h
+++ b/Engine/Source/Runtime/AIModule/Classes/EnvironmentQuery/EnvQueryDebugHelpers.h
@@ -156,30 +156,24 @@ class AIMODULE_API UEnvQueryDebugHelpers : public UObject
 #endif
 
 #if ENABLE_VISUAL_LOG && USE_EQS_DEBUGGER
-	static void LogQuery(struct FEnvQueryInstance& Query, const struct FLogCategoryBase& Category, ELogVerbosity::Type);
+	static void LogQuery(struct FEnvQueryInstance& Query, const struct FLogCategoryBase& Category, ELogVerbosity::Type Verbosity);
+
+private:
+	static void LogQueryInternal(struct FEnvQueryInstance& Query, const struct FLogCategoryBase& Category, ELogVerbosity::Type Verbosity, float TimeSeconds, FVisualLogEntry *CurrentEntry);
 #endif
 };
 
 #if ENABLE_VISUAL_LOG && USE_EQS_DEBUGGER
-FORCEINLINE void UEnvQueryDebugHelpers::LogQuery(struct FEnvQueryInstance& Query, const struct FLogCategoryBase& Category, ELogVerbosity::Type Type)
+FORCEINLINE void UEnvQueryDebugHelpers::LogQuery(struct FEnvQueryInstance& Query, const struct FLogCategoryBase& Category, ELogVerbosity::Type Verbosity)
 {
 	UWorld *World = NULL;
 	FVisualLogEntry *CurrentEntry = NULL;
-	if (CheckVisualLogInputInternal(Query.Owner.Get(), Category, Type, &World, &CurrentEntry) == false)
+	if (CheckVisualLogInputInternal(Query.Owner.Get(), Category, Verbosity, &World, &CurrentEntry) == false)
 	{
 		return;
 	}
 
-	const int32 UniqueId = FVisualLogger::Get().GetUniqueId(World->TimeSeconds);
-	TArray<uint8> BlobArray;
-	UEnvQueryDebugHelpers::QueryToBlobArray(Query, BlobArray);
-	FString AdditionalLogInfo = FString::Printf(TEXT("Executed EQS: \n - Name: '%s' (id=%d, option=%d),\n - All Items: %d,\n - ValidItems: %d"), *Query.QueryName, Query.QueryID, Query.OptionIndex, Query.ItemDetails.Num(), Query.NumValidItems);
-	FVisualLogEntry::FLogLine Line(Category.GetCategoryName(), Type, AdditionalLogInfo, Query.QueryID);
-	Line.TagName = *EVisLogTags::TAG_EQS;
-	Line.UniqueId = UniqueId;
-
-	CurrentEntry->LogLines.Add(Line);
-	CurrentEntry->AddDataBlock(EVisLogTags::TAG_EQS, BlobArray, Category.GetCategoryName()).UniqueId = UniqueId;
+	LogQueryInternal(Query, Category, Verbosity, World->TimeSeconds, CurrentEntry);
 }
 #endif
 
diff --git a/Engine/Source/Runtime/AIModule/Private/EnvironmentQuery/EnvQueryDebugHelpers.cpp b/Engine/Source/Runtime/AIModule/Private/EnvironmentQuery/EnvQueryDebugHelpers.cpp
--- a/Engine/Source/Runtime/AIModule/Private/EnvironmentQuery/EnvQueryDebugHelpers.cpp
+++ b/Engine/Source/Runtime/AIModule/Private/EnvironmentQuery/EnvQueryDebugHelpers.cpp
@@ -130,3 +130,19 @@ void  UEnvQueryDebugHelpers::BlobArrayToDebugData(const TArray<uint8>& BlobArray
 }
 
 #endif //ENABLE_VISUAL_LOG
+
+#if ENABLE_VISUAL_LOG && USE_EQS_DEBUGGER
+void UEnvQueryDebugHelpers::LogQueryInternal(struct FEnvQueryInstance& Query, const struct FLogCategoryBase& Category, ELogVerbosity::Type Verbosity, float TimeSeconds, FVisualLogEntry *CurrentEntry)
+{
+	const int32 UniqueId = FVisualLogger::Get().GetUniqueId(TimeSeconds);
+	TArray<uint8> BlobArray;
+	UEnvQueryDebugHelpers::QueryToBlobArray(Query, BlobArray);
+	const FString AdditionalLogInfo = FString::Printf(TEXT("Executed EQS: \n - Name: '%s' (id=%d, option=%d),\n - All Items: %d,\n - ValidItems: %d"), *Query.QueryName, Query.QueryID, Query.OptionIndex, Query.ItemDetails.Num(), Query.NumValidItems);
+	FVisualLogEntry::FLogLine Line(Category.GetCategoryName(), Verbosity, AdditionalLogInfo, Query.QueryID);
+	Line.TagName = *EVisLogTags::TAG_EQS;
+	Line.UniqueId = UniqueId;
+
+	CurrentEntry->LogLines.Add(Line);
+	CurrentEntry->AddDataBlock(EVisLogTags::TAG_EQS, BlobArray, Category.GetCategoryName()).UniqueId = UniqueId;
+}
+#endif //ENABLE_VISUAL_LOG && USE_EQS_DEBUGGER
