more accurate blackboard compatibility check

diff --git a/Engine/Source/Runtime/AIModule/Classes/BehaviorTree/BlackboardData.h b/Engine/Source/Runtime/AIModule/Classes/BehaviorTree/BlackboardData.h
--- a/Engine/Source/Runtime/AIModule/Classes/BehaviorTree/BlackboardData.h
+++ b/Engine/Source/Runtime/AIModule/Classes/BehaviorTree/BlackboardData.h
@@ -23,6 +23,8 @@ struct FBlackboardEntry
 	/** key type and additional properties */
 	UPROPERTY(EditAnywhere, Instanced, Category=Blackboard)
 	UBlackboardKeyType* KeyType;
+
+	bool operator==(const FBlackboardEntry& Other) const;
 };
 
 UCLASS(AutoExpandCategories=(Blackboard))
@@ -104,6 +106,9 @@ class AIMODULE_API UBlackboardData : public UDataAsset
 	 *  meant for adding game specific persistent keys */
 	static FKeyUpdate OnUpdateKeys;
 
+	/** updates parent key cache for editor */
+	void UpdateParentKeys();
+
 protected:
 
 	enum EKeyLookupMode
@@ -118,9 +123,6 @@ protected:
 	/** @return key ID from name */
 	FBlackboard::FKey InternalGetKeyID(const FName& KeyName, EKeyLookupMode LookupMode) const;
 
-	/** updates parent key cache for editor */
-	void UpdateParentKeys();
-
 	/** check if parent chain contains given blackboard data */
 	bool HasParent(const UBlackboardData* TestParent) const;
 };
diff --git a/Engine/Source/Runtime/AIModule/Private/BehaviorTree/BlackboardComponent.cpp b/Engine/Source/Runtime/AIModule/Private/BehaviorTree/BlackboardComponent.cpp
--- a/Engine/Source/Runtime/AIModule/Private/BehaviorTree/BlackboardComponent.cpp
+++ b/Engine/Source/Runtime/AIModule/Private/BehaviorTree/BlackboardComponent.cpp
@@ -202,6 +202,11 @@ bool UBlackboardComponent::IsCompatibleWith(class UBlackboardData* TestAsset) co
 		{
 			return true;
 		}
+
+		if (It->Keys == TestAsset->Keys)
+		{
+			return true;
+		}
 	}
 
 	return false;
diff --git a/Engine/Source/Runtime/AIModule/Private/BehaviorTree/BlackboardData.cpp b/Engine/Source/Runtime/AIModule/Private/BehaviorTree/BlackboardData.cpp
--- a/Engine/Source/Runtime/AIModule/Private/BehaviorTree/BlackboardData.cpp
+++ b/Engine/Source/Runtime/AIModule/Private/BehaviorTree/BlackboardData.cpp
@@ -15,6 +15,12 @@ static void UpdatePersistentKeys(class UBlackboardData* Asset)
 	}
 }
 
+bool FBlackboardEntry::operator==(const FBlackboardEntry& Other) const
+{
+	return (EntryName == Other.EntryName) &&
+		((KeyType && Other.KeyType && KeyType->GetClass() == Other.KeyType->GetClass()) || (KeyType == NULL && Other.KeyType == NULL));
+}
+
 UBlackboardData::UBlackboardData(const FObjectInitializer& ObjectInitializer) : Super(ObjectInitializer)
 {
 	if (HasAnyFlags(RF_ClassDefaultObject))
