Marked AIBlueprintHelperLibrary::SpawnAIFromClass as unsafe for construction script #UE4

- also, removed deprecated AIBlueprintHelperLibrary::SpawnAI

diff --git a/Engine/Config/BaseEditor.ini b/Engine/Config/BaseEditor.ini
--- a/Engine/Config/BaseEditor.ini
+++ b/Engine/Config/BaseEditor.ini
@@ -93,7 +93,6 @@ MenuPriority=11
 MenuPriority=10
 
 [/Script/BlueprintGraph.EdGraphSchema_K2]
-+EditoronlyBPFunctionRedirects=(ClassName="AIBlueprintHelperLibrary", OldFunctionName="SpawnAI", NewFunctionName="SpawnAIFromClass", BlueprintParamName="Pawn", ClassParamName="PawnClass")
 +EditoronlyBPFunctionRedirects=(ClassName="GameplayStatics", OldFunctionName="CreateSaveGameObjectFromBlueprint", NewFunctionName="CreateSaveGameObject", BlueprintParamName="SaveGameBlueprint", ClassParamName="SaveGameClass")
 
 
diff --git a/Engine/Source/Runtime/AIModule/Classes/Blueprint/AIBlueprintHelperLibrary.h b/Engine/Source/Runtime/AIModule/Classes/Blueprint/AIBlueprintHelperLibrary.h
--- a/Engine/Source/Runtime/AIModule/Classes/Blueprint/AIBlueprintHelperLibrary.h
+++ b/Engine/Source/Runtime/AIModule/Classes/Blueprint/AIBlueprintHelperLibrary.h
@@ -19,10 +19,7 @@ class AIMODULE_API UAIBlueprintHelperLibrary : public UBlueprintFunctionLibrary
 	UFUNCTION(BlueprintCallable, Category="AI", meta=(DefaultToSelf="MessageSource"))
 	static void SendAIMessage(APawn* Target, FName Message, UObject* MessageSource, bool bSuccess = true);
 
-	UFUNCTION(BlueprintCallable, Category="AI", meta=(WorldContext="WorldContextObject", DeprecatedFunction, DeprecationMessage="Use SpawnAIFromClass"))
-	static class APawn* SpawnAI(class UObject* WorldContextObject, UBlueprint* Pawn, class UBehaviorTree* BehaviorTree, FVector Location, FRotator Rotation=FRotator::ZeroRotator, bool bNoCollisionFail=false);
-
-	UFUNCTION(BlueprintCallable, Category="AI", meta=(WorldContext="WorldContextObject"))
+	UFUNCTION(BlueprintCallable, Category="AI", meta=(WorldContext="WorldContextObject", UnsafeDuringActorConstruction="true"))
 	static class APawn* SpawnAIFromClass(class UObject* WorldContextObject, TSubclassOf<APawn> PawnClass, class UBehaviorTree* BehaviorTree, FVector Location, FRotator Rotation=FRotator::ZeroRotator, bool bNoCollisionFail=false);
 
 	UFUNCTION(BlueprintPure, Category="AI", meta=(DefaultToSelf="Target"))
diff --git a/Engine/Source/Runtime/AIModule/Private/Blueprint/AIBlueprintHelperLibrary.cpp b/Engine/Source/Runtime/AIModule/Private/Blueprint/AIBlueprintHelperLibrary.cpp
--- a/Engine/Source/Runtime/AIModule/Private/Blueprint/AIBlueprintHelperLibrary.cpp
+++ b/Engine/Source/Runtime/AIModule/Private/Blueprint/AIBlueprintHelperLibrary.cpp
@@ -144,19 +144,6 @@ APawn* UAIBlueprintHelperLibrary::SpawnAIFromClass(class UObject* WorldContextOb
 	return NewPawn;
 }
 
-APawn* UAIBlueprintHelperLibrary::SpawnAI(class UObject* WorldContextObject, UBlueprint* Pawn, class UBehaviorTree* BehaviorTree, FVector Location, FRotator Rotation, bool bNoCollisionFail)
-{
-	APawn* NewPawn = NULL;
-
-	const bool bGoodBPGeneratedClass = Pawn && Pawn->GeneratedClass && Pawn->GeneratedClass->IsChildOf(APawn::StaticClass());
-	if (bGoodBPGeneratedClass)
-	{
-		NewPawn = SpawnAIFromClass(WorldContextObject, *(Pawn->GeneratedClass), BehaviorTree, Location, Rotation, bNoCollisionFail);
-	}
-
-	return NewPawn;
-}
-
 UBlackboardComponent* UAIBlueprintHelperLibrary::GetBlackboard(AActor* Target)
 {
 	UBlackboardComponent* BlackboardComp = nullptr;
