Cleaned out redundant "class" and "struct" keywords use from AIModule cpps #UE4

diff --git a/Engine/Source/Runtime/AIModule/Private/AIController.cpp b/Engine/Source/Runtime/AIModule/Private/AIController.cpp
--- a/Engine/Source/Runtime/AIModule/Private/AIController.cpp
+++ b/Engine/Source/Runtime/AIModule/Private/AIController.cpp
@@ -486,7 +486,7 @@ void AAIController::InitNavigationControl(UNavigationComponent*& PathFindingComp
 	PathFollowingComp = PathFollowingComponent;
 }
 
-EPathFollowingRequestResult::Type AAIController::MoveToActor(class AActor* Goal, float AcceptanceRadius, bool bStopOnOverlap, bool bUsePathfinding, bool bCanStrafe, TSubclassOf<UNavigationQueryFilter> FilterClass)
+EPathFollowingRequestResult::Type AAIController::MoveToActor(AActor* Goal, float AcceptanceRadius, bool bStopOnOverlap, bool bUsePathfinding, bool bCanStrafe, TSubclassOf<UNavigationQueryFilter> FilterClass)
 {
 	SCOPE_CYCLE_COUNTER(STAT_MoveToActor);
 	EPathFollowingRequestResult::Type Result = EPathFollowingRequestResult::Failed;
@@ -602,7 +602,7 @@ EPathFollowingRequestResult::Type AAIController::MoveToLocation(const FVector& D
 	return Result;
 }
 
-FAIRequestID AAIController::RequestMove(FNavPathSharedPtr Path, class AActor* Goal, float AcceptanceRadius, bool bStopOnOverlap, FCustomMoveSharedPtr CustomData)
+FAIRequestID AAIController::RequestMove(FNavPathSharedPtr Path, AActor* Goal, float AcceptanceRadius, bool bStopOnOverlap, FCustomMoveSharedPtr CustomData)
 {
 	uint32 RequestID = FAIRequestID::InvalidRequest;
 	if (PathFollowingComponent)
@@ -639,7 +639,7 @@ void AAIController::StopMovement()
 	PathFollowingComponent->AbortMove(TEXT("StopMovement"));
 }
 
-FNavPathSharedPtr AAIController::FindPath(class AActor* Goal, bool bUsePathfinding, TSharedPtr<const FNavigationQueryFilter> QueryFilter)
+FNavPathSharedPtr AAIController::FindPath(AActor* Goal, bool bUsePathfinding, TSharedPtr<const FNavigationQueryFilter> QueryFilter)
 {
 	if (NavComponent)
 	{
diff --git a/Engine/Source/Runtime/AIModule/Private/AIModule.cpp b/Engine/Source/Runtime/AIModule/Private/AIModule.cpp
--- a/Engine/Source/Runtime/AIModule/Private/AIModule.cpp
+++ b/Engine/Source/Runtime/AIModule/Private/AIModule.cpp
@@ -48,6 +48,6 @@ UAISystemBase* FAIModule::CreateAISystemInstance(UWorld* World)
 {
 	UE_LOG(LogAIModule, Log, TEXT("Creating AISystem for world %s"), *GetNameSafe(World));
 	
-	TSubclassOf<class UAISystemBase> AISystemClass = LoadClass<UAISystemBase>(NULL, *UAISystem::GetAISystemClassName().ToString(), NULL, LOAD_None, NULL);
+	TSubclassOf<UAISystemBase> AISystemClass = LoadClass<UAISystemBase>(NULL, *UAISystem::GetAISystemClassName().ToString(), NULL, LOAD_None, NULL);
 	return ConstructObject<UAISystemBase>(AISystemClass, World);
 }
diff --git a/Engine/Source/Runtime/AIModule/Private/Actions/PawnAction.cpp b/Engine/Source/Runtime/AIModule/Private/Actions/PawnAction.cpp
--- a/Engine/Source/Runtime/AIModule/Private/Actions/PawnAction.cpp
+++ b/Engine/Source/Runtime/AIModule/Private/Actions/PawnAction.cpp
@@ -167,7 +167,7 @@ void UPawnAction::SetFinishResult(EPawnActionResult::Type Result)
 	}
 }
 
-void UPawnAction::SetOwnerComponent(class UPawnActionsComponent* Component)
+void UPawnAction::SetOwnerComponent(UPawnActionsComponent* Component)
 {
 	if (OwnerComponent != NULL && OwnerComponent != Component)
 	{
diff --git a/Engine/Source/Runtime/AIModule/Private/Actions/PawnAction_Move.cpp b/Engine/Source/Runtime/AIModule/Private/Actions/PawnAction_Move.cpp
--- a/Engine/Source/Runtime/AIModule/Private/Actions/PawnAction_Move.cpp
+++ b/Engine/Source/Runtime/AIModule/Private/Actions/PawnAction_Move.cpp
@@ -25,7 +25,7 @@ void UPawnAction_Move::BeginDestroy()
 	Super::BeginDestroy();
 }
 
-UPawnAction_Move* UPawnAction_Move::CreateAction(UWorld& World, class AActor* GoalActor, EPawnActionMoveMode::Type Mode)
+UPawnAction_Move* UPawnAction_Move::CreateAction(UWorld& World, AActor* GoalActor, EPawnActionMoveMode::Type Mode)
 {
 	if (GoalActor == NULL)
 	{
@@ -289,7 +289,7 @@ void UPawnAction_Move::ClearPendingRepath()
 	}
 }
 
-bool UPawnAction_Move::CheckAlreadyAtGoal(class AAIController* Controller, const FVector& TestLocation, float Radius)
+bool UPawnAction_Move::CheckAlreadyAtGoal(AAIController* Controller, const FVector& TestLocation, float Radius)
 {
 	const bool bAlreadyAtGoal = Controller->GetPathFollowingComponent()->HasReached(TestLocation, Radius);
 	if (bAlreadyAtGoal)
@@ -301,7 +301,7 @@ bool UPawnAction_Move::CheckAlreadyAtGoal(class AAIController* Controller, const
 	return bAlreadyAtGoal;
 }
 
-bool UPawnAction_Move::CheckAlreadyAtGoal(class AAIController* Controller, const AActor* TestGoal, float Radius)
+bool UPawnAction_Move::CheckAlreadyAtGoal(AAIController* Controller, const AActor* TestGoal, float Radius)
 {
 	const bool bAlreadyAtGoal = Controller->GetPathFollowingComponent()->HasReached(TestGoal, Radius);
 	if (bAlreadyAtGoal)
diff --git a/Engine/Source/Runtime/AIModule/Private/BehaviorTree/BTAuxiliaryNode.cpp b/Engine/Source/Runtime/AIModule/Private/BehaviorTree/BTAuxiliaryNode.cpp
--- a/Engine/Source/Runtime/AIModule/Private/BehaviorTree/BTAuxiliaryNode.cpp
+++ b/Engine/Source/Runtime/AIModule/Private/BehaviorTree/BTAuxiliaryNode.cpp
@@ -11,7 +11,7 @@ UBTAuxiliaryNode::UBTAuxiliaryNode(const FObjectInitializer& ObjectInitializer)
 	bTickIntervals = false;
 }
 
-void UBTAuxiliaryNode::WrappedOnBecomeRelevant(class UBehaviorTreeComponent* OwnerComp, uint8* NodeMemory) const
+void UBTAuxiliaryNode::WrappedOnBecomeRelevant(UBehaviorTreeComponent* OwnerComp, uint8* NodeMemory) const
 {
 	if (bNotifyBecomeRelevant)
 	{
@@ -23,7 +23,7 @@ void UBTAuxiliaryNode::WrappedOnBecomeRelevant(class UBehaviorTreeComponent* Own
 	}
 }
 
-void UBTAuxiliaryNode::WrappedOnCeaseRelevant(class UBehaviorTreeComponent* OwnerComp, uint8* NodeMemory) const
+void UBTAuxiliaryNode::WrappedOnCeaseRelevant(UBehaviorTreeComponent* OwnerComp, uint8* NodeMemory) const
 {
 	if (bNotifyCeaseRelevant)
 	{
@@ -35,7 +35,7 @@ void UBTAuxiliaryNode::WrappedOnCeaseRelevant(class UBehaviorTreeComponent* Owne
 	}
 }
 
-void UBTAuxiliaryNode::WrappedTickNode(class UBehaviorTreeComponent* OwnerComp, uint8* NodeMemory, float DeltaSeconds) const
+void UBTAuxiliaryNode::WrappedTickNode(UBehaviorTreeComponent* OwnerComp, uint8* NodeMemory, float DeltaSeconds) const
 {
 	if (bNotifyTick)
 	{
@@ -73,7 +73,7 @@ void UBTAuxiliaryNode::SetNextTickTime(uint8* NodeMemory, float RemainingTime) c
 	}
 }
 
-void UBTAuxiliaryNode::DescribeRuntimeValues(const class UBehaviorTreeComponent* OwnerComp, uint8* NodeMemory, EBTDescriptionVerbosity::Type Verbosity, TArray<FString>& Values) const
+void UBTAuxiliaryNode::DescribeRuntimeValues(const UBehaviorTreeComponent* OwnerComp, uint8* NodeMemory, EBTDescriptionVerbosity::Type Verbosity, TArray<FString>& Values) const
 {
 	Super::DescribeRuntimeValues(OwnerComp, NodeMemory, Verbosity, Values);
 
@@ -84,17 +84,17 @@ void UBTAuxiliaryNode::DescribeRuntimeValues(const class UBehaviorTreeComponent*
 	}
 }
 
-void UBTAuxiliaryNode::OnBecomeRelevant(class UBehaviorTreeComponent* OwnerComp, uint8* NodeMemory)
+void UBTAuxiliaryNode::OnBecomeRelevant(UBehaviorTreeComponent* OwnerComp, uint8* NodeMemory)
 {
 	// empty in base class
 }
 
-void UBTAuxiliaryNode::OnCeaseRelevant(class UBehaviorTreeComponent* OwnerComp, uint8* NodeMemory)
+void UBTAuxiliaryNode::OnCeaseRelevant(UBehaviorTreeComponent* OwnerComp, uint8* NodeMemory)
 {
 	// empty in base class
 }
 
-void UBTAuxiliaryNode::TickNode(class UBehaviorTreeComponent* OwnerComp, uint8* NodeMemory, float DeltaSeconds)
+void UBTAuxiliaryNode::TickNode(UBehaviorTreeComponent* OwnerComp, uint8* NodeMemory, float DeltaSeconds)
 {
 	// empty in base class
 }
diff --git a/Engine/Source/Runtime/AIModule/Private/BehaviorTree/BTCompositeNode.cpp b/Engine/Source/Runtime/AIModule/Private/BehaviorTree/BTCompositeNode.cpp
--- a/Engine/Source/Runtime/AIModule/Private/BehaviorTree/BTCompositeNode.cpp
+++ b/Engine/Source/Runtime/AIModule/Private/BehaviorTree/BTCompositeNode.cpp
@@ -16,7 +16,7 @@ void UBTCompositeNode::InitializeComposite(uint16 InLastExecutionIndex)
 	LastExecutionIndex = InLastExecutionIndex;
 }
 
-int32 UBTCompositeNode::FindChildToExecute(struct FBehaviorTreeSearchData& SearchData, EBTNodeResult::Type& LastResult) const
+int32 UBTCompositeNode::FindChildToExecute(FBehaviorTreeSearchData& SearchData, EBTNodeResult::Type& LastResult) const
 {
 	FBTCompositeMemory* NodeMemory = GetNodeMemory<FBTCompositeMemory>(SearchData);
 	int32 RetIdx = BTSpecialChild::ReturnToParent;
@@ -45,7 +45,7 @@ int32 UBTCompositeNode::FindChildToExecute(struct FBehaviorTreeSearchData& Searc
 	return RetIdx;
 }
 
-int32 UBTCompositeNode::GetChildIndex(struct FBehaviorTreeSearchData& SearchData, const class UBTNode* ChildNode) const
+int32 UBTCompositeNode::GetChildIndex(FBehaviorTreeSearchData& SearchData, const UBTNode* ChildNode) const
 {
 	if (ChildNode->GetParentNode() != this)
 	{
@@ -56,7 +56,7 @@ int32 UBTCompositeNode::GetChildIndex(struct FBehaviorTreeSearchData& SearchData
 	return GetChildIndex(ChildNode);
 }
 
-int32 UBTCompositeNode::GetChildIndex(const class UBTNode* ChildNode) const
+int32 UBTCompositeNode::GetChildIndex(const UBTNode* ChildNode) const
 {
 	for (int32 ChildIndex = 0; ChildIndex < Children.Num(); ChildIndex++)
 	{
@@ -70,12 +70,12 @@ int32 UBTCompositeNode::GetChildIndex(const class UBTNode* ChildNode) const
 	return BTSpecialChild::ReturnToParent;
 }
 
-void UBTCompositeNode::OnChildActivation(struct FBehaviorTreeSearchData& SearchData, const class UBTNode* ChildNode) const
+void UBTCompositeNode::OnChildActivation(FBehaviorTreeSearchData& SearchData, const UBTNode* ChildNode) const
 {
 	OnChildActivation(SearchData, GetChildIndex(SearchData, ChildNode));
 }
 
-void UBTCompositeNode::OnChildActivation(struct FBehaviorTreeSearchData& SearchData, int32 ChildIndex) const
+void UBTCompositeNode::OnChildActivation(FBehaviorTreeSearchData& SearchData, int32 ChildIndex) const
 {
 	const FBTCompositeChild& ChildInfo = Children[ChildIndex];
 	FBTCompositeMemory* NodeMemory = GetNodeMemory<FBTCompositeMemory>(SearchData);
@@ -94,12 +94,12 @@ void UBTCompositeNode::OnChildActivation(struct FBehaviorTreeSearchData& SearchD
 	NodeMemory->CurrentChild = ChildIndex;
 }
 
-void UBTCompositeNode::OnChildDeactivation(struct FBehaviorTreeSearchData& SearchData, const class UBTNode* ChildNode, EBTNodeResult::Type& NodeResult) const
+void UBTCompositeNode::OnChildDeactivation(FBehaviorTreeSearchData& SearchData, const UBTNode* ChildNode, EBTNodeResult::Type& NodeResult) const
 {
 	OnChildDeactivation(SearchData, GetChildIndex(SearchData, ChildNode), NodeResult);
 }
 
-void UBTCompositeNode::OnChildDeactivation(struct FBehaviorTreeSearchData& SearchData, int32 ChildIndex, EBTNodeResult::Type& NodeResult) const
+void UBTCompositeNode::OnChildDeactivation(FBehaviorTreeSearchData& SearchData, int32 ChildIndex, EBTNodeResult::Type& NodeResult) const
 {
 	const FBTCompositeChild& ChildInfo = Children[ChildIndex];
 
@@ -117,7 +117,7 @@ void UBTCompositeNode::OnChildDeactivation(struct FBehaviorTreeSearchData& Searc
 	}
 }
 
-void UBTCompositeNode::OnNodeActivation(struct FBehaviorTreeSearchData& SearchData) const
+void UBTCompositeNode::OnNodeActivation(FBehaviorTreeSearchData& SearchData) const
 {
 	OnNodeRestart(SearchData);
 
@@ -136,7 +136,7 @@ void UBTCompositeNode::OnNodeActivation(struct FBehaviorTreeSearchData& SearchDa
 	}
 }
 
-void UBTCompositeNode::OnNodeDeactivation(struct FBehaviorTreeSearchData& SearchData, EBTNodeResult::Type& NodeResult) const
+void UBTCompositeNode::OnNodeDeactivation(FBehaviorTreeSearchData& SearchData, EBTNodeResult::Type& NodeResult) const
 {
 	if (bUseNodeDeactivationNotify)
 	{
@@ -150,14 +150,14 @@ void UBTCompositeNode::OnNodeDeactivation(struct FBehaviorTreeSearchData& Search
 	}
 }
 
-void UBTCompositeNode::OnNodeRestart(struct FBehaviorTreeSearchData& SearchData) const
+void UBTCompositeNode::OnNodeRestart(FBehaviorTreeSearchData& SearchData) const
 {
 	FBTCompositeMemory* NodeMemory = GetNodeMemory<FBTCompositeMemory>(SearchData);
 	NodeMemory->CurrentChild = BTSpecialChild::NotInitialized;
 	NodeMemory->OverrideChild = BTSpecialChild::NotInitialized;
 }
 
-void UBTCompositeNode::NotifyDecoratorsOnActivation(struct FBehaviorTreeSearchData& SearchData, int32 ChildIdx) const
+void UBTCompositeNode::NotifyDecoratorsOnActivation(FBehaviorTreeSearchData& SearchData, int32 ChildIdx) const
 {
 	const FBTCompositeChild& ChildInfo = Children[ChildIdx];
 	for (int32 DecoratorIndex = 0; DecoratorIndex < ChildInfo.Decorators.Num(); DecoratorIndex++)
@@ -182,7 +182,7 @@ void UBTCompositeNode::NotifyDecoratorsOnActivation(struct FBehaviorTreeSearchDa
 	}
 }
 
-void UBTCompositeNode::NotifyDecoratorsOnDeactivation(struct FBehaviorTreeSearchData& SearchData, int32 ChildIdx, EBTNodeResult::Type& NodeResult) const
+void UBTCompositeNode::NotifyDecoratorsOnDeactivation(FBehaviorTreeSearchData& SearchData, int32 ChildIdx, EBTNodeResult::Type& NodeResult) const
 {
 	const FBTCompositeChild& ChildInfo = Children[ChildIdx];
 	for (int32 DecoratorIndex = 0; DecoratorIndex < ChildInfo.Decorators.Num(); DecoratorIndex++)
@@ -205,7 +205,7 @@ void UBTCompositeNode::NotifyDecoratorsOnDeactivation(struct FBehaviorTreeSearch
 	}
 }
 
-void UBTCompositeNode::NotifyDecoratorsOnFailedActivation(struct FBehaviorTreeSearchData& SearchData, int32 ChildIdx, EBTNodeResult::Type& NodeResult) const
+void UBTCompositeNode::NotifyDecoratorsOnFailedActivation(FBehaviorTreeSearchData& SearchData, int32 ChildIdx, EBTNodeResult::Type& NodeResult) const
 {
 	const FBTCompositeChild& ChildInfo = Children[ChildIdx];
 	const uint16 ActiveInstanceIdx = SearchData.OwnerComp->GetActiveInstanceIdx();
@@ -225,19 +225,19 @@ void UBTCompositeNode::NotifyDecoratorsOnFailedActivation(struct FBehaviorTreeSe
 	SearchData.OwnerComp->StoreDebuggerSearchStep(GetChildNode(ChildIdx), ActiveInstanceIdx, NodeResult);
 }
 
-void UBTCompositeNode::NotifyChildExecution(class UBehaviorTreeComponent* OwnerComp, uint8* NodeMemory, int32 ChildIdx, EBTNodeResult::Type& NodeResult) const
+void UBTCompositeNode::NotifyChildExecution(UBehaviorTreeComponent* OwnerComp, uint8* NodeMemory, int32 ChildIdx, EBTNodeResult::Type& NodeResult) const
 {
 }
 
-void UBTCompositeNode::NotifyNodeActivation(struct FBehaviorTreeSearchData& SearchData) const
+void UBTCompositeNode::NotifyNodeActivation(FBehaviorTreeSearchData& SearchData) const
 {
 }
 
-void UBTCompositeNode::NotifyNodeDeactivation(struct FBehaviorTreeSearchData& SearchData, EBTNodeResult::Type& NodeResult) const
+void UBTCompositeNode::NotifyNodeDeactivation(FBehaviorTreeSearchData& SearchData, EBTNodeResult::Type& NodeResult) const
 {
 }
 
-void UBTCompositeNode::ConditionalNotifyChildExecution(class UBehaviorTreeComponent* OwnerComp, uint8* NodeMemory, const class UBTNode* ChildNode, EBTNodeResult::Type& NodeResult) const
+void UBTCompositeNode::ConditionalNotifyChildExecution(UBehaviorTreeComponent* OwnerComp, uint8* NodeMemory, const UBTNode* ChildNode, EBTNodeResult::Type& NodeResult) const
 {
 	if (bUseChildExecutionNotify)
 	{
@@ -275,7 +275,7 @@ struct FOperationStackInfo
 		NumLeft(DecoratorOp.Number), Op(DecoratorOp.Operation), bHasForcedResult(0) {};
 };
 
-static bool UpdateOperationStack(const class UBehaviorTreeComponent* OwnerComp, FString& Indent,
+static bool UpdateOperationStack(const UBehaviorTreeComponent* OwnerComp, FString& Indent,
 								 TArray<FOperationStackInfo>& Stack, bool bTestResult,
 								 int32& FailedDecoratorIdx, int32& NodeDecoratorIdx, bool& bShouldStoreNodeIndex)
 {
@@ -340,7 +340,7 @@ static bool UpdateOperationStack(const class UBehaviorTreeComponent* OwnerComp,
 	return bTestResult;
 }
 
-bool UBTCompositeNode::DoDecoratorsAllowExecution(class UBehaviorTreeComponent* OwnerComp, int32 InstanceIdx, int32 ChildIdx) const
+bool UBTCompositeNode::DoDecoratorsAllowExecution(UBehaviorTreeComponent* OwnerComp, int32 InstanceIdx, int32 ChildIdx) const
 {
 	const FBTCompositeChild& ChildInfo = Children[ChildIdx];
 	bool bResult = true;
@@ -438,7 +438,7 @@ bool UBTCompositeNode::DoDecoratorsAllowExecution(class UBehaviorTreeComponent*
 	return bResult;
 }
 
-int32 UBTCompositeNode::GetMatchingChildIndex(int32 ActiveInstanceIdx, struct FBTNodeIndex& NodeIdx) const
+int32 UBTCompositeNode::GetMatchingChildIndex(int32 ActiveInstanceIdx, FBTNodeIndex& NodeIdx) const
 {
 	const int32 OutsideRange = BTSpecialChild::ReturnToParent;
 	const int32 UnlimitedRange = Children.Num() - 1;
@@ -470,7 +470,7 @@ int32 UBTCompositeNode::GetMatchingChildIndex(int32 ActiveInstanceIdx, struct FB
 	return (ActiveInstanceIdx > NodeIdx.InstanceIndex) ? UnlimitedRange : OutsideRange;
 }
 
-int32 UBTCompositeNode::GetNextChild(struct FBehaviorTreeSearchData& SearchData, int32 LastChildIdx, EBTNodeResult::Type LastResult) const
+int32 UBTCompositeNode::GetNextChild(FBehaviorTreeSearchData& SearchData, int32 LastChildIdx, EBTNodeResult::Type LastResult) const
 {
 	FBTCompositeMemory* NodeMemory = GetNodeMemory<FBTCompositeMemory>(SearchData);
 	int32 NextChildIndex = BTSpecialChild::ReturnToParent;
@@ -505,7 +505,7 @@ void UBTCompositeNode::SetChildOverride(FBehaviorTreeSearchData& SearchData, int
 	}
 }
 
-void UBTCompositeNode::RequestDelayedExecution(class UBehaviorTreeComponent* OwnerComp, EBTNodeResult::Type LastResult) const
+void UBTCompositeNode::RequestDelayedExecution(UBehaviorTreeComponent* OwnerComp, EBTNodeResult::Type LastResult) const
 {
 	OwnerComp->RequestExecution(LastResult);
 }
@@ -522,12 +522,12 @@ uint16 UBTCompositeNode::GetChildExecutionIndex(int32 Index, EBTChildIndex Child
 	return (LastExecutionIndex + 1);
 }
 
-bool UBTCompositeNode::CanPushSubtree(class UBehaviorTreeComponent* OwnerComp, uint8* NodeMemory, int32 ChildIdx) const
+bool UBTCompositeNode::CanPushSubtree(UBehaviorTreeComponent* OwnerComp, uint8* NodeMemory, int32 ChildIdx) const
 {
 	return true;
 }
 
-void UBTCompositeNode::DescribeRuntimeValues(const class UBehaviorTreeComponent* OwnerComp, uint8* NodeMemory, EBTDescriptionVerbosity::Type Verbosity, TArray<FString>& Values) const
+void UBTCompositeNode::DescribeRuntimeValues(const UBehaviorTreeComponent* OwnerComp, uint8* NodeMemory, EBTDescriptionVerbosity::Type Verbosity, TArray<FString>& Values) const
 {
 	Super::DescribeRuntimeValues(OwnerComp, NodeMemory, Verbosity, Values);
 
diff --git a/Engine/Source/Runtime/AIModule/Private/BehaviorTree/BTDecorator.cpp b/Engine/Source/Runtime/AIModule/Private/BehaviorTree/BTDecorator.cpp
--- a/Engine/Source/Runtime/AIModule/Private/BehaviorTree/BTDecorator.cpp
+++ b/Engine/Source/Runtime/AIModule/Private/BehaviorTree/BTDecorator.cpp
@@ -25,7 +25,7 @@ void UBTDecorator::InitializeDecorator(uint8 MyChildIndex)
 	ChildIndex = MyChildIndex;
 }
 
-bool UBTDecorator::CalculateRawConditionValue(class UBehaviorTreeComponent* OwnerComp, uint8* NodeMemory) const
+bool UBTDecorator::CalculateRawConditionValue(UBehaviorTreeComponent* OwnerComp, uint8* NodeMemory) const
 {
 	return true;
 }
@@ -35,25 +35,25 @@ void UBTDecorator::SetIsInversed(bool bShouldBeInversed)
 	bInverseCondition = bShouldBeInversed;
 }
 
-void UBTDecorator::OnNodeActivation(struct FBehaviorTreeSearchData& SearchData)
+void UBTDecorator::OnNodeActivation(FBehaviorTreeSearchData& SearchData)
 {
 }
 
-void UBTDecorator::OnNodeDeactivation(struct FBehaviorTreeSearchData& SearchData, EBTNodeResult::Type NodeResult)
+void UBTDecorator::OnNodeDeactivation(FBehaviorTreeSearchData& SearchData, EBTNodeResult::Type NodeResult)
 {
 }
 
-void UBTDecorator::OnNodeProcessed(struct FBehaviorTreeSearchData& SearchData, EBTNodeResult::Type& NodeResult)
+void UBTDecorator::OnNodeProcessed(FBehaviorTreeSearchData& SearchData, EBTNodeResult::Type& NodeResult)
 {
 }
 
-bool UBTDecorator::WrappedCanExecute(class UBehaviorTreeComponent* OwnerComp, uint8* NodeMemory) const
+bool UBTDecorator::WrappedCanExecute(UBehaviorTreeComponent* OwnerComp, uint8* NodeMemory) const
 {
 	const UBTDecorator* NodeOb = bCreateNodeInstance ? (const UBTDecorator*)GetNodeInstance(OwnerComp, NodeMemory) : this;
 	return NodeOb ? (IsInversed() != NodeOb->CalculateRawConditionValue(OwnerComp, NodeMemory)) : false;
 }
 
-void UBTDecorator::WrappedOnNodeActivation(struct FBehaviorTreeSearchData& SearchData) const
+void UBTDecorator::WrappedOnNodeActivation(FBehaviorTreeSearchData& SearchData) const
 {
 	if (bNotifyActivation)
 	{
@@ -65,7 +65,7 @@ void UBTDecorator::WrappedOnNodeActivation(struct FBehaviorTreeSearchData& Searc
 	}
 };
 
-void UBTDecorator::WrappedOnNodeDeactivation(struct FBehaviorTreeSearchData& SearchData, EBTNodeResult::Type NodeResult) const
+void UBTDecorator::WrappedOnNodeDeactivation(FBehaviorTreeSearchData& SearchData, EBTNodeResult::Type NodeResult) const
 {
 	if (bNotifyDeactivation)
 	{
@@ -77,7 +77,7 @@ void UBTDecorator::WrappedOnNodeDeactivation(struct FBehaviorTreeSearchData& Sea
 	}
 }
 
-void UBTDecorator::WrappedOnNodeProcessed(struct FBehaviorTreeSearchData& SearchData, EBTNodeResult::Type& NodeResult) const
+void UBTDecorator::WrappedOnNodeProcessed(FBehaviorTreeSearchData& SearchData, EBTNodeResult::Type& NodeResult) const
 {
 	if (bNotifyProcessed)
 	{
diff --git a/Engine/Source/Runtime/AIModule/Private/BehaviorTree/BTNode.cpp b/Engine/Source/Runtime/AIModule/Private/BehaviorTree/BTNode.cpp
--- a/Engine/Source/Runtime/AIModule/Private/BehaviorTree/BTNode.cpp
+++ b/Engine/Source/Runtime/AIModule/Private/BehaviorTree/BTNode.cpp
@@ -36,7 +36,7 @@ UWorld* UBTNode::GetWorld() const
 		Cast<UWorld>(GetOuter()->GetOuter());
 }
 
-void UBTNode::InitializeNode(class UBTCompositeNode* InParentNode, uint16 InExecutionIndex, uint16 InMemoryOffset, uint8 InTreeDepth)
+void UBTNode::InitializeNode(UBTCompositeNode* InParentNode, uint16 InExecutionIndex, uint16 InMemoryOffset, uint8 InTreeDepth)
 {
 	ParentNode = InParentNode;
 	ExecutionIndex = InExecutionIndex;
@@ -44,27 +44,27 @@ void UBTNode::InitializeNode(class UBTCompositeNode* InParentNode, uint16 InExec
 	TreeDepth = InTreeDepth;
 }
 
-void UBTNode::InitializeMemory(class UBehaviorTreeComponent* OwnerComp, uint8* NodeMemory, EBTMemoryInit::Type InitType) const
+void UBTNode::InitializeMemory(UBehaviorTreeComponent* OwnerComp, uint8* NodeMemory, EBTMemoryInit::Type InitType) const
 {
-	// empty in base class 
+	// empty in base 
 }
 
-void UBTNode::CleanupMemory(class UBehaviorTreeComponent* OwnerComp, uint8* NodeMemory, EBTMemoryClear::Type CleanupType) const
+void UBTNode::CleanupMemory(UBehaviorTreeComponent* OwnerComp, uint8* NodeMemory, EBTMemoryClear::Type CleanupType) const
 {
-	// empty in base class 
+	// empty in base 
 }
 
-void UBTNode::OnInstanceCreated(class UBehaviorTreeComponent* OwnerComp)
+void UBTNode::OnInstanceCreated(UBehaviorTreeComponent* OwnerComp)
 {
 	// empty in base class
 }
 
-void UBTNode::OnInstanceDestroyed(class UBehaviorTreeComponent* OwnerComp)
+void UBTNode::OnInstanceDestroyed(UBehaviorTreeComponent* OwnerComp)
 {
 	// empty in base class
 }
 
-void UBTNode::InitializeInSubtree(class UBehaviorTreeComponent* OwnerComp, uint8* NodeMemory, int32& NextInstancedIndex, EBTMemoryInit::Type InitType) const
+void UBTNode::InitializeInSubtree(UBehaviorTreeComponent* OwnerComp, uint8* NodeMemory, int32& NextInstancedIndex, EBTMemoryInit::Type InitType) const
 {
 	if (bCreateNodeInstance)
 	{
@@ -93,7 +93,7 @@ void UBTNode::InitializeInSubtree(class UBehaviorTreeComponent* OwnerComp, uint8
 	}
 }
 
-void UBTNode::CleanupInSubtree(class UBehaviorTreeComponent* OwnerComp, uint8* NodeMemory, EBTMemoryClear::Type CleanupType) const
+void UBTNode::CleanupInSubtree(UBehaviorTreeComponent* OwnerComp, uint8* NodeMemory, EBTMemoryClear::Type CleanupType) const
 {
 	if (!bCreateNodeInstance && !bIsInjected)
 	{
@@ -102,7 +102,7 @@ void UBTNode::CleanupInSubtree(class UBehaviorTreeComponent* OwnerComp, uint8* N
 }
 
 #if USE_BEHAVIORTREE_DEBUGGER
-void UBTNode::InitializeExecutionOrder(class UBTNode* NextNode)
+void UBTNode::InitializeExecutionOrder(UBTNode* NextNode)
 {
 	NextExecutionNode = NextNode;
 }
@@ -128,19 +128,19 @@ uint16 UBTNode::GetSpecialMemorySize() const
 	return bCreateNodeInstance ? sizeof(FBTInstancedNodeMemory) : 0;
 }
 
-class UBTNode* UBTNode::GetNodeInstance(const class UBehaviorTreeComponent* OwnerComp, uint8* NodeMemory) const
+UBTNode* UBTNode::GetNodeInstance(const UBehaviorTreeComponent* OwnerComp, uint8* NodeMemory) const
 {
 	FBTInstancedNodeMemory* MyMemory = GetSpecialNodeMemory<FBTInstancedNodeMemory>(NodeMemory);
 	return OwnerComp && MyMemory && OwnerComp->NodeInstances.IsValidIndex(MyMemory->NodeIdx) ?
 		OwnerComp->NodeInstances[MyMemory->NodeIdx] : NULL;
 }
 
-class UBTNode* UBTNode::GetNodeInstance(struct FBehaviorTreeSearchData& SearchData) const
+UBTNode* UBTNode::GetNodeInstance(FBehaviorTreeSearchData& SearchData) const
 {
 	return GetNodeInstance(SearchData.OwnerComp, GetNodeMemory<uint8>(SearchData));
 }
 
-FString UBTNode::GetRuntimeDescription(const class UBehaviorTreeComponent* OwnerComp, uint8* NodeMemory, EBTDescriptionVerbosity::Type Verbosity) const
+FString UBTNode::GetRuntimeDescription(const UBehaviorTreeComponent* OwnerComp, uint8* NodeMemory, EBTDescriptionVerbosity::Type Verbosity) const
 {
 	FString Description = NodeName.Len() ? FString::Printf(TEXT("%s [%s]"), *NodeName, *GetStaticDescription()) : GetStaticDescription();
 	TArray<FString> RuntimeValues;
@@ -166,7 +166,7 @@ FString UBTNode::GetStaticDescription() const
 	return UBehaviorTreeTypes::GetShortTypeName(this);
 }
 
-void UBTNode::DescribeRuntimeValues(const class UBehaviorTreeComponent* OwnerComp, uint8* NodeMemory, EBTDescriptionVerbosity::Type Verbosity, TArray<FString>& Values) const
+void UBTNode::DescribeRuntimeValues(const UBehaviorTreeComponent* OwnerComp, uint8* NodeMemory, EBTDescriptionVerbosity::Type Verbosity, TArray<FString>& Values) const
 {
 	// nothing stored in memory for base class
 }
diff --git a/Engine/Source/Runtime/AIModule/Private/BehaviorTree/BTService.cpp b/Engine/Source/Runtime/AIModule/Private/BehaviorTree/BTService.cpp
--- a/Engine/Source/Runtime/AIModule/Private/BehaviorTree/BTService.cpp
+++ b/Engine/Source/Runtime/AIModule/Private/BehaviorTree/BTService.cpp
@@ -19,13 +19,13 @@ void UBTService::TickNode(UBehaviorTreeComponent* OwnerComp, uint8* NodeMemory,
 	SetNextTickTime(NodeMemory, NextTickTime);
 }
 
-void UBTService::OnSearchStart(struct FBehaviorTreeSearchData& SearchData)
+void UBTService::OnSearchStart(FBehaviorTreeSearchData& SearchData)
 {
 	uint8* NodeMemory = GetNodeMemory<uint8>(SearchData);
 	TickNode(SearchData.OwnerComp, NodeMemory, 0.0f);
 }
 
-void UBTService::NotifyParentActivation(struct FBehaviorTreeSearchData& SearchData)
+void UBTService::NotifyParentActivation(FBehaviorTreeSearchData& SearchData)
 {
 	if (bNotifyOnSearch)
 	{
diff --git a/Engine/Source/Runtime/AIModule/Private/BehaviorTree/BTTaskNode.cpp b/Engine/Source/Runtime/AIModule/Private/BehaviorTree/BTTaskNode.cpp
--- a/Engine/Source/Runtime/AIModule/Private/BehaviorTree/BTTaskNode.cpp
+++ b/Engine/Source/Runtime/AIModule/Private/BehaviorTree/BTTaskNode.cpp
@@ -9,29 +9,29 @@ UBTTaskNode::UBTTaskNode(const FObjectInitializer& ObjectInitializer) : Super(Ob
 	bNotifyTick = false;
 }
 
-EBTNodeResult::Type UBTTaskNode::ExecuteTask(class UBehaviorTreeComponent* OwnerComp, uint8* NodeMemory)
+EBTNodeResult::Type UBTTaskNode::ExecuteTask(UBehaviorTreeComponent* OwnerComp, uint8* NodeMemory)
 {
 	return EBTNodeResult::Succeeded;
 }
 
-EBTNodeResult::Type UBTTaskNode::AbortTask(class UBehaviorTreeComponent* OwnerComp, uint8* NodeMemory)
+EBTNodeResult::Type UBTTaskNode::AbortTask(UBehaviorTreeComponent* OwnerComp, uint8* NodeMemory)
 {
 	return EBTNodeResult::Aborted;
 }
 
-EBTNodeResult::Type UBTTaskNode::WrappedExecuteTask(class UBehaviorTreeComponent* OwnerComp, uint8* NodeMemory) const
+EBTNodeResult::Type UBTTaskNode::WrappedExecuteTask(UBehaviorTreeComponent* OwnerComp, uint8* NodeMemory) const
 {
 	const UBTNode* NodeOb = bCreateNodeInstance ? GetNodeInstance(OwnerComp, NodeMemory) : this;
 	return NodeOb ? ((UBTTaskNode*)NodeOb)->ExecuteTask(OwnerComp, NodeMemory) : EBTNodeResult::Failed;
 }
 
-EBTNodeResult::Type UBTTaskNode::WrappedAbortTask(class UBehaviorTreeComponent* OwnerComp, uint8* NodeMemory) const
+EBTNodeResult::Type UBTTaskNode::WrappedAbortTask(UBehaviorTreeComponent* OwnerComp, uint8* NodeMemory) const
 {
 	const UBTNode* NodeOb = bCreateNodeInstance ? GetNodeInstance(OwnerComp, NodeMemory) : this;
 	return NodeOb ? ((UBTTaskNode*)NodeOb)->AbortTask(OwnerComp, NodeMemory) : EBTNodeResult::Aborted;
 }
 
-void UBTTaskNode::WrappedTickTask(class UBehaviorTreeComponent* OwnerComp, uint8* NodeMemory, float DeltaSeconds) const
+void UBTTaskNode::WrappedTickTask(UBehaviorTreeComponent* OwnerComp, uint8* NodeMemory, float DeltaSeconds) const
 {
 	if (bNotifyTick)
 	{
@@ -43,7 +43,7 @@ void UBTTaskNode::WrappedTickTask(class UBehaviorTreeComponent* OwnerComp, uint8
 	}
 }
 
-void UBTTaskNode::ReceivedMessage(UBrainComponent* BrainComp, const struct FAIMessage& Message)
+void UBTTaskNode::ReceivedMessage(UBrainComponent* BrainComp, const FAIMessage& Message)
 {
 	UBehaviorTreeComponent* OwnerComp = (UBehaviorTreeComponent*)BrainComp;
 	uint16 InstanceIdx = OwnerComp->FindInstanceContainingNode(this);
@@ -52,12 +52,12 @@ void UBTTaskNode::ReceivedMessage(UBrainComponent* BrainComp, const struct FAIMe
 	OnMessage(OwnerComp, NodeMemory, Message.MessageName, Message.RequestID, Message.Status == FAIMessage::Success);
 }
 
-void UBTTaskNode::TickTask(class UBehaviorTreeComponent* OwnerComp, uint8* NodeMemory, float DeltaSeconds)
+void UBTTaskNode::TickTask(UBehaviorTreeComponent* OwnerComp, uint8* NodeMemory, float DeltaSeconds)
 {
 	// empty in base class
 }
 
-void UBTTaskNode::OnMessage(class UBehaviorTreeComponent* OwnerComp, uint8* NodeMemory, FName Message, int32 RequestID, bool bSuccess)
+void UBTTaskNode::OnMessage(UBehaviorTreeComponent* OwnerComp, uint8* NodeMemory, FName Message, int32 RequestID, bool bSuccess)
 {
 	const EBTTaskStatus::Type Status = OwnerComp->GetTaskStatus(this);
 	if (Status == EBTTaskStatus::Active)
@@ -70,33 +70,33 @@ void UBTTaskNode::OnMessage(class UBehaviorTreeComponent* OwnerComp, uint8* Node
 	}
 }
 
-void UBTTaskNode::FinishLatentTask(class UBehaviorTreeComponent* OwnerComp, EBTNodeResult::Type TaskResult) const
+void UBTTaskNode::FinishLatentTask(UBehaviorTreeComponent* OwnerComp, EBTNodeResult::Type TaskResult) const
 {
 	// OnTaskFinished must receive valid template node
 	UBTTaskNode* TemplateNode = (UBTTaskNode*)OwnerComp->FindTemplateNode(this);
 	OwnerComp->OnTaskFinished(TemplateNode, TaskResult);
 }
 
-void UBTTaskNode::FinishLatentAbort(class UBehaviorTreeComponent* OwnerComp) const
+void UBTTaskNode::FinishLatentAbort(UBehaviorTreeComponent* OwnerComp) const
 {
 	// OnTaskFinished must receive valid template node
 	UBTTaskNode* TemplateNode = (UBTTaskNode*)OwnerComp->FindTemplateNode(this);
 	OwnerComp->OnTaskFinished(TemplateNode, EBTNodeResult::Aborted);
 }
 
-void UBTTaskNode::WaitForMessage(class UBehaviorTreeComponent* OwnerComp, FName MessageType) const
+void UBTTaskNode::WaitForMessage(UBehaviorTreeComponent* OwnerComp, FName MessageType) const
 {
 	// messages delegates should be called on node instances (if they exists)
 	OwnerComp->RegisterMessageObserver(this, MessageType);
 }
 
-void UBTTaskNode::WaitForMessage(class UBehaviorTreeComponent* OwnerComp, FName MessageType, int32 RequestID) const
+void UBTTaskNode::WaitForMessage(UBehaviorTreeComponent* OwnerComp, FName MessageType, int32 RequestID) const
 {
 	// messages delegates should be called on node instances (if they exists)
 	OwnerComp->RegisterMessageObserver(this, MessageType, RequestID);
 }
 	
-void UBTTaskNode::StopWaitingForMessages(class UBehaviorTreeComponent* OwnerComp) const
+void UBTTaskNode::StopWaitingForMessages(UBehaviorTreeComponent* OwnerComp) const
 {
 	// messages delegates should be called on node instances (if they exists)
 	OwnerComp->UnregisterMessageObserversFrom(this);
diff --git a/Engine/Source/Runtime/AIModule/Private/BehaviorTree/BehaviorTreeComponent.cpp b/Engine/Source/Runtime/AIModule/Private/BehaviorTree/BehaviorTreeComponent.cpp
--- a/Engine/Source/Runtime/AIModule/Private/BehaviorTree/BehaviorTreeComponent.cpp
+++ b/Engine/Source/Runtime/AIModule/Private/BehaviorTree/BehaviorTreeComponent.cpp
@@ -650,7 +650,7 @@ void UBehaviorTreeComponent::RequestExecution(UBTCompositeNode* RequestedOn, int
 	ScheduleExecutionUpdate();
 }
 
-void UBehaviorTreeComponent::ApplySearchUpdates(const TArray<struct FBehaviorTreeSearchUpdate>& UpdateList, int32 NewNodeExecutionIndex, bool bPostUpdate)
+void UBehaviorTreeComponent::ApplySearchUpdates(const TArray<FBehaviorTreeSearchUpdate>& UpdateList, int32 NewNodeExecutionIndex, bool bPostUpdate)
 {
 	for (int32 Index = 0; Index < UpdateList.Num(); Index++)
 	{
@@ -756,7 +756,7 @@ void UBehaviorTreeComponent::ApplySearchData(UBTNode* NewActiveNode)
 
 void UBehaviorTreeComponent::ApplyDiscardedSearch()
 {
-	TArray<struct FBehaviorTreeSearchUpdate> UpdateList;
+	TArray<FBehaviorTreeSearchUpdate> UpdateList;
 	for (int32 Idx = 0; Idx < SearchData.PendingUpdates.Num(); Idx++)
 	{
 		const FBehaviorTreeSearchUpdate& UpdateInfo = SearchData.PendingUpdates[Idx];
@@ -1608,7 +1608,7 @@ FString UBehaviorTreeComponent::DescribeActiveTrees() const
 }
 
 #if ENABLE_VISUAL_LOG
-void UBehaviorTreeComponent::DescribeSelfToVisLog(struct FVisualLogEntry* Snapshot) const
+void UBehaviorTreeComponent::DescribeSelfToVisLog(FVisualLogEntry* Snapshot) const
 {
 	if (IsPendingKill())
 	{
diff --git a/Engine/Source/Runtime/AIModule/Private/BehaviorTree/BehaviorTreeManager.cpp b/Engine/Source/Runtime/AIModule/Private/BehaviorTree/BehaviorTreeManager.cpp
--- a/Engine/Source/Runtime/AIModule/Private/BehaviorTree/BehaviorTreeManager.cpp
+++ b/Engine/Source/Runtime/AIModule/Private/BehaviorTree/BehaviorTreeManager.cpp
@@ -162,7 +162,7 @@ static void InitializeNodeHelper(UBTCompositeNode* ParentNode, UBTNode* NodeOb,
 	}
 }
 
-bool UBehaviorTreeManager::LoadTree(UBehaviorTree& Asset, class UBTCompositeNode*& Root, uint16& InstanceMemorySize)
+bool UBehaviorTreeManager::LoadTree(UBehaviorTree& Asset, UBTCompositeNode*& Root, uint16& InstanceMemorySize)
 {
 	SCOPE_CYCLE_COUNTER(STAT_AI_BehaviorTree_LoadTime);
 
@@ -359,7 +359,7 @@ void UBehaviorTreeManager::DumpUsageStats() const
 		AllNodesCounter.Append(TreeNodeCounter);
 	}
 	
-	UE_LOG(LogBehaviorTree, Display, TEXT("--- Total Nodes class usage:"));
+	UE_LOG(LogBehaviorTree, Display, TEXT("--- Total Nodes usage:"));
 	AllNodesCounter.Print(TEXT(","));
 }
 
diff --git a/Engine/Source/Runtime/AIModule/Private/BehaviorTree/BehaviorTreeTypes.cpp b/Engine/Source/Runtime/AIModule/Private/BehaviorTree/BehaviorTreeTypes.cpp
--- a/Engine/Source/Runtime/AIModule/Private/BehaviorTree/BehaviorTreeTypes.cpp
+++ b/Engine/Source/Runtime/AIModule/Private/BehaviorTree/BehaviorTreeTypes.cpp
@@ -9,7 +9,7 @@
 //----------------------------------------------------------------------//
 // FBehaviorTreeInstance
 //----------------------------------------------------------------------//
-void FBehaviorTreeInstance::Initialize(class UBehaviorTreeComponent* OwnerComp, UBTCompositeNode* Node, int32& InstancedIndex, EBTMemoryInit::Type InitType)
+void FBehaviorTreeInstance::Initialize(UBehaviorTreeComponent* OwnerComp, UBTCompositeNode* Node, int32& InstancedIndex, EBTMemoryInit::Type InitType)
 {
 	if (Node == NULL)
 	{
@@ -43,7 +43,7 @@ void FBehaviorTreeInstance::Initialize(class UBehaviorTreeComponent* OwnerComp,
 	}
 }
 
-void FBehaviorTreeInstance::InjectNodes(class UBehaviorTreeComponent* OwnerComp, UBTCompositeNode* Node, int32& InstancedIndex)
+void FBehaviorTreeInstance::InjectNodes(UBehaviorTreeComponent* OwnerComp, UBTCompositeNode* Node, int32& InstancedIndex)
 {
 	if (Node == NULL)
 	{
@@ -69,7 +69,7 @@ void FBehaviorTreeInstance::InjectNodes(class UBehaviorTreeComponent* OwnerComp,
 	}
 }
 
-void FBehaviorTreeInstance::Cleanup(class UBehaviorTreeComponent* OwnerComp, EBTMemoryClear::Type CleanupType)
+void FBehaviorTreeInstance::Cleanup(UBehaviorTreeComponent* OwnerComp, EBTMemoryClear::Type CleanupType)
 {
 	FBehaviorTreeInstanceId& Info = OwnerComp->KnownInstances[InstanceIdIndex];
 	if (Info.FirstNodeInstance >= 0)
@@ -89,7 +89,7 @@ void FBehaviorTreeInstance::Cleanup(class UBehaviorTreeComponent* OwnerComp, EBT
 	Info.InstanceMemory = InstanceMemory;
 }
 
-void FBehaviorTreeInstance::CleanupNodes(class UBehaviorTreeComponent* OwnerComp, UBTCompositeNode* Node, EBTMemoryClear::Type CleanupType)
+void FBehaviorTreeInstance::CleanupNodes(UBehaviorTreeComponent* OwnerComp, UBTCompositeNode* Node, EBTMemoryClear::Type CleanupType)
 {
 	if (Node == NULL)
 	{
@@ -190,7 +190,7 @@ void FBehaviorTreeSearchData::AssignSearchId()
 //----------------------------------------------------------------------//
 // FBlackboardKeySelector
 //----------------------------------------------------------------------//
-void FBlackboardKeySelector::CacheSelectedKey(class UBlackboardData* BlackboardAsset)
+void FBlackboardKeySelector::CacheSelectedKey(UBlackboardData* BlackboardAsset)
 {
 	if (BlackboardAsset && !(bNoneIsAllowedValue && SelectedKeyName == NAME_None))
 	{
@@ -204,7 +204,7 @@ void FBlackboardKeySelector::CacheSelectedKey(class UBlackboardData* BlackboardA
 	}
 }
 
-void FBlackboardKeySelector::InitSelectedKey(class UBlackboardData* BlackboardAsset)
+void FBlackboardKeySelector::InitSelectedKey(UBlackboardData* BlackboardAsset)
 {
 	check(BlackboardAsset);
 	for (UBlackboardData* It = BlackboardAsset; It; It = It->Parent)
diff --git a/Engine/Source/Runtime/AIModule/Private/BehaviorTree/Blackboard/BlackboardKeyType_NativeEnum.cpp b/Engine/Source/Runtime/AIModule/Private/BehaviorTree/Blackboard/BlackboardKeyType_NativeEnum.cpp
--- a/Engine/Source/Runtime/AIModule/Private/BehaviorTree/Blackboard/BlackboardKeyType_NativeEnum.cpp
+++ b/Engine/Source/Runtime/AIModule/Private/BehaviorTree/Blackboard/BlackboardKeyType_NativeEnum.cpp
@@ -36,7 +36,7 @@ bool UBlackboardKeyType_NativeEnum::IsAllowedByFilter(UBlackboardKeyType* Filter
 }
 
 #if WITH_EDITOR
-void UBlackboardKeyType_NativeEnum::PostEditChangeProperty(struct FPropertyChangedEvent& PropertyChangedEvent)
+void UBlackboardKeyType_NativeEnum::PostEditChangeProperty(FPropertyChangedEvent& PropertyChangedEvent)
 {
 	Super::PostEditChangeProperty(PropertyChangedEvent);
 
diff --git a/Engine/Source/Runtime/AIModule/Private/BehaviorTree/BlackboardData.cpp b/Engine/Source/Runtime/AIModule/Private/BehaviorTree/BlackboardData.cpp
--- a/Engine/Source/Runtime/AIModule/Private/BehaviorTree/BlackboardData.cpp
+++ b/Engine/Source/Runtime/AIModule/Private/BehaviorTree/BlackboardData.cpp
@@ -6,7 +6,7 @@
 
 UBlackboardData::FKeyUpdate UBlackboardData::OnUpdateKeys;
 
-static void UpdatePersistentKeys(class UBlackboardData* Asset)
+static void UpdatePersistentKeys(UBlackboardData* Asset)
 {
 	UBlackboardKeyType_Object* SelfKeyType = Asset->UpdatePersistentKey<UBlackboardKeyType_Object>(FBlackboard::KeySelf);
 	if (SelfKeyType)
@@ -110,7 +110,7 @@ void UBlackboardData::PostLoad()
 
 #if WITH_EDITOR
 
-void UBlackboardData::PostEditChangeProperty(struct FPropertyChangedEvent& PropertyChangedEvent)
+void UBlackboardData::PostEditChangeProperty(FPropertyChangedEvent& PropertyChangedEvent)
 {
 	Super::PostEditChangeProperty(PropertyChangedEvent);
 
diff --git a/Engine/Source/Runtime/AIModule/Private/BehaviorTree/BlueprintNodeHelpers.cpp b/Engine/Source/Runtime/AIModule/Private/BehaviorTree/BlueprintNodeHelpers.cpp
--- a/Engine/Source/Runtime/AIModule/Private/BehaviorTree/BlueprintNodeHelpers.cpp
+++ b/Engine/Source/Runtime/AIModule/Private/BehaviorTree/BlueprintNodeHelpers.cpp
@@ -178,7 +178,7 @@ namespace BlueprintNodeHelpers
 		}
 	}
 
-	bool FindNodeOwner(AActor* OwningActor, class UBTNode* Node, UBehaviorTreeComponent*& OwningComp, int32& OwningInstanceIdx)
+	bool FindNodeOwner(AActor* OwningActor, UBTNode* Node, UBehaviorTreeComponent*& OwningComp, int32& OwningInstanceIdx)
 	{
 		bool bFound = false;
 
diff --git a/Engine/Source/Runtime/AIModule/Private/BehaviorTree/Composites/BTComposite_Selector.cpp b/Engine/Source/Runtime/AIModule/Private/BehaviorTree/Composites/BTComposite_Selector.cpp
--- a/Engine/Source/Runtime/AIModule/Private/BehaviorTree/Composites/BTComposite_Selector.cpp
+++ b/Engine/Source/Runtime/AIModule/Private/BehaviorTree/Composites/BTComposite_Selector.cpp
@@ -10,7 +10,7 @@ UBTComposite_Selector::UBTComposite_Selector(const FObjectInitializer& ObjectIni
 	OnNextChild.BindUObject(this, &UBTComposite_Selector::GetNextChildHandler);
 }
 
-int32 UBTComposite_Selector::GetNextChildHandler(struct FBehaviorTreeSearchData& SearchData, int32 PrevChild, EBTNodeResult::Type LastResult) const
+int32 UBTComposite_Selector::GetNextChildHandler(FBehaviorTreeSearchData& SearchData, int32 PrevChild, EBTNodeResult::Type LastResult) const
 {
 	// success = quit
 	int32 NextChildIdx = BTSpecialChild::ReturnToParent;
diff --git a/Engine/Source/Runtime/AIModule/Private/BehaviorTree/Composites/BTComposite_Sequence.cpp b/Engine/Source/Runtime/AIModule/Private/BehaviorTree/Composites/BTComposite_Sequence.cpp
--- a/Engine/Source/Runtime/AIModule/Private/BehaviorTree/Composites/BTComposite_Sequence.cpp
+++ b/Engine/Source/Runtime/AIModule/Private/BehaviorTree/Composites/BTComposite_Sequence.cpp
@@ -10,7 +10,7 @@ UBTComposite_Sequence::UBTComposite_Sequence(const FObjectInitializer& ObjectIni
 	OnNextChild.BindUObject(this, &UBTComposite_Sequence::GetNextChildHandler);
 }
 
-int32 UBTComposite_Sequence::GetNextChildHandler(struct FBehaviorTreeSearchData& SearchData, int32 PrevChild, EBTNodeResult::Type LastResult) const
+int32 UBTComposite_Sequence::GetNextChildHandler(FBehaviorTreeSearchData& SearchData, int32 PrevChild, EBTNodeResult::Type LastResult) const
 {
 	// failure = quit
 	int32 NextChildIdx = BTSpecialChild::ReturnToParent;
diff --git a/Engine/Source/Runtime/AIModule/Private/BehaviorTree/Composites/BTComposite_SimpleParallel.cpp b/Engine/Source/Runtime/AIModule/Private/BehaviorTree/Composites/BTComposite_SimpleParallel.cpp
--- a/Engine/Source/Runtime/AIModule/Private/BehaviorTree/Composites/BTComposite_SimpleParallel.cpp
+++ b/Engine/Source/Runtime/AIModule/Private/BehaviorTree/Composites/BTComposite_SimpleParallel.cpp
@@ -11,7 +11,7 @@ UBTComposite_SimpleParallel::UBTComposite_SimpleParallel(const FObjectInitialize
 	OnNextChild.BindUObject(this, &UBTComposite_SimpleParallel::GetNextChildHandler);
 }
 
-int32 UBTComposite_SimpleParallel::GetNextChildHandler(struct FBehaviorTreeSearchData& SearchData, int32 PrevChild, EBTNodeResult::Type LastResult) const
+int32 UBTComposite_SimpleParallel::GetNextChildHandler(FBehaviorTreeSearchData& SearchData, int32 PrevChild, EBTNodeResult::Type LastResult) const
 {
 	FBTParallelMemory* MyMemory = GetNodeMemory<FBTParallelMemory>(SearchData);
 	int32 NextChildIdx = BTSpecialChild::ReturnToParent;
@@ -31,7 +31,7 @@ int32 UBTComposite_SimpleParallel::GetNextChildHandler(struct FBehaviorTreeSearc
 	return NextChildIdx;
 }
 
-void UBTComposite_SimpleParallel::NotifyChildExecution(class UBehaviorTreeComponent* OwnerComp, uint8* NodeMemory, int32 ChildIdx, EBTNodeResult::Type& NodeResult) const
+void UBTComposite_SimpleParallel::NotifyChildExecution(UBehaviorTreeComponent* OwnerComp, uint8* NodeMemory, int32 ChildIdx, EBTNodeResult::Type& NodeResult) const
 {
 	FBTParallelMemory* MyMemory = (FBTParallelMemory*)NodeMemory;
 	if (ChildIdx == EBTParallelChild::MainTask)
@@ -81,7 +81,7 @@ void UBTComposite_SimpleParallel::NotifyChildExecution(class UBehaviorTreeCompon
 	}
 }
 
-void UBTComposite_SimpleParallel::NotifyNodeDeactivation(struct FBehaviorTreeSearchData& SearchData, EBTNodeResult::Type& NodeResult) const
+void UBTComposite_SimpleParallel::NotifyNodeDeactivation(FBehaviorTreeSearchData& SearchData, EBTNodeResult::Type& NodeResult) const
 {
 	FBTParallelMemory* MyMemory = GetNodeMemory<FBTParallelMemory>(SearchData);
 	const uint16 ActiveInstanceIdx = SearchData.OwnerComp->GetActiveInstanceIdx();
@@ -128,7 +128,7 @@ uint16 UBTComposite_SimpleParallel::GetInstanceMemorySize() const
 	return sizeof(FBTParallelMemory);
 }
 
-bool UBTComposite_SimpleParallel::CanPushSubtree(class UBehaviorTreeComponent* OwnerComp, uint8* NodeMemory, int32 ChildIdx) const
+bool UBTComposite_SimpleParallel::CanPushSubtree(UBehaviorTreeComponent* OwnerComp, uint8* NodeMemory, int32 ChildIdx) const
 {
 	return (ChildIdx != EBTParallelChild::MainTask);
 }
@@ -147,7 +147,7 @@ FString UBTComposite_SimpleParallel::GetStaticDescription() const
 		FinishMode < 2 ? *FinishDesc[FinishMode] : TEXT(""));
 }
 
-void UBTComposite_SimpleParallel::DescribeRuntimeValues(const class UBehaviorTreeComponent* OwnerComp, uint8* NodeMemory, EBTDescriptionVerbosity::Type Verbosity, TArray<FString>& Values) const
+void UBTComposite_SimpleParallel::DescribeRuntimeValues(const UBehaviorTreeComponent* OwnerComp, uint8* NodeMemory, EBTDescriptionVerbosity::Type Verbosity, TArray<FString>& Values) const
 {
 	Super::DescribeRuntimeValues(OwnerComp, NodeMemory, Verbosity, Values);
 
diff --git a/Engine/Source/Runtime/AIModule/Private/BehaviorTree/Decorators/BTDecorator_Blackboard.cpp b/Engine/Source/Runtime/AIModule/Private/BehaviorTree/Decorators/BTDecorator_Blackboard.cpp
--- a/Engine/Source/Runtime/AIModule/Private/BehaviorTree/Decorators/BTDecorator_Blackboard.cpp
+++ b/Engine/Source/Runtime/AIModule/Private/BehaviorTree/Decorators/BTDecorator_Blackboard.cpp
@@ -11,7 +11,7 @@ UBTDecorator_Blackboard::UBTDecorator_Blackboard(const FObjectInitializer& Objec
 	NotifyObserver = EBTBlackboardRestart::ResultChange;
 }
 
-bool UBTDecorator_Blackboard::CalculateRawConditionValue(class UBehaviorTreeComponent* OwnerComp, uint8* NodeMemory) const
+bool UBTDecorator_Blackboard::CalculateRawConditionValue(UBehaviorTreeComponent* OwnerComp, uint8* NodeMemory) const
 {
 	const UBlackboardComponent* BlackboardComp = OwnerComp->GetBlackboardComponent();
 	return EvaluateOnBlackboard(BlackboardComp);
@@ -48,7 +48,7 @@ bool UBTDecorator_Blackboard::EvaluateOnBlackboard(const UBlackboardComponent* B
 	return bResult;
 }
 
-void UBTDecorator_Blackboard::OnBlackboardChange(const class UBlackboardComponent* Blackboard, FBlackboard::FKey ChangedKeyID)
+void UBTDecorator_Blackboard::OnBlackboardChange(const UBlackboardComponent* Blackboard, FBlackboard::FKey ChangedKeyID)
 {
 	UBehaviorTreeComponent* BehaviorComp = Blackboard ? (UBehaviorTreeComponent*)Blackboard->GetBrainComponent() : NULL;
 	if (BlackboardKey.GetSelectedKeyID() == ChangedKeyID &&
@@ -89,7 +89,7 @@ void UBTDecorator_Blackboard::OnBlackboardChange(const class UBlackboardComponen
 	}
 }
 
-void UBTDecorator_Blackboard::DescribeRuntimeValues(const class UBehaviorTreeComponent* OwnerComp, uint8* NodeMemory, EBTDescriptionVerbosity::Type Verbosity, TArray<FString>& Values) const
+void UBTDecorator_Blackboard::DescribeRuntimeValues(const UBehaviorTreeComponent* OwnerComp, uint8* NodeMemory, EBTDescriptionVerbosity::Type Verbosity, TArray<FString>& Values) const
 {
 	Super::DescribeRuntimeValues(OwnerComp, NodeMemory, Verbosity, Values);
 
diff --git a/Engine/Source/Runtime/AIModule/Private/BehaviorTree/Decorators/BTDecorator_BlueprintBase.cpp b/Engine/Source/Runtime/AIModule/Private/BehaviorTree/Decorators/BTDecorator_BlueprintBase.cpp
--- a/Engine/Source/Runtime/AIModule/Private/BehaviorTree/Decorators/BTDecorator_BlueprintBase.cpp
+++ b/Engine/Source/Runtime/AIModule/Private/BehaviorTree/Decorators/BTDecorator_BlueprintBase.cpp
@@ -104,14 +104,14 @@ void UBTDecorator_BlueprintBase::OnCeaseRelevant(UBehaviorTreeComponent* OwnerCo
 	}
 }
 
-void UBTDecorator_BlueprintBase::OnNodeActivation(struct FBehaviorTreeSearchData& SearchData)
+void UBTDecorator_BlueprintBase::OnNodeActivation(FBehaviorTreeSearchData& SearchData)
 {
 	// skip flag, will be handled by bNotifyActivation
 
 	ReceiveExecutionStart(SearchData.OwnerComp->GetOwner());
 }
 
-void UBTDecorator_BlueprintBase::OnNodeDeactivation(struct FBehaviorTreeSearchData& SearchData, EBTNodeResult::Type NodeResult)
+void UBTDecorator_BlueprintBase::OnNodeDeactivation(FBehaviorTreeSearchData& SearchData, EBTNodeResult::Type NodeResult)
 {
 	// skip flag, will be handled by bNotifyDeactivation
 
@@ -125,7 +125,7 @@ void UBTDecorator_BlueprintBase::TickNode(UBehaviorTreeComponent* OwnerComp, uin
 	ReceiveTick(OwnerComp->GetOwner(), DeltaSeconds);
 }
 
-bool UBTDecorator_BlueprintBase::CalculateRawConditionValue(class UBehaviorTreeComponent* OwnerComp, uint8* NodeMemory) const
+bool UBTDecorator_BlueprintBase::CalculateRawConditionValue(UBehaviorTreeComponent* OwnerComp, uint8* NodeMemory) const
 {
 	CurrentCallResult = false;
 	if (bImplementsReceiveConditionCheck)
@@ -176,7 +176,7 @@ FString UBTDecorator_BlueprintBase::GetStaticDescription() const
 	return ReturnDesc;
 }
 
-void UBTDecorator_BlueprintBase::DescribeRuntimeValues(const class UBehaviorTreeComponent* OwnerComp, uint8* NodeMemory, EBTDescriptionVerbosity::Type Verbosity, TArray<FString>& Values) const
+void UBTDecorator_BlueprintBase::DescribeRuntimeValues(const UBehaviorTreeComponent* OwnerComp, uint8* NodeMemory, EBTDescriptionVerbosity::Type Verbosity, TArray<FString>& Values) const
 {
 	UBTDecorator_BlueprintBase* CDO = (UBTDecorator_BlueprintBase*)(GetClass()->GetDefaultObject());
 	if (CDO && CDO->PropertyData.Num())
@@ -194,7 +194,7 @@ void UBTDecorator_BlueprintBase::OnBlackboardChange(const UBlackboardComponent*
 	}
 }
 
-void UBTDecorator_BlueprintBase::OnInstanceDestroyed(class UBehaviorTreeComponent* OwnerComp)
+void UBTDecorator_BlueprintBase::OnInstanceDestroyed(UBehaviorTreeComponent* OwnerComp)
 {
 	// force dropping all pending latent actions associated with this blueprint
 	BlueprintNodeHelpers::AbortLatentActions(OwnerComp, this);
diff --git a/Engine/Source/Runtime/AIModule/Private/BehaviorTree/Decorators/BTDecorator_CompareBBEntries.cpp b/Engine/Source/Runtime/AIModule/Private/BehaviorTree/Decorators/BTDecorator_CompareBBEntries.cpp
--- a/Engine/Source/Runtime/AIModule/Private/BehaviorTree/Decorators/BTDecorator_CompareBBEntries.cpp
+++ b/Engine/Source/Runtime/AIModule/Private/BehaviorTree/Decorators/BTDecorator_CompareBBEntries.cpp
@@ -30,7 +30,7 @@ void UBTDecorator_CompareBBEntries::InitializeFromAsset(UBehaviorTree& Asset)
 
 // @note I know it's ugly to have "return" statements in many places inside a function, but the way 
 // around was very awkward here 
-bool UBTDecorator_CompareBBEntries::CalculateRawConditionValue(class UBehaviorTreeComponent* OwnerComp, uint8* NodeMemory) const 
+bool UBTDecorator_CompareBBEntries::CalculateRawConditionValue(UBehaviorTreeComponent* OwnerComp, uint8* NodeMemory) const 
 {
 	// first of all require same type
 	// @todo this could be checked statically (i.e. in editor, asset creation time)!
diff --git a/Engine/Source/Runtime/AIModule/Private/BehaviorTree/Decorators/BTDecorator_ConditionalLoop.cpp b/Engine/Source/Runtime/AIModule/Private/BehaviorTree/Decorators/BTDecorator_ConditionalLoop.cpp
--- a/Engine/Source/Runtime/AIModule/Private/BehaviorTree/Decorators/BTDecorator_ConditionalLoop.cpp
+++ b/Engine/Source/Runtime/AIModule/Private/BehaviorTree/Decorators/BTDecorator_ConditionalLoop.cpp
@@ -14,18 +14,18 @@ UBTDecorator_ConditionalLoop::UBTDecorator_ConditionalLoop(const FObjectInitiali
 	bAllowAbortChildNodes = false;
 }
 
-bool UBTDecorator_ConditionalLoop::CalculateRawConditionValue(class UBehaviorTreeComponent* OwnerComp, uint8* NodeMemory) const
+bool UBTDecorator_ConditionalLoop::CalculateRawConditionValue(UBehaviorTreeComponent* OwnerComp, uint8* NodeMemory) const
 {
 	// always allows execution
 	return true;
 }
 
-void UBTDecorator_ConditionalLoop::OnBlackboardChange(const class UBlackboardComponent* Blackboard, FBlackboard::FKey ChangedKeyID)
+void UBTDecorator_ConditionalLoop::OnBlackboardChange(const UBlackboardComponent* Blackboard, FBlackboard::FKey ChangedKeyID)
 {
 	// empty, don't react to blackboard value changes
 }
 
-void UBTDecorator_ConditionalLoop::OnNodeDeactivation(struct FBehaviorTreeSearchData& SearchData, EBTNodeResult::Type NodeResult)
+void UBTDecorator_ConditionalLoop::OnNodeDeactivation(FBehaviorTreeSearchData& SearchData, EBTNodeResult::Type NodeResult)
 {
 	if (NodeResult != EBTNodeResult::Aborted && SearchData.OwnerComp)
 	{
diff --git a/Engine/Source/Runtime/AIModule/Private/BehaviorTree/Decorators/BTDecorator_ConeCheck.cpp b/Engine/Source/Runtime/AIModule/Private/BehaviorTree/Decorators/BTDecorator_ConeCheck.cpp
--- a/Engine/Source/Runtime/AIModule/Private/BehaviorTree/Decorators/BTDecorator_ConeCheck.cpp
+++ b/Engine/Source/Runtime/AIModule/Private/BehaviorTree/Decorators/BTDecorator_ConeCheck.cpp
@@ -76,7 +76,7 @@ FORCEINLINE bool UBTDecorator_ConeCheck::CalcConditionImpl(UBehaviorTreeComponen
 		&& ConeDir.CosineAngle2D(DirectionToObserve) > ConeHalfAngleDot;
 }
 
-bool UBTDecorator_ConeCheck::CalculateRawConditionValue(class UBehaviorTreeComponent* OwnerComp, uint8* NodeMemory) const
+bool UBTDecorator_ConeCheck::CalculateRawConditionValue(UBehaviorTreeComponent* OwnerComp, uint8* NodeMemory) const
 {
 	return CalcConditionImpl(OwnerComp, NodeMemory);
 }
@@ -87,7 +87,7 @@ void UBTDecorator_ConeCheck::OnBecomeRelevant(UBehaviorTreeComponent* OwnerComp,
 	DecoratorMemory->bLastRawResult = CalcConditionImpl(OwnerComp, NodeMemory);
 }
 
-void UBTDecorator_ConeCheck::OnBlackboardChange(const class UBlackboardComponent* Blackboard, FBlackboard::FKey ChangedKeyID)
+void UBTDecorator_ConeCheck::OnBlackboardChange(const UBlackboardComponent* Blackboard, FBlackboard::FKey ChangedKeyID)
 {
 	check(false);
 }
@@ -109,7 +109,7 @@ FString UBTDecorator_ConeCheck::GetStaticDescription() const
 		, ConeHalfAngle * 2, *ConeOrigin.SelectedKeyName.ToString(), *ConeDirection.SelectedKeyName.ToString());
 }
 
-void UBTDecorator_ConeCheck::DescribeRuntimeValues(const class UBehaviorTreeComponent* OwnerComp, uint8* NodeMemory, EBTDescriptionVerbosity::Type Verbosity, TArray<FString>& Values) const
+void UBTDecorator_ConeCheck::DescribeRuntimeValues(const UBehaviorTreeComponent* OwnerComp, uint8* NodeMemory, EBTDescriptionVerbosity::Type Verbosity, TArray<FString>& Values) const
 {
 	const UBlackboardComponent* BBComponent = OwnerComp->GetBlackboardComponent();
 			
diff --git a/Engine/Source/Runtime/AIModule/Private/BehaviorTree/Decorators/BTDecorator_Cooldown.cpp b/Engine/Source/Runtime/AIModule/Private/BehaviorTree/Decorators/BTDecorator_Cooldown.cpp
--- a/Engine/Source/Runtime/AIModule/Private/BehaviorTree/Decorators/BTDecorator_Cooldown.cpp
+++ b/Engine/Source/Runtime/AIModule/Private/BehaviorTree/Decorators/BTDecorator_Cooldown.cpp
@@ -21,14 +21,14 @@ void UBTDecorator_Cooldown::PostLoad()
 	bNotifyTick = (FlowAbortMode != EBTFlowAbortMode::None);
 }
 
-bool UBTDecorator_Cooldown::CalculateRawConditionValue(class UBehaviorTreeComponent* OwnerComp, uint8* NodeMemory) const 
+bool UBTDecorator_Cooldown::CalculateRawConditionValue(UBehaviorTreeComponent* OwnerComp, uint8* NodeMemory) const 
 {
 	FBTCooldownDecoratorMemory* DecoratorMemory = (FBTCooldownDecoratorMemory*)NodeMemory;
 	const float TimePassed = (OwnerComp->GetWorld()->GetTimeSeconds() - DecoratorMemory->LastUseTimestamp);
 	return TimePassed >= CoolDownTime;
 }
 
-void UBTDecorator_Cooldown::InitializeMemory(class UBehaviorTreeComponent* OwnerComp, uint8* NodeMemory, EBTMemoryInit::Type InitType) const
+void UBTDecorator_Cooldown::InitializeMemory(UBehaviorTreeComponent* OwnerComp, uint8* NodeMemory, EBTMemoryInit::Type InitType) const
 {
 	FBTCooldownDecoratorMemory* DecoratorMemory = (FBTCooldownDecoratorMemory*)NodeMemory;
 	if (InitType == EBTMemoryInit::Initialize)
@@ -39,14 +39,14 @@ void UBTDecorator_Cooldown::InitializeMemory(class UBehaviorTreeComponent* Owner
 	DecoratorMemory->bRequestedRestart = false;
 }
 
-void UBTDecorator_Cooldown::OnNodeDeactivation(struct FBehaviorTreeSearchData& SearchData, EBTNodeResult::Type NodeResult)
+void UBTDecorator_Cooldown::OnNodeDeactivation(FBehaviorTreeSearchData& SearchData, EBTNodeResult::Type NodeResult)
 {
 	FBTCooldownDecoratorMemory* DecoratorMemory = GetNodeMemory<FBTCooldownDecoratorMemory>(SearchData);
 	DecoratorMemory->LastUseTimestamp = SearchData.OwnerComp->GetWorld()->GetTimeSeconds();
 	DecoratorMemory->bRequestedRestart = false;
 }
 
-void UBTDecorator_Cooldown::TickNode(class UBehaviorTreeComponent* OwnerComp, uint8* NodeMemory, float DeltaSeconds)
+void UBTDecorator_Cooldown::TickNode(UBehaviorTreeComponent* OwnerComp, uint8* NodeMemory, float DeltaSeconds)
 {
 	FBTCooldownDecoratorMemory* DecoratorMemory = (FBTCooldownDecoratorMemory*)NodeMemory;
 	if (!DecoratorMemory->bRequestedRestart)
@@ -67,7 +67,7 @@ FString UBTDecorator_Cooldown::GetStaticDescription() const
 		CoolDownTime, *UBehaviorTreeTypes::DescribeNodeResult(EBTNodeResult::Failed));
 }
 
-void UBTDecorator_Cooldown::DescribeRuntimeValues(const class UBehaviorTreeComponent* OwnerComp, uint8* NodeMemory, EBTDescriptionVerbosity::Type Verbosity, TArray<FString>& Values) const
+void UBTDecorator_Cooldown::DescribeRuntimeValues(const UBehaviorTreeComponent* OwnerComp, uint8* NodeMemory, EBTDescriptionVerbosity::Type Verbosity, TArray<FString>& Values) const
 {
 	Super::DescribeRuntimeValues(OwnerComp, NodeMemory, Verbosity, Values);
 
diff --git a/Engine/Source/Runtime/AIModule/Private/BehaviorTree/Decorators/BTDecorator_ForceSuccess.cpp b/Engine/Source/Runtime/AIModule/Private/BehaviorTree/Decorators/BTDecorator_ForceSuccess.cpp
--- a/Engine/Source/Runtime/AIModule/Private/BehaviorTree/Decorators/BTDecorator_ForceSuccess.cpp
+++ b/Engine/Source/Runtime/AIModule/Private/BehaviorTree/Decorators/BTDecorator_ForceSuccess.cpp
@@ -13,7 +13,7 @@ UBTDecorator_ForceSuccess::UBTDecorator_ForceSuccess(const FObjectInitializer& O
 	bAllowAbortChildNodes = false;
 }
 
-void UBTDecorator_ForceSuccess::OnNodeProcessed(struct FBehaviorTreeSearchData& SearchData, EBTNodeResult::Type& NodeResult)
+void UBTDecorator_ForceSuccess::OnNodeProcessed(FBehaviorTreeSearchData& SearchData, EBTNodeResult::Type& NodeResult)
 {
 	NodeResult = EBTNodeResult::Succeeded;
 }
diff --git a/Engine/Source/Runtime/AIModule/Private/BehaviorTree/Decorators/BTDecorator_KeepInCone.cpp b/Engine/Source/Runtime/AIModule/Private/BehaviorTree/Decorators/BTDecorator_KeepInCone.cpp
--- a/Engine/Source/Runtime/AIModule/Private/BehaviorTree/Decorators/BTDecorator_KeepInCone.cpp
+++ b/Engine/Source/Runtime/AIModule/Private/BehaviorTree/Decorators/BTDecorator_KeepInCone.cpp
@@ -105,7 +105,7 @@ FString UBTDecorator_KeepInCone::GetStaticDescription() const
 		*Observed.SelectedKeyName.ToString());
 }
 
-void UBTDecorator_KeepInCone::DescribeRuntimeValues(const class UBehaviorTreeComponent* OwnerComp, uint8* NodeMemory, EBTDescriptionVerbosity::Type Verbosity, TArray<FString>& Values) const
+void UBTDecorator_KeepInCone::DescribeRuntimeValues(const UBehaviorTreeComponent* OwnerComp, uint8* NodeMemory, EBTDescriptionVerbosity::Type Verbosity, TArray<FString>& Values) const
 {
 	TNodeInstanceMemory* DecoratorMemory = (TNodeInstanceMemory*)NodeMemory;
 	FVector CurrentDir(1.0f, 0, 0);
diff --git a/Engine/Source/Runtime/AIModule/Private/BehaviorTree/Decorators/BTDecorator_Loop.cpp b/Engine/Source/Runtime/AIModule/Private/BehaviorTree/Decorators/BTDecorator_Loop.cpp
--- a/Engine/Source/Runtime/AIModule/Private/BehaviorTree/Decorators/BTDecorator_Loop.cpp
+++ b/Engine/Source/Runtime/AIModule/Private/BehaviorTree/Decorators/BTDecorator_Loop.cpp
@@ -54,7 +54,7 @@ FString UBTDecorator_Loop::GetStaticDescription() const
 		FString::Printf(TEXT("%s: %d loops"), *Super::GetStaticDescription(), NumLoops);
 }
 
-void UBTDecorator_Loop::DescribeRuntimeValues(const class UBehaviorTreeComponent* OwnerComp, uint8* NodeMemory, EBTDescriptionVerbosity::Type Verbosity, TArray<FString>& Values) const
+void UBTDecorator_Loop::DescribeRuntimeValues(const UBehaviorTreeComponent* OwnerComp, uint8* NodeMemory, EBTDescriptionVerbosity::Type Verbosity, TArray<FString>& Values) const
 {
 	Super::DescribeRuntimeValues(OwnerComp, NodeMemory, Verbosity, Values);
 
diff --git a/Engine/Source/Runtime/AIModule/Private/BehaviorTree/Decorators/BTDecorator_TimeLimit.cpp b/Engine/Source/Runtime/AIModule/Private/BehaviorTree/Decorators/BTDecorator_TimeLimit.cpp
--- a/Engine/Source/Runtime/AIModule/Private/BehaviorTree/Decorators/BTDecorator_TimeLimit.cpp
+++ b/Engine/Source/Runtime/AIModule/Private/BehaviorTree/Decorators/BTDecorator_TimeLimit.cpp
@@ -39,7 +39,7 @@ FString UBTDecorator_TimeLimit::GetStaticDescription() const
 		*UBehaviorTreeTypes::DescribeNodeResult(EBTNodeResult::Failed), TimeLimit);
 }
 
-void UBTDecorator_TimeLimit::DescribeRuntimeValues(const class UBehaviorTreeComponent* OwnerComp, uint8* NodeMemory, EBTDescriptionVerbosity::Type Verbosity, TArray<FString>& Values) const
+void UBTDecorator_TimeLimit::DescribeRuntimeValues(const UBehaviorTreeComponent* OwnerComp, uint8* NodeMemory, EBTDescriptionVerbosity::Type Verbosity, TArray<FString>& Values) const
 {
 	Super::DescribeRuntimeValues(OwnerComp, NodeMemory, Verbosity, Values);
 
diff --git a/Engine/Source/Runtime/AIModule/Private/BehaviorTree/Services/BTService_BlueprintBase.cpp b/Engine/Source/Runtime/AIModule/Private/BehaviorTree/Services/BTService_BlueprintBase.cpp
--- a/Engine/Source/Runtime/AIModule/Private/BehaviorTree/Services/BTService_BlueprintBase.cpp
+++ b/Engine/Source/Runtime/AIModule/Private/BehaviorTree/Services/BTService_BlueprintBase.cpp
@@ -71,7 +71,7 @@ void UBTService_BlueprintBase::OnCeaseRelevant(UBehaviorTreeComponent* OwnerComp
 	}
 }
 
-void UBTService_BlueprintBase::OnSearchStart(struct FBehaviorTreeSearchData& SearchData)
+void UBTService_BlueprintBase::OnSearchStart(FBehaviorTreeSearchData& SearchData)
 {
 	// skip flag, will be handled by bNotifyOnSearch
 
diff --git a/Engine/Source/Runtime/AIModule/Private/BehaviorTree/Tasks/BTTask_MoveTo.cpp b/Engine/Source/Runtime/AIModule/Private/BehaviorTree/Tasks/BTTask_MoveTo.cpp
--- a/Engine/Source/Runtime/AIModule/Private/BehaviorTree/Tasks/BTTask_MoveTo.cpp
+++ b/Engine/Source/Runtime/AIModule/Private/BehaviorTree/Tasks/BTTask_MoveTo.cpp
@@ -20,7 +20,7 @@ UBTTask_MoveTo::UBTTask_MoveTo(const FObjectInitializer& ObjectInitializer)
 	BlackboardKey.AddVectorFilter(this);
 }
 
-EBTNodeResult::Type UBTTask_MoveTo::ExecuteTask(class UBehaviorTreeComponent* OwnerComp, uint8* NodeMemory)
+EBTNodeResult::Type UBTTask_MoveTo::ExecuteTask(UBehaviorTreeComponent* OwnerComp, uint8* NodeMemory)
 {
 	EBTNodeResult::Type NodeResult = EBTNodeResult::InProgress;
 
@@ -40,7 +40,7 @@ EBTNodeResult::Type UBTTask_MoveTo::ExecuteTask(class UBehaviorTreeComponent* Ow
 	return NodeResult;
 }
 
-EBTNodeResult::Type UBTTask_MoveTo::PerformMoveTask(class UBehaviorTreeComponent* OwnerComp, uint8* NodeMemory)
+EBTNodeResult::Type UBTTask_MoveTo::PerformMoveTask(UBehaviorTreeComponent* OwnerComp, uint8* NodeMemory)
 {
 	const UBlackboardComponent* MyBlackboard = OwnerComp->GetBlackboardComponent();
 	FBTMoveToTaskMemory* MyMemory = (FBTMoveToTaskMemory*)NodeMemory;
@@ -90,7 +90,7 @@ EBTNodeResult::Type UBTTask_MoveTo::PerformMoveTask(class UBehaviorTreeComponent
 	return NodeResult;
 }
 
-EBTNodeResult::Type UBTTask_MoveTo::AbortTask(class UBehaviorTreeComponent* OwnerComp, uint8* NodeMemory)
+EBTNodeResult::Type UBTTask_MoveTo::AbortTask(UBehaviorTreeComponent* OwnerComp, uint8* NodeMemory)
 {
 	FBTMoveToTaskMemory* MyMemory = (FBTMoveToTaskMemory*)NodeMemory;
 	if (!MyMemory->bWaitingForPath)
@@ -106,7 +106,7 @@ EBTNodeResult::Type UBTTask_MoveTo::AbortTask(class UBehaviorTreeComponent* Owne
 	return Super::AbortTask(OwnerComp, NodeMemory);
 }
 
-void UBTTask_MoveTo::TickTask(class UBehaviorTreeComponent* OwnerComp, uint8* NodeMemory, float DeltaSeconds)
+void UBTTask_MoveTo::TickTask(UBehaviorTreeComponent* OwnerComp, uint8* NodeMemory, float DeltaSeconds)
 {
 	FBTMoveToTaskMemory* MyMemory = (FBTMoveToTaskMemory*)NodeMemory;
 	if (MyMemory->bWaitingForPath && !OwnerComp->IsPaused())
@@ -126,7 +126,7 @@ void UBTTask_MoveTo::TickTask(class UBehaviorTreeComponent* OwnerComp, uint8* No
 	}
 }
 
-void UBTTask_MoveTo::OnMessage(class UBehaviorTreeComponent* OwnerComp, uint8* NodeMemory, FName Message, int32 SenderID, bool bSuccess)
+void UBTTask_MoveTo::OnMessage(UBehaviorTreeComponent* OwnerComp, uint8* NodeMemory, FName Message, int32 SenderID, bool bSuccess)
 {
 	// AIMessage_RepathFailed means task has failed
 	bSuccess &= (Message != UBrainComponent::AIMessage_RepathFailed);
@@ -145,7 +145,7 @@ FString UBTTask_MoveTo::GetStaticDescription() const
 	return FString::Printf(TEXT("%s: %s"), *Super::GetStaticDescription(), *KeyDesc);
 }
 
-void UBTTask_MoveTo::DescribeRuntimeValues(const class UBehaviorTreeComponent* OwnerComp, uint8* NodeMemory, EBTDescriptionVerbosity::Type Verbosity, TArray<FString>& Values) const
+void UBTTask_MoveTo::DescribeRuntimeValues(const UBehaviorTreeComponent* OwnerComp, uint8* NodeMemory, EBTDescriptionVerbosity::Type Verbosity, TArray<FString>& Values) const
 {
 	Super::DescribeRuntimeValues(OwnerComp, NodeMemory, Verbosity, Values);
 
diff --git a/Engine/Source/Runtime/AIModule/Private/BehaviorTree/Tasks/BTTask_RunBehavior.cpp b/Engine/Source/Runtime/AIModule/Private/BehaviorTree/Tasks/BTTask_RunBehavior.cpp
--- a/Engine/Source/Runtime/AIModule/Private/BehaviorTree/Tasks/BTTask_RunBehavior.cpp
+++ b/Engine/Source/Runtime/AIModule/Private/BehaviorTree/Tasks/BTTask_RunBehavior.cpp
@@ -240,7 +240,7 @@ void UBTTask_RunBehavior::InjectNodes(UBehaviorTreeComponent* OwnerComp, uint8*
 	}
 }
 
-void UBTTask_RunBehavior::CleanupMemory(class UBehaviorTreeComponent* OwnerComp, uint8* NodeMemory, EBTMemoryClear::Type CleanupType) const
+void UBTTask_RunBehavior::CleanupMemory(UBehaviorTreeComponent* OwnerComp, uint8* NodeMemory, EBTMemoryClear::Type CleanupType) const
 {
 	Super::CleanupMemory(OwnerComp, NodeMemory, CleanupType);
 
diff --git a/Engine/Source/Runtime/AIModule/Private/BehaviorTree/Tasks/BTTask_RunEQSQuery.cpp b/Engine/Source/Runtime/AIModule/Private/BehaviorTree/Tasks/BTTask_RunEQSQuery.cpp
--- a/Engine/Source/Runtime/AIModule/Private/BehaviorTree/Tasks/BTTask_RunEQSQuery.cpp
+++ b/Engine/Source/Runtime/AIModule/Private/BehaviorTree/Tasks/BTTask_RunEQSQuery.cpp
@@ -18,7 +18,7 @@ UBTTask_RunEQSQuery::UBTTask_RunEQSQuery(const FObjectInitializer& ObjectInitial
 	}
 }
 
-EBTNodeResult::Type UBTTask_RunEQSQuery::ExecuteTask(class UBehaviorTreeComponent* OwnerComp, uint8* NodeMemory)
+EBTNodeResult::Type UBTTask_RunEQSQuery::ExecuteTask(UBehaviorTreeComponent* OwnerComp, uint8* NodeMemory)
 {
 	AActor* QueryOwner = OwnerComp->GetOwner();
 	if (AController* ControllerOwner = Cast<AController>(QueryOwner))
@@ -42,7 +42,7 @@ EBTNodeResult::Type UBTTask_RunEQSQuery::ExecuteTask(class UBehaviorTreeComponen
 	return EBTNodeResult::Failed;
 }
 
-EBTNodeResult::Type UBTTask_RunEQSQuery::AbortTask(class UBehaviorTreeComponent* OwnerComp, uint8* NodeMemory)
+EBTNodeResult::Type UBTTask_RunEQSQuery::AbortTask(UBehaviorTreeComponent* OwnerComp, uint8* NodeMemory)
 {
 	UWorld* MyWorld = OwnerComp->GetWorld();
 	UEnvQueryManager* QueryManager = UEnvQueryManager::GetCurrent(MyWorld);
@@ -62,7 +62,7 @@ FString UBTTask_RunEQSQuery::GetStaticDescription() const
 		*GetNameSafe(QueryTemplate), *BlackboardKey.SelectedKeyName.ToString());
 }
 
-void UBTTask_RunEQSQuery::DescribeRuntimeValues(const class UBehaviorTreeComponent* OwnerComp, uint8* NodeMemory, EBTDescriptionVerbosity::Type Verbosity, TArray<FString>& Values) const
+void UBTTask_RunEQSQuery::DescribeRuntimeValues(const UBehaviorTreeComponent* OwnerComp, uint8* NodeMemory, EBTDescriptionVerbosity::Type Verbosity, TArray<FString>& Values) const
 {
 	Super::DescribeRuntimeValues(OwnerComp, NodeMemory, Verbosity, Values);
 
@@ -78,7 +78,7 @@ uint16 UBTTask_RunEQSQuery::GetInstanceMemorySize() const
 	return sizeof(FBTEnvQueryTaskMemory);
 }
 
-void UBTTask_RunEQSQuery::OnQueryFinished(TSharedPtr<struct FEnvQueryResult> Result)
+void UBTTask_RunEQSQuery::OnQueryFinished(TSharedPtr<FEnvQueryResult> Result)
 {
 	if (Result->Status == EEnvQueryStatus::Aborted)
 	{
@@ -126,7 +126,7 @@ void UBTTask_RunEQSQuery::CollectKeyFilters()
 }
 
 #if WITH_EDITOR
-void UBTTask_RunEQSQuery::PostEditChangeProperty(struct FPropertyChangedEvent& PropertyChangedEvent)
+void UBTTask_RunEQSQuery::PostEditChangeProperty(FPropertyChangedEvent& PropertyChangedEvent)
 {
 	Super::PostEditChangeProperty(PropertyChangedEvent);
 
diff --git a/Engine/Source/Runtime/AIModule/Private/BehaviorTree/Tasks/BTTask_Wait.cpp b/Engine/Source/Runtime/AIModule/Private/BehaviorTree/Tasks/BTTask_Wait.cpp
--- a/Engine/Source/Runtime/AIModule/Private/BehaviorTree/Tasks/BTTask_Wait.cpp
+++ b/Engine/Source/Runtime/AIModule/Private/BehaviorTree/Tasks/BTTask_Wait.cpp
@@ -10,7 +10,7 @@ UBTTask_Wait::UBTTask_Wait(const FObjectInitializer& ObjectInitializer) : Super(
 	bNotifyTick = true;
 }
 
-EBTNodeResult::Type UBTTask_Wait::ExecuteTask(class UBehaviorTreeComponent* OwnerComp, uint8* NodeMemory)
+EBTNodeResult::Type UBTTask_Wait::ExecuteTask(UBehaviorTreeComponent* OwnerComp, uint8* NodeMemory)
 {
 	FBTWaitTaskMemory* MyMemory = (FBTWaitTaskMemory*)NodeMemory;
 	MyMemory->RemainingWaitTime = WaitTime;
@@ -18,7 +18,7 @@ EBTNodeResult::Type UBTTask_Wait::ExecuteTask(class UBehaviorTreeComponent* Owne
 	return EBTNodeResult::InProgress;
 }
 
-void UBTTask_Wait::TickTask(class UBehaviorTreeComponent* OwnerComp, uint8* NodeMemory, float DeltaSeconds)
+void UBTTask_Wait::TickTask(UBehaviorTreeComponent* OwnerComp, uint8* NodeMemory, float DeltaSeconds)
 {
 	FBTWaitTaskMemory* MyMemory = (FBTWaitTaskMemory*)NodeMemory;
 	MyMemory->RemainingWaitTime -= DeltaSeconds;
@@ -35,7 +35,7 @@ FString UBTTask_Wait::GetStaticDescription() const
 	return FString::Printf(TEXT("%s: %.1fs"), *Super::GetStaticDescription(), WaitTime);
 }
 
-void UBTTask_Wait::DescribeRuntimeValues(const class UBehaviorTreeComponent* OwnerComp, uint8* NodeMemory, EBTDescriptionVerbosity::Type Verbosity, TArray<FString>& Values) const
+void UBTTask_Wait::DescribeRuntimeValues(const UBehaviorTreeComponent* OwnerComp, uint8* NodeMemory, EBTDescriptionVerbosity::Type Verbosity, TArray<FString>& Values) const
 {
 	Super::DescribeRuntimeValues(OwnerComp, NodeMemory, Verbosity, Values);
 
diff --git a/Engine/Source/Runtime/AIModule/Private/Blueprint/AIBlueprintHelperLibrary.cpp b/Engine/Source/Runtime/AIModule/Private/Blueprint/AIBlueprintHelperLibrary.cpp
--- a/Engine/Source/Runtime/AIModule/Private/Blueprint/AIBlueprintHelperLibrary.cpp
+++ b/Engine/Source/Runtime/AIModule/Private/Blueprint/AIBlueprintHelperLibrary.cpp
@@ -78,7 +78,7 @@ UAIBlueprintHelperLibrary::UAIBlueprintHelperLibrary(const FObjectInitializer& O
 {
 }
 
-class UAIAsyncTaskBlueprintProxy* UAIBlueprintHelperLibrary::CreateMoveToProxyObject(class UObject* WorldContextObject, APawn* Pawn, FVector Destination, AActor* TargetActor, float AcceptanceRadius, bool bStopOnOverlap)
+UAIAsyncTaskBlueprintProxy* UAIBlueprintHelperLibrary::CreateMoveToProxyObject(UObject* WorldContextObject, APawn* Pawn, FVector Destination, AActor* TargetActor, float AcceptanceRadius, bool bStopOnOverlap)
 {
 	check(WorldContextObject);
 	if (!Pawn)
@@ -111,7 +111,7 @@ void UAIBlueprintHelperLibrary::SendAIMessage(APawn* Target, FName Message, UObj
 	FAIMessage::Send(Target, FAIMessage(Message, MessageSource, bSuccess));
 }
 
-APawn* UAIBlueprintHelperLibrary::SpawnAIFromClass(class UObject* WorldContextObject, TSubclassOf<APawn> PawnClass, class UBehaviorTree* BehaviorTree, FVector Location, FRotator Rotation, bool bNoCollisionFail)
+APawn* UAIBlueprintHelperLibrary::SpawnAIFromClass(UObject* WorldContextObject, TSubclassOf<APawn> PawnClass, UBehaviorTree* BehaviorTree, FVector Location, FRotator Rotation, bool bNoCollisionFail)
 {
 	APawn* NewPawn = NULL;
 
@@ -165,7 +165,7 @@ UBlackboardComponent* UAIBlueprintHelperLibrary::GetBlackboard(AActor* Target)
 	return BlackboardComp;
 }
 
-void UAIBlueprintHelperLibrary::LockAIResourcesWithAnimation(class UAnimInstance* AnimInstance, bool bLockMovement, bool LockAILogic)
+void UAIBlueprintHelperLibrary::LockAIResourcesWithAnimation(UAnimInstance* AnimInstance, bool bLockMovement, bool LockAILogic)
 {
 	if (AnimInstance == NULL)
 	{
@@ -190,7 +190,7 @@ void UAIBlueprintHelperLibrary::LockAIResourcesWithAnimation(class UAnimInstance
 	}
 }
 
-void UAIBlueprintHelperLibrary::UnlockAIResourcesWithAnimation(class UAnimInstance* AnimInstance, bool bUnlockMovement, bool UnlockAILogic)
+void UAIBlueprintHelperLibrary::UnlockAIResourcesWithAnimation(UAnimInstance* AnimInstance, bool bUnlockMovement, bool UnlockAILogic)
 {
 	if (AnimInstance == NULL)
 	{
diff --git a/Engine/Source/Runtime/AIModule/Private/BrainComponent.cpp b/Engine/Source/Runtime/AIModule/Private/BrainComponent.cpp
--- a/Engine/Source/Runtime/AIModule/Private/BrainComponent.cpp
+++ b/Engine/Source/Runtime/AIModule/Private/BrainComponent.cpp
@@ -184,7 +184,7 @@ UBrainComponent::UBrainComponent(const FObjectInitializer& ObjectInitializer) :
 }
 
 #if ENABLE_VISUAL_LOG
-void UBrainComponent::DescribeSelfToVisLog(struct FVisualLogEntry* Snapshot) const
+void UBrainComponent::DescribeSelfToVisLog(FVisualLogEntry* Snapshot) const
 {
 	const static UEnum* SourceEnum = FindObject<UEnum>(ANY_PACKAGE, TEXT("EAILockSource"));
 
diff --git a/Engine/Source/Runtime/AIModule/Private/EnvironmentQuery/Contexts/EnvQueryContext_BlueprintBase.cpp b/Engine/Source/Runtime/AIModule/Private/EnvironmentQuery/Contexts/EnvQueryContext_BlueprintBase.cpp
--- a/Engine/Source/Runtime/AIModule/Private/EnvironmentQuery/Contexts/EnvQueryContext_BlueprintBase.cpp
+++ b/Engine/Source/Runtime/AIModule/Private/EnvironmentQuery/Contexts/EnvQueryContext_BlueprintBase.cpp
@@ -55,7 +55,7 @@ UEnvQueryContext_BlueprintBase::UEnvQueryContext_BlueprintBase(const FObjectInit
 	}
 }
 
-void UEnvQueryContext_BlueprintBase::ProvideContext(struct FEnvQueryInstance& QueryInstance, struct FEnvQueryContextData& ContextData) const
+void UEnvQueryContext_BlueprintBase::ProvideContext(FEnvQueryInstance& QueryInstance, FEnvQueryContextData& ContextData) const
 {
 	AActor* QuerierActor = Cast<AActor>(QueryInstance.Owner.Get());
 
@@ -99,7 +99,7 @@ void UEnvQueryContext_BlueprintBase::ProvideContext(struct FEnvQueryInstance& Qu
 	}
 }
 
-class UWorld* UEnvQueryContext_BlueprintBase::GetWorld() const
+UWorld* UEnvQueryContext_BlueprintBase::GetWorld() const
 {
 	check(GetOuter() != NULL);
 	
@@ -110,10 +110,10 @@ class UWorld* UEnvQueryContext_BlueprintBase::GetWorld() const
 	}
 
 	// Outer should always be UEnvQueryManager* in the game, which implements GetWorld() and therefore makes this
-	// class a correct world context.  However, in the editor the outer is /Script/AIModule (at compile time), which
+	// a correct world context.  However, in the editor the outer is /Script/AIModule (at compile time), which
 	// does not explicitly implement GetWorld().  For that reason, calling GetWorld() generically in that case on the
-	// AIModule calls to the base class implementation, which causes a blueprint compile warning in the Editor
-	// which states that the function isn't safe to call on this class (due to requiring WorldContext which it doesn't
+	// AIModule calls to the base implementation, which causes a blueprint compile warning in the Editor
+	// which states that the function isn't safe to call on this (due to requiring WorldContext which it doesn't
 	// provide).  Simply returning NULL in this case fixes those erroneous blueprint compile warnings.
 	return NULL;
 }
\ No newline at end of file
diff --git a/Engine/Source/Runtime/AIModule/Private/EnvironmentQuery/Contexts/EnvQueryContext_Querier.cpp b/Engine/Source/Runtime/AIModule/Private/EnvironmentQuery/Contexts/EnvQueryContext_Querier.cpp
--- a/Engine/Source/Runtime/AIModule/Private/EnvironmentQuery/Contexts/EnvQueryContext_Querier.cpp
+++ b/Engine/Source/Runtime/AIModule/Private/EnvironmentQuery/Contexts/EnvQueryContext_Querier.cpp
@@ -8,7 +8,7 @@ UEnvQueryContext_Querier::UEnvQueryContext_Querier(const FObjectInitializer& Obj
 {
 }
 
-void UEnvQueryContext_Querier::ProvideContext(struct FEnvQueryInstance& QueryInstance, struct FEnvQueryContextData& ContextData) const
+void UEnvQueryContext_Querier::ProvideContext(FEnvQueryInstance& QueryInstance, FEnvQueryContextData& ContextData) const
 {
 	AActor* QueryOwner = Cast<AActor>(QueryInstance.Owner.Get());
 	UEnvQueryItemType_Actor::SetContextHelper(ContextData, QueryOwner);
diff --git a/Engine/Source/Runtime/AIModule/Private/EnvironmentQuery/EQSRenderingComponent.cpp b/Engine/Source/Runtime/AIModule/Private/EnvironmentQuery/EQSRenderingComponent.cpp
--- a/Engine/Source/Runtime/AIModule/Private/EnvironmentQuery/EQSRenderingComponent.cpp
+++ b/Engine/Source/Runtime/AIModule/Private/EnvironmentQuery/EQSRenderingComponent.cpp
@@ -10,7 +10,7 @@ static const int32 EQSMaxItemsDrawn = 10000;
 
 namespace FEQSRenderingHelper
 {
-	FVector ExtractLocation(TSubclassOf<class UEnvQueryItemType> ItemType, TArray<uint8> RawData, const TArray<FEnvQueryItem>& Items, int32 Index)
+	FVector ExtractLocation(TSubclassOf<UEnvQueryItemType> ItemType, TArray<uint8> RawData, const TArray<FEnvQueryItem>& Items, int32 Index)
 	{
 		if (Items.IsValidIndex(Index) &&
 			ItemType->IsChildOf(UEnvQueryItemType_VectorBase::StaticClass()))
@@ -90,7 +90,7 @@ FEQSSceneProxy::FEQSSceneProxy(const UPrimitiveComponent* InComponent, const FSt
 }
 
 #if  USE_EQS_DEBUGGER 
-void FEQSSceneProxy::CollectEQSData(const UPrimitiveComponent* InComponent, const class IEQSQueryResultSourceInterface* InQueryDataSource, TArray<FSphere>& Spheres, TArray<FText3d>& Texts, TArray<EQSDebug::FDebugHelper>& DebugItems)
+void FEQSSceneProxy::CollectEQSData(const UPrimitiveComponent* InComponent, const IEQSQueryResultSourceInterface* InQueryDataSource, TArray<FSphere>& Spheres, TArray<FText3d>& Texts, TArray<EQSDebug::FDebugHelper>& DebugItems)
 {
 	AActor* ActorOwner = InComponent ? InComponent->GetOwner() : NULL;
 	IEQSQueryResultSourceInterface* QueryDataSource = const_cast<IEQSQueryResultSourceInterface*>(InQueryDataSource);
@@ -121,7 +121,7 @@ void FEQSSceneProxy::CollectEQSData(const UPrimitiveComponent* InComponent, cons
 	}
 }
 
-void FEQSSceneProxy::CollectEQSData(const struct FEnvQueryResult* ResultItems, const struct FEnvQueryInstance* QueryInstance, TArray<FSphere>& Spheres, TArray<FText3d>& Texts, bool ShouldDrawFailedItems, TArray<EQSDebug::FDebugHelper>& DebugItems)
+void FEQSSceneProxy::CollectEQSData(const FEnvQueryResult* ResultItems, const FEnvQueryInstance* QueryInstance, TArray<FSphere>& Spheres, TArray<FText3d>& Texts, bool ShouldDrawFailedItems, TArray<EQSDebug::FDebugHelper>& DebugItems)
 {
 	if (ResultItems == NULL)
 	{
diff --git a/Engine/Source/Runtime/AIModule/Private/EnvironmentQuery/EQSTestingPawn.cpp b/Engine/Source/Runtime/AIModule/Private/EnvironmentQuery/EQSTestingPawn.cpp
--- a/Engine/Source/Runtime/AIModule/Private/EnvironmentQuery/EQSTestingPawn.cpp
+++ b/Engine/Source/Runtime/AIModule/Private/EnvironmentQuery/EQSTestingPawn.cpp
@@ -14,7 +14,7 @@
 #if WITH_EDITOR
 #include "Engine/Brush.h"
 #include "Editor/EditorEngine.h"
-extern UNREALED_API class UEditorEngine* GEditor;
+extern UNREALED_API UEditorEngine* GEditor;
 #endif // WITH_EDITOR
 
 #include "Engine/Selection.h"
@@ -231,7 +231,7 @@ const FEnvQueryInstance* AEQSTestingPawn::GetQueryInstance() const
 
 #if WITH_EDITOR
 
-void AEQSTestingPawn::PostEditChangeProperty( struct FPropertyChangedEvent& PropertyChangedEvent)
+void AEQSTestingPawn::PostEditChangeProperty( FPropertyChangedEvent& PropertyChangedEvent)
 {
 	static const FName NAME_QueryTemplate = GET_MEMBER_NAME_CHECKED(AEQSTestingPawn, QueryTemplate);
 	static const FName NAME_StepToDebugDraw = GET_MEMBER_NAME_CHECKED(AEQSTestingPawn, StepToDebugDraw);
diff --git a/Engine/Source/Runtime/AIModule/Private/EnvironmentQuery/EnvQueryContext.cpp b/Engine/Source/Runtime/AIModule/Private/EnvironmentQuery/EnvQueryContext.cpp
--- a/Engine/Source/Runtime/AIModule/Private/EnvironmentQuery/EnvQueryContext.cpp
+++ b/Engine/Source/Runtime/AIModule/Private/EnvironmentQuery/EnvQueryContext.cpp
@@ -7,7 +7,7 @@ UEnvQueryContext::UEnvQueryContext(const FObjectInitializer& ObjectInitializer)
 {
 }
 
-void UEnvQueryContext::ProvideContext(struct FEnvQueryInstance& QueryInstance, struct FEnvQueryContextData& ContextData) const
+void UEnvQueryContext::ProvideContext(FEnvQueryInstance& QueryInstance, FEnvQueryContextData& ContextData) const
 {
 	// empty in base class
 }
diff --git a/Engine/Source/Runtime/AIModule/Private/EnvironmentQuery/EnvQueryDebugHelpers.cpp b/Engine/Source/Runtime/AIModule/Private/EnvironmentQuery/EnvQueryDebugHelpers.cpp
--- a/Engine/Source/Runtime/AIModule/Private/EnvironmentQuery/EnvQueryDebugHelpers.cpp
+++ b/Engine/Source/Runtime/AIModule/Private/EnvironmentQuery/EnvQueryDebugHelpers.cpp
@@ -7,7 +7,7 @@
 #include "EnvironmentQuery/EQSQueryResultSourceInterface.h"
 
 #if USE_EQS_DEBUGGER
-void UEnvQueryDebugHelpers::QueryToBlobArray(struct FEnvQueryInstance& Query, TArray<uint8>& BlobArray, bool bUseCompression)
+void UEnvQueryDebugHelpers::QueryToBlobArray(FEnvQueryInstance& Query, TArray<uint8>& BlobArray, bool bUseCompression)
 {
 	EQSDebug::FQueryData EQSLocalData;
 	UEnvQueryDebugHelpers::QueryToDebugData(Query, EQSLocalData);
@@ -38,7 +38,7 @@ void UEnvQueryDebugHelpers::QueryToBlobArray(struct FEnvQueryInstance& Query, TA
 	}
 }
 
-void UEnvQueryDebugHelpers::QueryToDebugData(struct FEnvQueryInstance& Query, EQSDebug::FQueryData& EQSLocalData)
+void UEnvQueryDebugHelpers::QueryToDebugData(FEnvQueryInstance& Query, EQSDebug::FQueryData& EQSLocalData)
 {
 	// step 1: data for rendering component
 	EQSLocalData.Reset();
@@ -132,7 +132,7 @@ void  UEnvQueryDebugHelpers::BlobArrayToDebugData(const TArray<uint8>& BlobArray
 #endif //ENABLE_VISUAL_LOG
 
 #if ENABLE_VISUAL_LOG && USE_EQS_DEBUGGER
-void UEnvQueryDebugHelpers::LogQueryInternal(struct FEnvQueryInstance& Query, const struct FLogCategoryBase& Category, ELogVerbosity::Type Verbosity, float TimeSeconds, FVisualLogEntry *CurrentEntry)
+void UEnvQueryDebugHelpers::LogQueryInternal(FEnvQueryInstance& Query, const FLogCategoryBase& Category, ELogVerbosity::Type Verbosity, float TimeSeconds, FVisualLogEntry *CurrentEntry)
 {
 	const int32 UniqueId = FVisualLogger::Get().GetUniqueId(TimeSeconds);
 	TArray<uint8> BlobArray;
diff --git a/Engine/Source/Runtime/AIModule/Private/EnvironmentQuery/EnvQueryInstance.cpp b/Engine/Source/Runtime/AIModule/Private/EnvironmentQuery/EnvQueryInstance.cpp
--- a/Engine/Source/Runtime/AIModule/Private/EnvironmentQuery/EnvQueryInstance.cpp
+++ b/Engine/Source/Runtime/AIModule/Private/EnvironmentQuery/EnvQueryInstance.cpp
@@ -418,9 +418,9 @@ void FEnvQueryInstance::SortScores()
 	struct FSortHelperForDebugData
 	{
 		FEnvQueryItem	Item;
-		struct FEnvQueryItemDetails ItemDetails;
+		FEnvQueryItemDetails ItemDetails;
 
-		FSortHelperForDebugData(const FEnvQueryItem&	InItem, struct FEnvQueryItemDetails& InDebugDetails) : Item(InItem), ItemDetails(InDebugDetails) {}
+		FSortHelperForDebugData(const FEnvQueryItem&	InItem, FEnvQueryItemDetails& InDebugDetails) : Item(InItem), ItemDetails(InDebugDetails) {}
 		bool operator<(const FSortHelperForDebugData& Other) const
 		{
 			return Item < Other.Item;
diff --git a/Engine/Source/Runtime/AIModule/Private/EnvironmentQuery/EnvQueryManager.cpp b/Engine/Source/Runtime/AIModule/Private/EnvironmentQuery/EnvQueryManager.cpp
--- a/Engine/Source/Runtime/AIModule/Private/EnvironmentQuery/EnvQueryManager.cpp
+++ b/Engine/Source/Runtime/AIModule/Private/EnvironmentQuery/EnvQueryManager.cpp
@@ -14,7 +14,7 @@
 #include "Editor/EditorEngine.h"
 #include "EngineUtils.h"
 
-extern UNREALED_API class UEditorEngine* GEditor;
+extern UNREALED_API UEditorEngine* GEditor;
 #endif // WITH_EDITOR
 
 DEFINE_LOG_CATEGORY(LogEQS);
@@ -30,7 +30,7 @@ DEFINE_STAT(STAT_AI_EQS_InstanceMemory);
 //////////////////////////////////////////////////////////////////////////
 // FEnvQueryRequest
 
-FEnvQueryRequest& FEnvQueryRequest::SetNamedParams(const TArray<struct FEnvNamedValue>& Params)
+FEnvQueryRequest& FEnvQueryRequest::SetNamedParams(const TArray<FEnvNamedValue>& Params)
 {
 	for (int32 ParamIndex = 0; ParamIndex < Params.Num(); ParamIndex++)
 	{
@@ -88,7 +88,7 @@ UEnvQueryManager::UEnvQueryManager(const FObjectInitializer& ObjectInitializer)
 	NextQueryID = 0;
 }
 
-class UWorld* UEnvQueryManager::GetWorld() const
+UWorld* UEnvQueryManager::GetWorld() const
 {
 	return Cast<UWorld>(GetOuter());
 }
@@ -161,7 +161,7 @@ int32 UEnvQueryManager::RunQuery(const FEnvQueryRequest& Request, EEnvQueryRunMo
 	return RunQuery(QueryInstance, FinishDelegate);
 }
 
-int32 UEnvQueryManager::RunQuery(TSharedPtr<struct FEnvQueryInstance> QueryInstance, FQueryFinishedSignature const& FinishDelegate)
+int32 UEnvQueryManager::RunQuery(TSharedPtr<FEnvQueryInstance> QueryInstance, FQueryFinishedSignature const& FinishDelegate)
 {
 	if (QueryInstance.IsValid() == false)
 	{
@@ -196,7 +196,7 @@ TSharedPtr<FEnvQueryResult> UEnvQueryManager::RunInstantQuery(const FEnvQueryReq
 	return QueryInstance;
 }
 
-TSharedPtr<struct FEnvQueryInstance> UEnvQueryManager::PrepareQueryInstance(const FEnvQueryRequest& Request, EEnvQueryRunMode::Type RunMode)
+TSharedPtr<FEnvQueryInstance> UEnvQueryManager::PrepareQueryInstance(const FEnvQueryRequest& Request, EEnvQueryRunMode::Type RunMode)
 {
 	TSharedPtr<FEnvQueryInstance> QueryInstance = CreateQueryInstance(Request.QueryTemplate, RunMode);
 	if (!QueryInstance.IsValid())
@@ -510,7 +510,7 @@ TSharedPtr<FEnvQueryInstance> UEnvQueryManager::CreateQueryInstance(const UEnvQu
 	return NewInstance;
 }
 
-void UEnvQueryManager::CreateOptionInstance(UEnvQueryOption* OptionTemplate, const TArray<UEnvQueryTest*>& SortedTests, struct FEnvQueryInstance& Instance)
+void UEnvQueryManager::CreateOptionInstance(UEnvQueryOption* OptionTemplate, const TArray<UEnvQueryTest*>& SortedTests, FEnvQueryInstance& Instance)
 {
 	FEnvQueryOptionInstance OptionInstance;
 	OptionInstance.Generator = OptionTemplate->Generator;
diff --git a/Engine/Source/Runtime/AIModule/Private/EnvironmentQuery/EnvQueryTypes.cpp b/Engine/Source/Runtime/AIModule/Private/EnvironmentQuery/EnvQueryTypes.cpp
--- a/Engine/Source/Runtime/AIModule/Private/EnvironmentQuery/EnvQueryTypes.cpp
+++ b/Engine/Source/Runtime/AIModule/Private/EnvironmentQuery/EnvQueryTypes.cpp
@@ -57,7 +57,7 @@ FText UEnvQueryTypes::GetShortTypeName(const UObject* Ob)
 	return FText::FromString(TypeDesc);
 }
 
-FText UEnvQueryTypes::DescribeContext(TSubclassOf<class UEnvQueryContext> ContextClass)
+FText UEnvQueryTypes::DescribeContext(TSubclassOf<UEnvQueryContext> ContextClass)
 {
 	return GetShortTypeName(ContextClass);
 }
diff --git a/Engine/Source/Runtime/AIModule/Private/EnvironmentQuery/Generators/EnvQueryGenerator_ActorsOfClass.cpp b/Engine/Source/Runtime/AIModule/Private/EnvironmentQuery/Generators/EnvQueryGenerator_ActorsOfClass.cpp
--- a/Engine/Source/Runtime/AIModule/Private/EnvironmentQuery/Generators/EnvQueryGenerator_ActorsOfClass.cpp
+++ b/Engine/Source/Runtime/AIModule/Private/EnvironmentQuery/Generators/EnvQueryGenerator_ActorsOfClass.cpp
@@ -55,7 +55,7 @@ FText UEnvQueryGenerator_ActorsOfClass::GetDescriptionTitle() const
 	Args.Add(TEXT("ActorsClass"), FText::FromString(GetNameSafe(SearchedActorClass)));
 	Args.Add(TEXT("DescribeContext"), UEnvQueryTypes::DescribeContext(SearchCenter));
 
-	return FText::Format(LOCTEXT("DescriptionGenerateActorsAroundContext", "{DescriptionTitle}: generate set of actors of class {ActorsClass} around {DescribeContext}"), Args);
+	return FText::Format(LOCTEXT("DescriptionGenerateActorsAroundContext", "{DescriptionTitle}: generate set of actors of {ActorsClass} around {DescribeContext}"), Args);
 };
 
 FText UEnvQueryGenerator_ActorsOfClass::GetDescriptionDetails() const
diff --git a/Engine/Source/Runtime/AIModule/Private/EnvironmentQuery/Generators/EnvQueryGenerator_OnCircle.cpp b/Engine/Source/Runtime/AIModule/Private/EnvironmentQuery/Generators/EnvQueryGenerator_OnCircle.cpp
--- a/Engine/Source/Runtime/AIModule/Private/EnvironmentQuery/Generators/EnvQueryGenerator_OnCircle.cpp
+++ b/Engine/Source/Runtime/AIModule/Private/EnvironmentQuery/Generators/EnvQueryGenerator_OnCircle.cpp
@@ -351,7 +351,7 @@ FText UEnvQueryGenerator_OnCircle::GetDescriptionDetails() const
 
 #if WITH_EDITOR
 
-void UEnvQueryGenerator_OnCircle::PostEditChangeProperty( struct FPropertyChangedEvent& PropertyChangedEvent) 
+void UEnvQueryGenerator_OnCircle::PostEditChangeProperty( FPropertyChangedEvent& PropertyChangedEvent) 
 {
 	static const FName NAME_Angle = GET_MEMBER_NAME_CHECKED(UEnvQueryGenerator_OnCircle, Angle);
 	static const FName NAME_Radius = GET_MEMBER_NAME_CHECKED(UEnvQueryGenerator_OnCircle, Radius);
diff --git a/Engine/Source/Runtime/AIModule/Private/EnvironmentQuery/Generators/EnvQueryGenerator_PathingGrid.cpp b/Engine/Source/Runtime/AIModule/Private/EnvironmentQuery/Generators/EnvQueryGenerator_PathingGrid.cpp
--- a/Engine/Source/Runtime/AIModule/Private/EnvironmentQuery/Generators/EnvQueryGenerator_PathingGrid.cpp
+++ b/Engine/Source/Runtime/AIModule/Private/EnvironmentQuery/Generators/EnvQueryGenerator_PathingGrid.cpp
@@ -108,7 +108,7 @@ FText UEnvQueryGenerator_PathingGrid::GetDescriptionDetails() const
 #if WITH_RECAST
 #define ENVQUERY_CLUSTER_SEARCH 0
 
-void UEnvQueryGenerator_PathingGrid::FindNodeRefsInPathDistance(const class ARecastNavMesh* NavMesh, const FVector& ContextLocation, float InMaxPathDistance, bool bPathFromContext, TArray<NavNodeRef>& NodeRefs, FBox& NodeRefsBounds) const
+void UEnvQueryGenerator_PathingGrid::FindNodeRefsInPathDistance(const ARecastNavMesh* NavMesh, const FVector& ContextLocation, float InMaxPathDistance, bool bPathFromContext, TArray<NavNodeRef>& NodeRefs, FBox& NodeRefsBounds) const
 {
 	FBox MyBounds(0);
 
@@ -150,8 +150,8 @@ void UEnvQueryGenerator_PathingGrid::FindNodeRefsInPathDistance(const class ARec
 	NodeRefsBounds = MyBounds;
 }
 
-bool UEnvQueryGenerator_PathingGrid::IsNavLocationInPathDistance(const class ARecastNavMesh* NavMesh,
-		const struct FNavLocation& NavLocation, const TArray<NavNodeRef>& NodeRefs) const
+bool UEnvQueryGenerator_PathingGrid::IsNavLocationInPathDistance(const ARecastNavMesh* NavMesh,
+		const FNavLocation& NavLocation, const TArray<NavNodeRef>& NodeRefs) const
 {
 #if ENVQUERY_CLUSTER_SEARCH
 	const NavNodeRef ClusterRef = NavMesh->GetClusterRef(NavLocation.NodeRef);
diff --git a/Engine/Source/Runtime/AIModule/Private/EnvironmentQuery/Items/EnvQueryItemType.cpp b/Engine/Source/Runtime/AIModule/Private/EnvironmentQuery/Items/EnvQueryItemType.cpp
--- a/Engine/Source/Runtime/AIModule/Private/EnvironmentQuery/Items/EnvQueryItemType.cpp
+++ b/Engine/Source/Runtime/AIModule/Private/EnvironmentQuery/Items/EnvQueryItemType.cpp
@@ -21,11 +21,11 @@ void UEnvQueryItemType::FinishDestroy()
 	Super::FinishDestroy();
 }
 
-void UEnvQueryItemType::AddBlackboardFilters(struct FBlackboardKeySelector& KeySelector, UObject* FilterOwner) const
+void UEnvQueryItemType::AddBlackboardFilters(FBlackboardKeySelector& KeySelector, UObject* FilterOwner) const
 {
 }
 
-bool UEnvQueryItemType::StoreInBlackboard(struct FBlackboardKeySelector& KeySelector, class UBlackboardComponent* Blackboard, const uint8* RawData) const
+bool UEnvQueryItemType::StoreInBlackboard(FBlackboardKeySelector& KeySelector, UBlackboardComponent* Blackboard, const uint8* RawData) const
 {
 	return false;
 }
diff --git a/Engine/Source/Runtime/AIModule/Private/EnvironmentQuery/Items/EnvQueryItemType_Actor.cpp b/Engine/Source/Runtime/AIModule/Private/EnvironmentQuery/Items/EnvQueryItemType_Actor.cpp
--- a/Engine/Source/Runtime/AIModule/Private/EnvironmentQuery/Items/EnvQueryItemType_Actor.cpp
+++ b/Engine/Source/Runtime/AIModule/Private/EnvironmentQuery/Items/EnvQueryItemType_Actor.cpp
@@ -37,7 +37,7 @@ AActor* UEnvQueryItemType_Actor::GetActor(const uint8* RawData) const
 	return UEnvQueryItemType_Actor::GetValue(RawData);
 }
 
-void UEnvQueryItemType_Actor::SetContextHelper(struct FEnvQueryContextData& ContextData, const AActor* SingleActor)
+void UEnvQueryItemType_Actor::SetContextHelper(FEnvQueryContextData& ContextData, const AActor* SingleActor)
 {
 	ContextData.ValueType = UEnvQueryItemType_Actor::StaticClass();
 	ContextData.NumValues = 1;
@@ -46,7 +46,7 @@ void UEnvQueryItemType_Actor::SetContextHelper(struct FEnvQueryContextData& Cont
 	UEnvQueryItemType_Actor::SetValue(ContextData.RawData.GetData(), SingleActor);
 }
 
-void UEnvQueryItemType_Actor::SetContextHelper(struct FEnvQueryContextData& ContextData, const TArray<const AActor*>& MultipleActors)
+void UEnvQueryItemType_Actor::SetContextHelper(FEnvQueryContextData& ContextData, const TArray<const AActor*>& MultipleActors)
 {
 	ContextData.ValueType = UEnvQueryItemType_Actor::StaticClass();
 	ContextData.NumValues = MultipleActors.Num();
@@ -60,7 +60,7 @@ void UEnvQueryItemType_Actor::SetContextHelper(struct FEnvQueryContextData& Cont
 	}
 }
 
-void UEnvQueryItemType_Actor::SetContextHelper(struct FEnvQueryContextData& ContextData, const TArray<AActor*>& MultipleActors)
+void UEnvQueryItemType_Actor::SetContextHelper(FEnvQueryContextData& ContextData, const TArray<AActor*>& MultipleActors)
 {
 	ContextData.ValueType = UEnvQueryItemType_Actor::StaticClass();
 	ContextData.NumValues = MultipleActors.Num();
diff --git a/Engine/Source/Runtime/AIModule/Private/EnvironmentQuery/Items/EnvQueryItemType_ActorBase.cpp b/Engine/Source/Runtime/AIModule/Private/EnvironmentQuery/Items/EnvQueryItemType_ActorBase.cpp
--- a/Engine/Source/Runtime/AIModule/Private/EnvironmentQuery/Items/EnvQueryItemType_ActorBase.cpp
+++ b/Engine/Source/Runtime/AIModule/Private/EnvironmentQuery/Items/EnvQueryItemType_ActorBase.cpp
@@ -9,13 +9,13 @@ UEnvQueryItemType_ActorBase::UEnvQueryItemType_ActorBase(const FObjectInitialize
 {
 }
 
-void UEnvQueryItemType_ActorBase::AddBlackboardFilters(struct FBlackboardKeySelector& KeySelector, UObject* FilterOwner) const
+void UEnvQueryItemType_ActorBase::AddBlackboardFilters(FBlackboardKeySelector& KeySelector, UObject* FilterOwner) const
 {
 	Super::AddBlackboardFilters(KeySelector, FilterOwner);
 	KeySelector.AddObjectFilter(FilterOwner, AActor::StaticClass());
 }
 
-bool UEnvQueryItemType_ActorBase::StoreInBlackboard(struct FBlackboardKeySelector& KeySelector, class UBlackboardComponent* Blackboard, const uint8* RawData) const
+bool UEnvQueryItemType_ActorBase::StoreInBlackboard(FBlackboardKeySelector& KeySelector, UBlackboardComponent* Blackboard, const uint8* RawData) const
 {
 	bool bStored = Super::StoreInBlackboard(KeySelector, Blackboard, RawData);
 	if (!bStored && KeySelector.SelectedKeyType == UBlackboardKeyType_Object::StaticClass())
diff --git a/Engine/Source/Runtime/AIModule/Private/EnvironmentQuery/Items/EnvQueryItemType_Direction.cpp b/Engine/Source/Runtime/AIModule/Private/EnvironmentQuery/Items/EnvQueryItemType_Direction.cpp
--- a/Engine/Source/Runtime/AIModule/Private/EnvironmentQuery/Items/EnvQueryItemType_Direction.cpp
+++ b/Engine/Source/Runtime/AIModule/Private/EnvironmentQuery/Items/EnvQueryItemType_Direction.cpp
@@ -39,7 +39,7 @@ FString UEnvQueryItemType_Direction::GetDescription(const uint8* RawData) const
 	return FString::Printf(TEXT("(P=%.0f,Y=%.0f,R=%.0f)"), Rot.Pitch, Rot.Yaw, Rot.Roll);
 }
 
-void UEnvQueryItemType_Direction::SetContextHelper(struct FEnvQueryContextData& ContextData, const FVector& SingleDirection)
+void UEnvQueryItemType_Direction::SetContextHelper(FEnvQueryContextData& ContextData, const FVector& SingleDirection)
 {
 	ContextData.ValueType = UEnvQueryItemType_Direction::StaticClass();
 	ContextData.NumValues = 1;
@@ -48,7 +48,7 @@ void UEnvQueryItemType_Direction::SetContextHelper(struct FEnvQueryContextData&
 	UEnvQueryItemType_Direction::SetValue((uint8*)ContextData.RawData.GetData(), SingleDirection);
 }
 
-void UEnvQueryItemType_Direction::SetContextHelper(struct FEnvQueryContextData& ContextData, const FRotator& SingleRotation)
+void UEnvQueryItemType_Direction::SetContextHelper(FEnvQueryContextData& ContextData, const FRotator& SingleRotation)
 {
 	ContextData.ValueType = UEnvQueryItemType_Direction::StaticClass();
 	ContextData.NumValues = 1;
@@ -57,7 +57,7 @@ void UEnvQueryItemType_Direction::SetContextHelper(struct FEnvQueryContextData&
 	UEnvQueryItemType_Direction::SetValueRot((uint8*)ContextData.RawData.GetData(), SingleRotation);
 }
 
-void UEnvQueryItemType_Direction::SetContextHelper(struct FEnvQueryContextData& ContextData, const TArray<FVector>& MultipleDirections)
+void UEnvQueryItemType_Direction::SetContextHelper(FEnvQueryContextData& ContextData, const TArray<FVector>& MultipleDirections)
 {
 	ContextData.ValueType = UEnvQueryItemType_Direction::StaticClass();
 	ContextData.NumValues = MultipleDirections.Num();
@@ -71,7 +71,7 @@ void UEnvQueryItemType_Direction::SetContextHelper(struct FEnvQueryContextData&
 	}
 }
 
-void UEnvQueryItemType_Direction::SetContextHelper(struct FEnvQueryContextData& ContextData, const TArray<FRotator>& MultipleRotations)
+void UEnvQueryItemType_Direction::SetContextHelper(FEnvQueryContextData& ContextData, const TArray<FRotator>& MultipleRotations)
 {
 	ContextData.ValueType = UEnvQueryItemType_Direction::StaticClass();
 	ContextData.NumValues = MultipleRotations.Num();
diff --git a/Engine/Source/Runtime/AIModule/Private/EnvironmentQuery/Items/EnvQueryItemType_Point.cpp b/Engine/Source/Runtime/AIModule/Private/EnvironmentQuery/Items/EnvQueryItemType_Point.cpp
--- a/Engine/Source/Runtime/AIModule/Private/EnvironmentQuery/Items/EnvQueryItemType_Point.cpp
+++ b/Engine/Source/Runtime/AIModule/Private/EnvironmentQuery/Items/EnvQueryItemType_Point.cpp
@@ -23,7 +23,7 @@ FVector UEnvQueryItemType_Point::GetLocation(const uint8* RawData) const
 	return UEnvQueryItemType_Point::GetValue(RawData);
 }
 
-void UEnvQueryItemType_Point::SetContextHelper(struct FEnvQueryContextData& ContextData, const FVector& SinglePoint)
+void UEnvQueryItemType_Point::SetContextHelper(FEnvQueryContextData& ContextData, const FVector& SinglePoint)
 {
 	ContextData.ValueType = UEnvQueryItemType_Point::StaticClass();
 	ContextData.NumValues = 1;
@@ -32,7 +32,7 @@ void UEnvQueryItemType_Point::SetContextHelper(struct FEnvQueryContextData& Cont
 	UEnvQueryItemType_Point::SetValue((uint8*)ContextData.RawData.GetData(), SinglePoint);
 }
 
-void UEnvQueryItemType_Point::SetContextHelper(struct FEnvQueryContextData& ContextData, const TArray<FVector>& MultiplePoints)
+void UEnvQueryItemType_Point::SetContextHelper(FEnvQueryContextData& ContextData, const TArray<FVector>& MultiplePoints)
 {
 	ContextData.ValueType = UEnvQueryItemType_Point::StaticClass();
 	ContextData.NumValues = MultiplePoints.Num();
diff --git a/Engine/Source/Runtime/AIModule/Private/EnvironmentQuery/Items/EnvQueryItemType_VectorBase.cpp b/Engine/Source/Runtime/AIModule/Private/EnvironmentQuery/Items/EnvQueryItemType_VectorBase.cpp
--- a/Engine/Source/Runtime/AIModule/Private/EnvironmentQuery/Items/EnvQueryItemType_VectorBase.cpp
+++ b/Engine/Source/Runtime/AIModule/Private/EnvironmentQuery/Items/EnvQueryItemType_VectorBase.cpp
@@ -9,12 +9,12 @@ UEnvQueryItemType_VectorBase::UEnvQueryItemType_VectorBase(const FObjectInitiali
 {
 }
 
-void UEnvQueryItemType_VectorBase::AddBlackboardFilters(struct FBlackboardKeySelector& KeySelector, UObject* FilterOwner) const
+void UEnvQueryItemType_VectorBase::AddBlackboardFilters(FBlackboardKeySelector& KeySelector, UObject* FilterOwner) const
 {
 	KeySelector.AddVectorFilter(FilterOwner);
 }
 
-bool UEnvQueryItemType_VectorBase::StoreInBlackboard(struct FBlackboardKeySelector& KeySelector, class UBlackboardComponent* Blackboard, const uint8* RawData) const
+bool UEnvQueryItemType_VectorBase::StoreInBlackboard(FBlackboardKeySelector& KeySelector, UBlackboardComponent* Blackboard, const uint8* RawData) const
 {
 	bool bStored = false;
 	if (KeySelector.SelectedKeyType == UBlackboardKeyType_Vector::StaticClass())
diff --git a/Engine/Source/Runtime/AIModule/Private/EnvironmentQuery/Tests/EnvQueryTest_Dot.cpp b/Engine/Source/Runtime/AIModule/Private/EnvironmentQuery/Tests/EnvQueryTest_Dot.cpp
--- a/Engine/Source/Runtime/AIModule/Private/EnvironmentQuery/Tests/EnvQueryTest_Dot.cpp
+++ b/Engine/Source/Runtime/AIModule/Private/EnvironmentQuery/Tests/EnvQueryTest_Dot.cpp
@@ -110,8 +110,8 @@ void UEnvQueryTest_Dot::RunTest(FEnvQueryInstance& QueryInstance) const
 	}
 }
 
-void UEnvQueryTest_Dot::GatherLineDirections(TArray<FVector>& Directions, struct FEnvQueryInstance& QueryInstance, const FVector& ItemLocation,
-	TSubclassOf<class UEnvQueryContext> LineFrom, TSubclassOf<class UEnvQueryContext> LineTo) const
+void UEnvQueryTest_Dot::GatherLineDirections(TArray<FVector>& Directions, FEnvQueryInstance& QueryInstance, const FVector& ItemLocation,
+	TSubclassOf<UEnvQueryContext> LineFrom, TSubclassOf<UEnvQueryContext> LineTo) const
 {
 	TArray<FVector> ContextLocationFrom;
 	if (IsContextPerItem(LineFrom))
@@ -143,7 +143,7 @@ void UEnvQueryTest_Dot::GatherLineDirections(TArray<FVector>& Directions, struct
 	}
 }
 
-void UEnvQueryTest_Dot::GatherLineDirections(TArray<FVector>& Directions, struct FEnvQueryInstance& QueryInstance, const FRotator& ItemRotation, TSubclassOf<class UEnvQueryContext> LineDirection) const
+void UEnvQueryTest_Dot::GatherLineDirections(TArray<FVector>& Directions, FEnvQueryInstance& QueryInstance, const FRotator& ItemRotation, TSubclassOf<UEnvQueryContext> LineDirection) const
 {
 	TArray<FRotator> ContextRotations;
 	if (IsContextPerItem(LineDirection))
@@ -162,8 +162,8 @@ void UEnvQueryTest_Dot::GatherLineDirections(TArray<FVector>& Directions, struct
 	}
 }
 
-void UEnvQueryTest_Dot::GatherLineDirections(TArray<FVector>& Directions, struct FEnvQueryInstance& QueryInstance,
-	TSubclassOf<class UEnvQueryContext> LineFrom, TSubclassOf<class UEnvQueryContext> LineTo, TSubclassOf<class UEnvQueryContext> LineDirection, bool bUseDirectionContext,
+void UEnvQueryTest_Dot::GatherLineDirections(TArray<FVector>& Directions, FEnvQueryInstance& QueryInstance,
+	TSubclassOf<UEnvQueryContext> LineFrom, TSubclassOf<UEnvQueryContext> LineTo, TSubclassOf<UEnvQueryContext> LineDirection, bool bUseDirectionContext,
 	const FVector& ItemLocation, const FRotator& ItemRotation) const
 {
 	if (bUseDirectionContext)
@@ -176,7 +176,7 @@ void UEnvQueryTest_Dot::GatherLineDirections(TArray<FVector>& Directions, struct
 	}
 }
 
-bool UEnvQueryTest_Dot::RequiresPerItemUpdates(TSubclassOf<class UEnvQueryContext> LineFrom, TSubclassOf<class UEnvQueryContext> LineTo, TSubclassOf<class UEnvQueryContext> LineDirection, bool bUseDirectionContext) const
+bool UEnvQueryTest_Dot::RequiresPerItemUpdates(TSubclassOf<UEnvQueryContext> LineFrom, TSubclassOf<UEnvQueryContext> LineTo, TSubclassOf<UEnvQueryContext> LineDirection, bool bUseDirectionContext) const
 {
 	bool bRequirePerItemUpdate = false;
 	if (bUseDirectionContext)
diff --git a/Engine/Source/Runtime/AIModule/Private/EnvironmentQuery/Tests/EnvQueryTest_Pathfinding.cpp b/Engine/Source/Runtime/AIModule/Private/EnvironmentQuery/Tests/EnvQueryTest_Pathfinding.cpp
--- a/Engine/Source/Runtime/AIModule/Private/EnvironmentQuery/Tests/EnvQueryTest_Pathfinding.cpp
+++ b/Engine/Source/Runtime/AIModule/Private/EnvironmentQuery/Tests/EnvQueryTest_Pathfinding.cpp
@@ -186,7 +186,7 @@ FText UEnvQueryTest_Pathfinding::GetDescriptionDetails() const
 }
 
 #if WITH_EDITOR
-void UEnvQueryTest_Pathfinding::PostEditChangeProperty(struct FPropertyChangedEvent& PropertyChangedEvent)
+void UEnvQueryTest_Pathfinding::PostEditChangeProperty(FPropertyChangedEvent& PropertyChangedEvent)
 {
 	Super::PostEditChangeProperty(PropertyChangedEvent);
 	if (PropertyChangedEvent.Property && PropertyChangedEvent.Property->GetFName() == GET_MEMBER_NAME_CHECKED(UEnvQueryTest_Pathfinding,TestMode))
diff --git a/Engine/Source/Runtime/AIModule/Private/Navigation/CrowdFollowingComponent.cpp b/Engine/Source/Runtime/AIModule/Private/Navigation/CrowdFollowingComponent.cpp
--- a/Engine/Source/Runtime/AIModule/Private/Navigation/CrowdFollowingComponent.cpp
+++ b/Engine/Source/Runtime/AIModule/Private/Navigation/CrowdFollowingComponent.cpp
@@ -384,7 +384,7 @@ void UCrowdFollowingComponent::OnLanded()
 	}
 }
 
-void UCrowdFollowingComponent::FinishUsingCustomLink(class INavLinkCustomInterface* CustomNavLink)
+void UCrowdFollowingComponent::FinishUsingCustomLink(INavLinkCustomInterface* CustomNavLink)
 {
 	const bool bPrevCustomLink = CurrentCustomLinkOb.IsValid();
 	Super::FinishUsingCustomLink(CustomNavLink);
@@ -847,7 +847,7 @@ void UCrowdFollowingComponent::GetDebugStringTokens(TArray<FString>& Tokens, TAr
 
 #if ENABLE_VISUAL_LOG
 
-void UCrowdFollowingComponent::DescribeSelfToVisLog(struct FVisualLogEntry* Snapshot) const
+void UCrowdFollowingComponent::DescribeSelfToVisLog(FVisualLogEntry* Snapshot) const
 {
 	if (!bEnableCrowdSimulation)
 	{
diff --git a/Engine/Source/Runtime/AIModule/Private/Navigation/CrowdManager.cpp b/Engine/Source/Runtime/AIModule/Private/Navigation/CrowdManager.cpp
--- a/Engine/Source/Runtime/AIModule/Private/Navigation/CrowdManager.cpp
+++ b/Engine/Source/Runtime/AIModule/Private/Navigation/CrowdManager.cpp
@@ -279,7 +279,7 @@ void UCrowdManager::Tick(float DeltaTime)
 }
 
 #if WITH_EDITOR
-void UCrowdManager::PostEditChangeProperty(struct FPropertyChangedEvent& PropertyChangedEvent)
+void UCrowdManager::PostEditChangeProperty(FPropertyChangedEvent& PropertyChangedEvent)
 {
 	Super::PostEditChangeProperty(PropertyChangedEvent);
 	
@@ -320,13 +320,13 @@ void UCrowdManager::UnregisterAgent(const ICrowdAgentInterface* Agent)
 	ActiveAgents.Remove(Agent);
 }
 
-bool UCrowdManager::IsAgentValid(const class UCrowdFollowingComponent* AgentComponent) const
+bool UCrowdManager::IsAgentValid(const UCrowdFollowingComponent* AgentComponent) const
 {
 	const FCrowdAgentData* AgentData = ActiveAgents.Find(AgentComponent);
 	return AgentData && AgentData->IsValid();
 }
 
-bool UCrowdManager::IsAgentValid(const class ICrowdAgentInterface* Agent) const
+bool UCrowdManager::IsAgentValid(const ICrowdAgentInterface* Agent) const
 {
 	const FCrowdAgentData* AgentData = ActiveAgents.Find(Agent);
 	return AgentData && AgentData->IsValid();
@@ -352,7 +352,7 @@ void UCrowdManager::UpdateAgentParams(const ICrowdAgentInterface* Agent) const
 #endif
 }
 
-void UCrowdManager::UpdateAgentState(const class ICrowdAgentInterface* Agent) const
+void UCrowdManager::UpdateAgentState(const ICrowdAgentInterface* Agent) const
 {
 	SCOPE_CYCLE_COUNTER(STAT_AI_Crowd_AgentUpdateTime);
 
@@ -365,7 +365,7 @@ void UCrowdManager::UpdateAgentState(const class ICrowdAgentInterface* Agent) co
 #endif
 }
 
-void UCrowdManager::OnAgentFinishedCustomLink(const class ICrowdAgentInterface* Agent) const
+void UCrowdManager::OnAgentFinishedCustomLink(const ICrowdAgentInterface* Agent) const
 {
 #if WITH_RECAST
 	const FCrowdAgentData* AgentData = ActiveAgents.Find(Agent);
@@ -402,7 +402,7 @@ bool UCrowdManager::SetAgentMoveTarget(const UCrowdFollowingComponent* AgentComp
 	return bSuccess;
 }
 
-bool UCrowdManager::SetAgentMoveDirection(const class UCrowdFollowingComponent* AgentComponent, const FVector& MoveDirection) const
+bool UCrowdManager::SetAgentMoveDirection(const UCrowdFollowingComponent* AgentComponent, const FVector& MoveDirection) const
 {
 	SCOPE_CYCLE_COUNTER(STAT_AI_Crowd_AgentUpdateTime);
 
@@ -422,7 +422,7 @@ bool UCrowdManager::SetAgentMoveDirection(const class UCrowdFollowingComponent*
 	return bSuccess;
 }
 
-bool UCrowdManager::SetAgentMovePath(const class UCrowdFollowingComponent* AgentComponent, const FNavMeshPath* Path, int32 PathSectionStart, int32 PathSectionEnd) const
+bool UCrowdManager::SetAgentMovePath(const UCrowdFollowingComponent* AgentComponent, const FNavMeshPath* Path, int32 PathSectionStart, int32 PathSectionEnd) const
 {
 	SCOPE_CYCLE_COUNTER(STAT_AI_Crowd_AgentUpdateTime);
 
@@ -556,7 +556,7 @@ bool UCrowdManager::SetAvoidanceConfig(int32 Idx, const FCrowdAvoidanceConfig& D
 	return true;
 }
 
-void UCrowdManager::AdjustAgentPathStart(const class UCrowdFollowingComponent* AgentComponent, const FNavMeshPath* Path, int32& PathStartIdx) const
+void UCrowdManager::AdjustAgentPathStart(const UCrowdFollowingComponent* AgentComponent, const FNavMeshPath* Path, int32& PathStartIdx) const
 {
 #if WITH_RECAST
 	const FCrowdAgentData* AgentData = ActiveAgents.Find(AgentComponent);
@@ -645,7 +645,7 @@ void UCrowdManager::RemoveAgent(const ICrowdAgentInterface* Agent, FCrowdAgentDa
 	AgentData->ClearFilter();
 }
 
-void UCrowdManager::GetAgentParams(const ICrowdAgentInterface* Agent, struct dtCrowdAgentParams* AgentParams) const
+void UCrowdManager::GetAgentParams(const ICrowdAgentInterface* Agent, dtCrowdAgentParams* AgentParams) const
 {
 	float CylRadius = 0.0f, CylHalfHeight = 0.0f;
 	Agent->GetCrowdAgentCollisions(CylRadius, CylHalfHeight);
@@ -842,7 +842,7 @@ void UCrowdManager::DestroyCrowdManager()
 	DetourCrowd = NULL;
 }
 
-void UCrowdManager::DrawDebugCorners(const struct dtCrowdAgent* CrowdAgent) const
+void UCrowdManager::DrawDebugCorners(const dtCrowdAgent* CrowdAgent) const
 {
 	{
 		FVector P0 = Recast2UnrealPoint(CrowdAgent->npos);
@@ -861,7 +861,7 @@ void UCrowdManager::DrawDebugCorners(const struct dtCrowdAgent* CrowdAgent) cons
 	}
 }
 
-void UCrowdManager::DrawDebugCollisionSegments(const struct dtCrowdAgent* CrowdAgent) const
+void UCrowdManager::DrawDebugCollisionSegments(const dtCrowdAgent* CrowdAgent) const
 {
 	FVector Center = Recast2UnrealPoint(CrowdAgent->boundary.getCenter()) + CrowdDebugDrawing::Offset;
 	DrawDebugCylinder(GetWorld(), Center - CrowdDebugDrawing::Offset, Center, CrowdAgent->params.collisionQueryRange, 32, CrowdDebugDrawing::CollisionRange);
@@ -877,7 +877,7 @@ void UCrowdManager::DrawDebugCollisionSegments(const struct dtCrowdAgent* CrowdA
 	}
 }
 
-void UCrowdManager::DrawDebugPath(const struct dtCrowdAgent* CrowdAgent) const
+void UCrowdManager::DrawDebugPath(const dtCrowdAgent* CrowdAgent) const
 {
 	ARecastNavMesh* NavMesh = Cast<ARecastNavMesh>(MyNavData);
 	if (NavMesh == NULL)
@@ -914,7 +914,7 @@ void UCrowdManager::DrawDebugPath(const struct dtCrowdAgent* CrowdAgent) const
 	NavMesh->FinishBatchQuery();
 }
 
-void UCrowdManager::DrawDebugVelocityObstacles(const struct dtCrowdAgent* CrowdAgent) const
+void UCrowdManager::DrawDebugVelocityObstacles(const dtCrowdAgent* CrowdAgent) const
 {
 	FVector Center = Recast2UnrealPoint(CrowdAgent->npos) + CrowdDebugDrawing::Offset;
 	DrawDebugCylinder(GetWorld(), Center - CrowdDebugDrawing::Offset, Center, CrowdAgent->params.maxSpeed, 32, CrowdDebugDrawing::AvoidanceRange);
@@ -950,7 +950,7 @@ void UCrowdManager::DrawDebugVelocityObstacles(const struct dtCrowdAgent* CrowdA
 	}
 }
 
-void UCrowdManager::DrawDebugPathOptimization(const struct dtCrowdAgent* CrowdAgent) const
+void UCrowdManager::DrawDebugPathOptimization(const dtCrowdAgent* CrowdAgent) const
 {
 	FVector Pt0 = Recast2UnrealPoint(DetourAgentDebug->optStart) + CrowdDebugDrawing::Offset * 1.25f;
 	FVector Pt1 = Recast2UnrealPoint(DetourAgentDebug->optEnd) + CrowdDebugDrawing::Offset * 1.25f;
@@ -958,7 +958,7 @@ void UCrowdManager::DrawDebugPathOptimization(const struct dtCrowdAgent* CrowdAg
 	DrawDebugLine(GetWorld(), Pt0, Pt1, CrowdDebugDrawing::PathOpt, false, -1.0f, SDPG_World, 2.5f);
 }
 
-void UCrowdManager::DrawDebugNeighbors(const struct dtCrowdAgent* CrowdAgent) const
+void UCrowdManager::DrawDebugNeighbors(const dtCrowdAgent* CrowdAgent) const
 {
 	FVector Center = Recast2UnrealPoint(CrowdAgent->npos) + CrowdDebugDrawing::Offset;
 	DrawDebugCylinder(GetWorld(), Center - CrowdDebugDrawing::Offset, Center, CrowdAgent->params.collisionQueryRange, 32, CrowdDebugDrawing::CollisionRange);
@@ -996,7 +996,7 @@ void UCrowdManager::DebugTick() const
 	}
 
 	// on screen debugging
-	const struct dtCrowdAgent* SelectedAgent = DetourAgentDebug->idx >= 0 ? DetourCrowd->getAgent(DetourAgentDebug->idx) : NULL;
+	const dtCrowdAgent* SelectedAgent = DetourAgentDebug->idx >= 0 ? DetourCrowd->getAgent(DetourAgentDebug->idx) : NULL;
 	if (SelectedAgent && CrowdDebugDrawing::bDebugSelectedActors)
 	{
 		if (CrowdDebugDrawing::bDrawDebugCorners)
@@ -1040,7 +1040,7 @@ void UCrowdManager::DebugTick() const
 			const AActor* LogOwner = AgentOb ? Cast<const AActor>(AgentOb->GetOuter()) : NULL;
 
 			const FCrowdAgentData& AgentData = It.Value();
-			const struct dtCrowdAgent* CrowdAgent = AgentData.IsValid() ? DetourCrowd->getAgent(AgentData.AgentIndex) : NULL;
+			const dtCrowdAgent* CrowdAgent = AgentData.IsValid() ? DetourCrowd->getAgent(AgentData.AgentIndex) : NULL;
 
 			if (CrowdAgent && LogOwner)
 			{
@@ -1132,19 +1132,19 @@ void UCrowdManager::PostProximityUpdate()
 	// empty in base class
 }
 
-class UWorld* UCrowdManager::GetWorld() const
+UWorld* UCrowdManager::GetWorld() const
 {
 	UNavigationSystem* NavSys = Cast<UNavigationSystem>(GetOuter());
 	return NavSys ? NavSys->GetWorld() : NULL;
 }
 
-UCrowdManager* UCrowdManager::GetCurrent(class UObject* WorldContextObject)
+UCrowdManager* UCrowdManager::GetCurrent(UObject* WorldContextObject)
 {
 	UNavigationSystem* NavSys = UNavigationSystem::GetCurrent(WorldContextObject);
 	return NavSys ? NavSys->GetCrowdManager() : NULL;
 }
 
-UCrowdManager* UCrowdManager::GetCurrent(class UWorld* World)
+UCrowdManager* UCrowdManager::GetCurrent(UWorld* World)
 {
 	UNavigationSystem* NavSys = UNavigationSystem::GetCurrent(World);
 	return NavSys ? NavSys->GetCrowdManager() : NULL;
diff --git a/Engine/Source/Runtime/AIModule/Private/Navigation/NavigationComponent.cpp b/Engine/Source/Runtime/AIModule/Private/Navigation/NavigationComponent.cpp
--- a/Engine/Source/Runtime/AIModule/Private/Navigation/NavigationComponent.cpp
+++ b/Engine/Source/Runtime/AIModule/Private/Navigation/NavigationComponent.cpp
@@ -870,7 +870,7 @@ void UNavigationComponent::CacheNavQueryExtent() const
 	}
 }
 
-void UNavigationComponent::OnNavDataRegistered(class ANavigationData* NavData)
+void UNavigationComponent::OnNavDataRegistered(ANavigationData* NavData)
 {
 	// re-pick regardless of whether we have MyNavData set already - the new one can be better!
 	PickNavData();
@@ -902,7 +902,7 @@ void UNavigationComponent::SetReceiveSmartLinkUpdates(bool bEnabled)
 	bUpdateForSmartLinks = bEnabled;
 }
 
-void UNavigationComponent::OnCustomLinkBroadcast(class UNavLinkCustomComponent* NearbyLink)
+void UNavigationComponent::OnCustomLinkBroadcast(UNavLinkCustomComponent* NearbyLink)
 {
 	// update only when agent is actually moving
 	if (NearbyLink == NULL || !Path.IsValid() ||
diff --git a/Engine/Source/Runtime/AIModule/Private/Navigation/PathFollowingComponent.cpp b/Engine/Source/Runtime/AIModule/Private/Navigation/PathFollowingComponent.cpp
--- a/Engine/Source/Runtime/AIModule/Private/Navigation/PathFollowingComponent.cpp
+++ b/Engine/Source/Runtime/AIModule/Private/Navigation/PathFollowingComponent.cpp
@@ -77,7 +77,7 @@ void LogPathHelper(AActor* LogOwner, FNavPathSharedPtr Path, const AActor* GoalA
 #endif // ENABLE_VISUAL_LOG
 }
 
-void LogBlockHelper(AActor* LogOwner, class UNavMovementComponent* MoveComp, float RadiusPct, float HeightPct, const FVector& SegmentStart, const FVector& SegmentEnd)
+void LogBlockHelper(AActor* LogOwner, UNavMovementComponent* MoveComp, float RadiusPct, float HeightPct, const FVector& SegmentStart, const FVector& SegmentEnd)
 {
 #if ENABLE_VISUAL_LOG
 	if (MoveComp && LogOwner)
@@ -1004,7 +1004,7 @@ void UPathFollowingComponent::DebugReachTest(float& CurrentDot, float& CurrentDi
 	bHeightFailed = (CurrentHeight > UseHeight) ? 1 : 0;
 }
 
-void UPathFollowingComponent::StartUsingCustomLink(class INavLinkCustomInterface* CustomNavLink, const FVector& DestPoint)
+void UPathFollowingComponent::StartUsingCustomLink(INavLinkCustomInterface* CustomNavLink, const FVector& DestPoint)
 {
 	INavLinkCustomInterface* PrevNavLink = Cast<INavLinkCustomInterface>(CurrentCustomLinkOb.Get());
 	if (PrevNavLink)
@@ -1031,7 +1031,7 @@ void UPathFollowingComponent::StartUsingCustomLink(class INavLinkCustomInterface
 	}
 }
 
-void UPathFollowingComponent::FinishUsingCustomLink(class INavLinkCustomInterface* CustomNavLink)
+void UPathFollowingComponent::FinishUsingCustomLink(INavLinkCustomInterface* CustomNavLink)
 {
 	UObject* NavLinkOb = Cast<UObject>(CustomNavLink);
 	if (CurrentCustomLinkOb == NavLinkOb)
@@ -1286,7 +1286,7 @@ FString UPathFollowingComponent::GetResultDesc(EPathFollowingResult::Type Result
 	return TEXT("Unknown");
 }
 
-void UPathFollowingComponent::DisplayDebug(class UCanvas* Canvas, const FDebugDisplayInfo& DebugDisplay, float& YL, float& YPos) const
+void UPathFollowingComponent::DisplayDebug(UCanvas* Canvas, const FDebugDisplayInfo& DebugDisplay, float& YL, float& YPos) const
 {
 	Canvas->SetDrawColor(FColor::Blue);
 	UFont* RenderFont = GEngine->GetSmallFont();
@@ -1306,7 +1306,7 @@ void UPathFollowingComponent::DisplayDebug(class UCanvas* Canvas, const FDebugDi
 }
 
 #if ENABLE_VISUAL_LOG
-void UPathFollowingComponent::DescribeSelfToVisLog(struct FVisualLogEntry* Snapshot) const
+void UPathFollowingComponent::DescribeSelfToVisLog(FVisualLogEntry* Snapshot) const
 {
 	FVisualLogEntry::FStatusCategory Category;
 	Category.Category = TEXT("Path following");
diff --git a/Engine/Source/Runtime/AIModule/Private/Perception/AIPerceptionSystem.cpp b/Engine/Source/Runtime/AIModule/Private/Perception/AIPerceptionSystem.cpp
--- a/Engine/Source/Runtime/AIModule/Private/Perception/AIPerceptionSystem.cpp
+++ b/Engine/Source/Runtime/AIModule/Private/Perception/AIPerceptionSystem.cpp
@@ -181,7 +181,7 @@ void UAIPerceptionSystem::AgeStimuli()
 	}
 }
 
-UAIPerceptionSystem* UAIPerceptionSystem::GetCurrent(class UObject* WorldContextObject)
+UAIPerceptionSystem* UAIPerceptionSystem::GetCurrent(UObject* WorldContextObject)
 {
 	UWorld* World = Cast<UWorld>(WorldContextObject);
 
diff --git a/Engine/Source/Runtime/AIModule/Private/Perception/AIPerceptionTypes.cpp b/Engine/Source/Runtime/AIModule/Private/Perception/AIPerceptionTypes.cpp
--- a/Engine/Source/Runtime/AIModule/Private/Perception/AIPerceptionTypes.cpp
+++ b/Engine/Source/Runtime/AIModule/Private/Perception/AIPerceptionTypes.cpp
@@ -11,7 +11,7 @@ const float FAIStimulus::NeverHappenedAge = FLT_MAX;
 //----------------------------------------------------------------------//
 const FPerceptionListener FPerceptionListener::NullListener((UAIPerceptionComponent*)NULL);
 
-FPerceptionListener::FPerceptionListener(class UAIPerceptionComponent* InListener) 
+FPerceptionListener::FPerceptionListener(UAIPerceptionComponent* InListener) 
 	: Listener(InListener)
 	, PeripheralVisionAngleCos(1.f)
 	, HearingRangeSq(-1.f)
diff --git a/Engine/Source/Runtime/AIModule/Private/Perception/PawnSensingComponent.cpp b/Engine/Source/Runtime/AIModule/Private/Perception/PawnSensingComponent.cpp
--- a/Engine/Source/Runtime/AIModule/Private/Perception/PawnSensingComponent.cpp
+++ b/Engine/Source/Runtime/AIModule/Private/Perception/PawnSensingComponent.cpp
@@ -179,7 +179,7 @@ bool UPawnSensingComponent::IsSensorActor(const AActor* Actor) const
 	return (Actor == GetSensorActor());
 }
 
-bool UPawnSensingComponent::HasLineOfSightTo(const class AActor* Other) const
+bool UPawnSensingComponent::HasLineOfSightTo(const AActor* Other) const
 {
 	AController* SensorController = GetSensorController();
 	if (SensorController == NULL)
diff --git a/Engine/Source/Runtime/AIModule/Private/VisualLoggerExtension.cpp b/Engine/Source/Runtime/AIModule/Private/VisualLoggerExtension.cpp
--- a/Engine/Source/Runtime/AIModule/Private/VisualLoggerExtension.cpp
+++ b/Engine/Source/Runtime/AIModule/Private/VisualLoggerExtension.cpp
@@ -15,7 +15,7 @@ FVisualLoggerExtension::FVisualLoggerExtension()
 
 }
 
-void FVisualLoggerExtension::OnTimestampChange(float Timestamp, class UWorld* InWorld, class AActor* HelperActor )
+void FVisualLoggerExtension::OnTimestampChange(float Timestamp, UWorld* InWorld, AActor* HelperActor )
 {
 #if USE_EQS_DEBUGGER
 	if (CurrentTimestamp != Timestamp)
@@ -27,7 +27,7 @@ void FVisualLoggerExtension::OnTimestampChange(float Timestamp, class UWorld* In
 #endif
 }
 
-void FVisualLoggerExtension::DisableDrawingForData(class UWorld* InWorld, class UCanvas* Canvas, class AActor* HelperActor, const FName& TagName, const FVisualLogEntry::FDataBlock& DataBlock, float Timestamp)
+void FVisualLoggerExtension::DisableDrawingForData(UWorld* InWorld, UCanvas* Canvas, AActor* HelperActor, const FName& TagName, const FVisualLogEntry::FDataBlock& DataBlock, float Timestamp)
 {
 #if USE_EQS_DEBUGGER
 	if (TagName == *EVisLogTags::TAG_EQS && CurrentTimestamp == Timestamp)
@@ -37,7 +37,7 @@ void FVisualLoggerExtension::DisableDrawingForData(class UWorld* InWorld, class
 #endif
 }
 
-void FVisualLoggerExtension::DisableEQSRendering(class AActor* HelperActor)
+void FVisualLoggerExtension::DisableEQSRendering(AActor* HelperActor)
 {
 #if USE_EQS_DEBUGGER
 	if (HelperActor)
@@ -55,7 +55,7 @@ void FVisualLoggerExtension::DisableEQSRendering(class AActor* HelperActor)
 #endif
 }
 
-void FVisualLoggerExtension::LogEntryLineSelectionChanged(TSharedPtr<struct FLogEntryItem> SelectedItem, int64 UserData, FName TagName)
+void FVisualLoggerExtension::LogEntryLineSelectionChanged(TSharedPtr<FLogEntryItem> SelectedItem, int64 UserData, FName TagName)
 {
 	if (TagName == *EVisLogTags::TAG_EQS)
 	{
@@ -63,9 +63,9 @@ void FVisualLoggerExtension::LogEntryLineSelectionChanged(TSharedPtr<struct FLog
 	}
 }
 
-extern RENDERCORE_API class FTexture* GWhiteTexture;
+extern RENDERCORE_API FTexture* GWhiteTexture;
 
-void FVisualLoggerExtension::DrawData(class UWorld* InWorld, class UCanvas* Canvas, class AActor* HelperActor, const FName& TagName, const FVisualLogEntry::FDataBlock& DataBlock, float Timestamp)
+void FVisualLoggerExtension::DrawData(UWorld* InWorld, UCanvas* Canvas, AActor* HelperActor, const FName& TagName, const FVisualLogEntry::FDataBlock& DataBlock, float Timestamp)
 {
 #if USE_EQS_DEBUGGER
 	if (TagName == *EVisLogTags::TAG_EQS && HelperActor)
