Renamed AISense classes to AISenseImplementation #UE4

diff --git a/Engine/Source/Runtime/AIModule/Classes/Perception/AIPerceptionSystem.h b/Engine/Source/Runtime/AIModule/Classes/Perception/AIPerceptionSystem.h
--- a/Engine/Source/Runtime/AIModule/Classes/Perception/AIPerceptionSystem.h
+++ b/Engine/Source/Runtime/AIModule/Classes/Perception/AIPerceptionSystem.h
@@ -7,7 +7,7 @@
 
 DECLARE_LOG_CATEGORY_EXTERN(LogAIPerception, Warning, All);
 
-class UAISense;
+class UAISenseImplementation;
 class UAIPerceptionComponent;
 
 /**
@@ -22,7 +22,7 @@ protected:
 	AIPerception::FListenerMap ListenerContainer;
 
 	UPROPERTY()
-	TArray<UAISense*> Senses;
+	TArray<UAISenseImplementation*> Senses;
 
 	UPROPERTY(config)
 	float PerceptionAgingRate;
@@ -105,7 +105,7 @@ protected:
 
 	void AgeStimuli();
 
-	friend class UAISense;
+	friend class UAISenseImplementation;
 	FORCEINLINE AIPerception::FListenerMap& GetListenersMap() { return ListenerContainer; }
 
 private:
diff --git a/Engine/Source/Runtime/AIModule/Classes/Perception/AISense.h b/Engine/Source/Runtime/AIModule/Classes/Perception/AISense.h
deleted file mode 100644
diff --git a/Engine/Source/Runtime/AIModule/Classes/Perception/AISenseImplementation.h b/Engine/Source/Runtime/AIModule/Classes/Perception/AISenseImplementation.h
new file mode 100644
--- /dev/null
+++ b/Engine/Source/Runtime/AIModule/Classes/Perception/AISenseImplementation.h
@@ -0,0 +1,77 @@
+// Copyright 1998-2014 Epic Games, Inc. All Rights Reserved.
+
+#pragma once
+
+#include "AIPerceptionTypes.h"
+#include "AISenseImplementation.generated.h"
+
+DECLARE_DELEGATE_OneParam(FOnPerceptionListenerUpdateDelegate, const FPerceptionListener&);
+
+UCLASS(ClassGroup=AI, abstract)
+class AIMODULE_API UAISenseImplementation : public UObject
+{
+	GENERATED_UCLASS_BODY()
+
+	static const float SuspendNextUpdate;
+
+private:
+	UPROPERTY()
+	class UAIPerceptionSystem* PerceptionSystemInstance;
+
+	/** then this count reaches 0 sense will be updated */
+	float TimeUntilNextUpdate;
+
+protected:
+	/**	If bound will be called when new FPerceptionListener gets registers with AIPerceptionSystem */
+	FOnPerceptionListenerUpdateDelegate OnNewListenerDelegate;
+
+	/**	If bound will be called when a FPerceptionListener's in AIPerceptionSystem change */
+	FOnPerceptionListenerUpdateDelegate OnListenerUpdateDelegate;
+
+	/**	If bound will be called when a FPerceptionListener's in removed from AIPerceptionSystem */
+	FOnPerceptionListenerUpdateDelegate OnListenerRemovedDelegate;
+			
+public:
+	static FAISenseId GetSenseIndex() { check(0 && "This should never get called!"); return AIPerception::InvalidSenseIndex; }
+
+	virtual void PostInitProperties() override;
+
+	/** 
+	 *	@return should this sense be ticked now
+	 */
+	bool ProgressTime(float DeltaSeconds)
+	{
+		TimeUntilNextUpdate -= DeltaSeconds;
+		return TimeUntilNextUpdate <= 0.f;
+	}
+
+	void Tick()
+	{
+		if (TimeUntilNextUpdate <= 0.f)
+		{
+			TimeUntilNextUpdate = Update();
+		}
+	}
+
+	virtual void RegisterSource(AActor& SourceActor) {}
+
+	FORCEINLINE void OnNewListener(const FPerceptionListener& NewListener) { OnNewListenerDelegate.ExecuteIfBound(NewListener); }
+	FORCEINLINE void OnListenerUpdate(const FPerceptionListener& NewListener) { OnListenerUpdateDelegate.ExecuteIfBound(NewListener); }
+	FORCEINLINE void OnListenerRemoved(const FPerceptionListener& NewListener) { OnListenerRemovedDelegate.ExecuteIfBound(NewListener); }
+
+protected:
+	/** @return time until next update */
+	virtual float Update() { return FLT_MAX; }
+
+	/** will result in updating as soon as possible */
+	FORCEINLINE void RequestImmediateUpdate() { TimeUntilNextUpdate = 0.f; }
+
+	/** will result in updating in specified number of seconds */
+	FORCEINLINE void RequestUpdateInSeconds(float UpdateInSeconds) { TimeUntilNextUpdate = UpdateInSeconds; }
+
+	FORCEINLINE UAIPerceptionSystem* GetPerceptionSystem() { return PerceptionSystemInstance; }
+
+	/** returning pointer rather then a reference to prevent users from
+	 *	accidentally creating copies by creating non-reference local vars */
+	AIPerception::FListenerMap* GetListeners();
+};
diff --git a/Engine/Source/Runtime/AIModule/Classes/Perception/AISenseImplementation_Damage.h b/Engine/Source/Runtime/AIModule/Classes/Perception/AISenseImplementation_Damage.h
new file mode 100644
--- /dev/null
+++ b/Engine/Source/Runtime/AIModule/Classes/Perception/AISenseImplementation_Damage.h
@@ -0,0 +1,53 @@
+// Copyright 1998-2014 Epic Games, Inc. All Rights Reserved.
+
+#pragma once
+
+#include "Perception/AIPerceptionSystem.h"
+#include "Perception/AISenseImplementation.h"
+#include "AISenseImplementation_Damage.generated.h"
+
+USTRUCT()
+struct AIMODULE_API FAIDamageEvent
+{	
+	GENERATED_USTRUCT_BODY()
+
+	typedef class UAISenseImplementation_Damage FSenseClass;
+
+	float Amount;
+	FVector Location;
+	FVector HitLocation;
+	
+	UPROPERTY()
+	AActor* DamagedActor;
+
+	UPROPERTY()
+	AActor* Instigator;
+	
+	FAIDamageEvent(){}
+
+	FAIDamageEvent(class AActor* InDamagedActor, class AActor* InInstigator, float DamageAmount, const FVector& EventLocation, const FVector& InHitLocation = FAISystem::InvalidLocation)
+		: Amount(DamageAmount), Location(EventLocation), HitLocation(InHitLocation), DamagedActor(InDamagedActor), Instigator(InInstigator)
+	{
+		if (FAISystem::IsValidLocation(InHitLocation))
+		{
+			HitLocation = EventLocation;
+		}
+	}
+};
+
+UCLASS(ClassGroup=AI)
+class AIMODULE_API UAISenseImplementation_Damage : public UAISenseImplementation
+{
+	GENERATED_UCLASS_BODY()
+
+	UPROPERTY()
+	TArray<FAIDamageEvent> RegisteredEvents;
+
+public:
+	FORCEINLINE static FAISenseId GetSenseIndex() { return FAISenseId(ECorePerceptionTypes::Damage); }
+		
+	void RegisterEvent(const FAIDamageEvent& Event);	
+
+protected:
+	virtual float Update() override;
+};
diff --git a/Engine/Source/Runtime/AIModule/Classes/Perception/AISenseImplementation_Hearing.h b/Engine/Source/Runtime/AIModule/Classes/Perception/AISenseImplementation_Hearing.h
new file mode 100644
--- /dev/null
+++ b/Engine/Source/Runtime/AIModule/Classes/Perception/AISenseImplementation_Hearing.h
@@ -0,0 +1,44 @@
+// Copyright 1998-2014 Epic Games, Inc. All Rights Reserved.
+
+#pragma once
+
+#include "Perception/AIPerceptionSystem.h"
+#include "Perception/AISenseImplementation.h"
+#include "AISenseImplementation_Hearing.generated.h"
+
+USTRUCT()
+struct AIMODULE_API FAINoiseEvent
+{	
+	GENERATED_USTRUCT_BODY()
+
+	typedef class UAISenseImplementation_Hearing FSenseClass;
+
+	float Age;
+	FVector NoiseLocation;
+	float Loudness;
+	FGenericTeamId TeamIdentifier;
+	
+	UPROPERTY()
+	AActor* Instigator;
+	
+	FAINoiseEvent(){}
+	FAINoiseEvent(class AActor* InInstigator, const FVector& InNoiseLocation, float InLoudness = 1.f);
+};
+
+UCLASS(ClassGroup=AI)
+class AIMODULE_API UAISenseImplementation_Hearing : public UAISenseImplementation
+{
+	GENERATED_UCLASS_BODY()
+		
+protected:
+	UPROPERTY()
+	TArray<FAINoiseEvent> NoiseEvents;
+
+public:
+	FORCEINLINE static FAISenseId GetSenseIndex() { return FAISenseId(ECorePerceptionTypes::Hearing); }
+	
+	void RegisterEvent(const FAINoiseEvent& Event);	
+
+protected:
+	virtual float Update() override;
+};
diff --git a/Engine/Source/Runtime/AIModule/Classes/Perception/AISenseImplementation_Prediction.h b/Engine/Source/Runtime/AIModule/Classes/Perception/AISenseImplementation_Prediction.h
new file mode 100644
--- /dev/null
+++ b/Engine/Source/Runtime/AIModule/Classes/Perception/AISenseImplementation_Prediction.h
@@ -0,0 +1,47 @@
+// Copyright 1998-2014 Epic Games, Inc. All Rights Reserved.
+
+#pragma once
+
+#include "Perception/AIPerceptionSystem.h"
+#include "Perception/AISenseImplementation.h"
+#include "AISenseImplementation_Prediction.generated.h"
+
+USTRUCT()
+struct AIMODULE_API FAIPredictionEvent
+{	
+	GENERATED_USTRUCT_BODY()
+
+	typedef class UAISenseImplementation_Prediction FSenseClass;
+	
+	UPROPERTY()
+	class AActor* Requestor;
+
+	UPROPERTY()
+	class AActor* PredictedActor;
+
+	float TimeToPredict;
+		
+	FAIPredictionEvent(){}
+	
+	FAIPredictionEvent(class AActor* InRequestor, class AActor* InPredictedActor, float PredictionTime)
+		: Requestor(InRequestor), PredictedActor(InPredictedActor), TimeToPredict(PredictionTime)
+	{
+	}
+};
+
+UCLASS(ClassGroup=AI)
+class AIMODULE_API UAISenseImplementation_Prediction : public UAISenseImplementation
+{
+	GENERATED_UCLASS_BODY()
+
+	UPROPERTY()
+	TArray<FAIPredictionEvent> RegisteredEvents;
+
+public:
+	FORCEINLINE static FAISenseId GetSenseIndex() { return FAISenseId(ECorePerceptionTypes::Prediction); }
+		
+	void RegisterEvent(const FAIPredictionEvent& Event);	
+
+protected:
+	virtual float Update() override;
+};
diff --git a/Engine/Source/Runtime/AIModule/Classes/Perception/AISenseImplementation_Sight.h b/Engine/Source/Runtime/AIModule/Classes/Perception/AISenseImplementation_Sight.h
new file mode 100644
--- /dev/null
+++ b/Engine/Source/Runtime/AIModule/Classes/Perception/AISenseImplementation_Sight.h
@@ -0,0 +1,154 @@
+// Copyright 1998-2014 Epic Games, Inc. All Rights Reserved.
+
+#pragma once
+
+#include "Perception/AIPerceptionSystem.h"
+#include "Perception/AISenseImplementation.h"
+#include "AISenseImplementation_Sight.generated.h"
+
+class IAISightTargetInterface;
+
+namespace ESightPerceptionEventName
+{
+	enum Type
+	{
+		Undefined,
+		GainedSight,
+		LostSight
+	};
+}
+
+USTRUCT()
+struct AIMODULE_API FAISightEvent
+{
+	GENERATED_USTRUCT_BODY()
+
+	typedef class UAISenseImplementation_Sight FSenseClass;
+
+	float Age;
+	ESightPerceptionEventName::Type EventType;	
+
+	UPROPERTY()
+	class AActor* SeenActor;
+
+	UPROPERTY()
+	class AActor* Observer;
+
+	FAISightEvent(){}
+
+	FAISightEvent(class AActor* InSeenActor, class AActor* InObserver, ESightPerceptionEventName::Type InEventType)
+		: Age(0.f), EventType(InEventType), SeenActor(InSeenActor), Observer(InObserver)
+	{
+	}
+};
+
+struct FAISightTarget
+{
+	typedef FName FTargetId;
+	static const FTargetId InvalidTargetId;
+
+	TWeakObjectPtr<AActor> Target;
+	IAISightTargetInterface* SightTargetInterface;
+	FGenericTeamId TeamId;
+	FTargetId TargetId;
+
+	FAISightTarget(class AActor* InTarget = NULL, FGenericTeamId InTeamId = FGenericTeamId::NoTeam);
+
+	FORCEINLINE FVector GetLocationSimple() const
+	{
+		return Target.IsValid() ? Target->GetActorLocation() : FVector::ZeroVector;
+	}
+
+	FORCEINLINE const AActor* GetTargetActor() const { return Target.Get(); }
+};
+
+struct FAISightQuery
+{
+	uint32 ObserverId;
+	FAISightTarget::FTargetId TargetId;
+
+	float Age;
+	float Score;
+	float Importance;
+
+	uint32 bLastResult : 1;
+
+	FAISightQuery(uint32 ListenerId = AIPerception::InvalidListenerId, FAISightTarget::FTargetId Target = FAISightTarget::InvalidTargetId)
+		: ObserverId(ListenerId), TargetId(Target), Age(0), Score(0), Importance(0), bLastResult(false)
+	{
+	}
+
+	void RecalcScore()
+	{
+		Score = Age + Importance;
+	}
+
+	class FSortPredicate
+	{
+	public:
+		FSortPredicate()
+		{}
+
+		bool operator()(const FAISightQuery& A, const FAISightQuery& B) const
+		{
+			return A.Score > B.Score;
+		}
+	};
+};
+
+
+UCLASS(ClassGroup=AI, config=Game)
+class AIMODULE_API UAISenseImplementation_Sight : public UAISenseImplementation
+{
+	GENERATED_UCLASS_BODY()
+
+	TMap<FAISightTarget::FTargetId, FAISightTarget> ObservedTargets;
+
+	TArray<FAISightQuery> SightQueryQueue;
+
+protected:
+	UPROPERTY(config)
+	int32 MaxTracesPerTick;
+
+	UPROPERTY(config)
+	float HighImportanceQueryDistanceThreshold;
+
+	float HighImportanceDistanceSquare;
+
+	UPROPERTY(config)
+	float MaxQueryImportance;
+
+	UPROPERTY(config)
+	float SightLimitQueryImportance;
+
+public:
+
+	virtual void PostInitProperties() override;
+
+	FORCEINLINE static FAISenseId GetSenseIndex() { return FAISenseId(ECorePerceptionTypes::Sight); }
+
+	void RegisterEvent(const FAISightEvent& Event);	
+
+	virtual void RegisterSource(AActor& SourceActor) override;
+	
+protected:
+	virtual float Update() override;
+
+	void OnNewListenerImpl(const FPerceptionListener& NewListener);
+	void OnListenerUpdateImpl(const FPerceptionListener& UpdatedListener);
+	void OnListenerRemovedImpl(const FPerceptionListener& UpdatedListener);	
+
+	enum FQueriesOperationPostProcess
+	{
+		DontSort,
+		Sort
+	};
+	void RemoveAllQueriesByListener(const FPerceptionListener& Listener, FQueriesOperationPostProcess PostProcess);
+	void RemoveAllQueriesToTarget(const FName& TargetId, FQueriesOperationPostProcess PostProcess);
+
+	void RegisterTarget(AActor& TargetActor, FQueriesOperationPostProcess PostProcess);
+
+	FORCEINLINE void SortQueries() { SightQueryQueue.Sort(FAISightQuery::FSortPredicate()); }
+
+	float CalcQueryImportance(const FPerceptionListener& Listener, const FVector& TargetLocation, const float SightRadiusSq) const;
+};
diff --git a/Engine/Source/Runtime/AIModule/Classes/Perception/AISenseImplementation_Team.h b/Engine/Source/Runtime/AIModule/Classes/Perception/AISenseImplementation_Team.h
new file mode 100644
--- /dev/null
+++ b/Engine/Source/Runtime/AIModule/Classes/Perception/AISenseImplementation_Team.h
@@ -0,0 +1,59 @@
+// Copyright 1998-2014 Epic Games, Inc. All Rights Reserved.
+
+#pragma once
+
+#include "Perception/AIPerceptionSystem.h"
+#include "Perception/AISenseImplementation.h"
+#include "AISenseImplementation_Team.generated.h"
+
+USTRUCT()
+struct AIMODULE_API FAITeamStimulusEvent
+{	
+	GENERATED_USTRUCT_BODY()
+
+	typedef class UAISenseImplementation_Team FSenseClass;
+
+	FVector LastKnowLocation;
+private:
+	FVector BroadcastLocation;
+public:
+	float RangeSq;
+	float InformationAge;
+	FGenericTeamId TeamIdentifier;
+private:
+	UPROPERTY()
+	class AActor* Broadcaster;
+public:
+	UPROPERTY()
+	class AActor* Enemy;
+		
+	FAITeamStimulusEvent(){}	
+	FAITeamStimulusEvent(class AActor* InBroadcaster, class AActor* InEnemy, const FVector& InLastKnowLocation, float EventRange, float PassedInfoAge = 0.f);
+
+	FORCEINLINE void CacheBroadcastLocation()
+	{
+		BroadcastLocation = Broadcaster ? Broadcaster->GetActorLocation() : FAISystem::InvalidLocation;
+	}
+
+	FORCEINLINE const FVector& GetBroadcastLocation() const 
+	{
+		return BroadcastLocation;
+	}
+};
+
+UCLASS(ClassGroup=AI)
+class AIMODULE_API UAISenseImplementation_Team : public UAISenseImplementation
+{
+	GENERATED_UCLASS_BODY()
+
+	UPROPERTY()
+	TArray<FAITeamStimulusEvent> RegisteredEvents;
+
+public:
+	FORCEINLINE static FAISenseId GetSenseIndex() { return FAISenseId(ECorePerceptionTypes::Team); }
+		
+	void RegisterEvent(const FAITeamStimulusEvent& Event);	
+
+protected:
+	virtual float Update() override;
+};
diff --git a/Engine/Source/Runtime/AIModule/Classes/Perception/AISenseImplementation_Touch.h b/Engine/Source/Runtime/AIModule/Classes/Perception/AISenseImplementation_Touch.h
new file mode 100644
--- /dev/null
+++ b/Engine/Source/Runtime/AIModule/Classes/Perception/AISenseImplementation_Touch.h
@@ -0,0 +1,47 @@
+// Copyright 1998-2014 Epic Games, Inc. All Rights Reserved.
+
+#pragma once
+
+#include "Perception/AIPerceptionSystem.h"
+#include "Perception/AISenseImplementation.h"
+#include "AISenseImplementation_Touch.generated.h"
+
+USTRUCT()
+struct AIMODULE_API FAITouchEvent
+{	
+	GENERATED_USTRUCT_BODY()
+
+	typedef class UAISenseImplementation_Touch FSenseClass;
+
+	FVector Location;
+	
+	UPROPERTY()
+	class AActor* TouchReceiver;
+
+	UPROPERTY()
+	class AActor* OtherActor;
+		
+	FAITouchEvent(){}
+	
+	FAITouchEvent(class AActor* InTouchReceiver, class AActor* InOtherActor, const FVector& EventLocation)
+		: Location(EventLocation), TouchReceiver(InTouchReceiver), OtherActor(InOtherActor)
+	{
+	}
+};
+
+UCLASS(ClassGroup=AI)
+class AIMODULE_API UAISenseImplementation_Touch : public UAISenseImplementation
+{
+	GENERATED_UCLASS_BODY()
+
+	UPROPERTY()
+	TArray<FAITouchEvent> RegisteredEvents;
+
+public:
+	FORCEINLINE static FAISenseId GetSenseIndex() { return FAISenseId(ECorePerceptionTypes::Touch); }
+		
+	void RegisterEvent(const FAITouchEvent& Event);	
+
+protected:
+	virtual float Update() override;
+};
diff --git a/Engine/Source/Runtime/AIModule/Classes/Perception/AISense_Damage.h b/Engine/Source/Runtime/AIModule/Classes/Perception/AISense_Damage.h
deleted file mode 100644
diff --git a/Engine/Source/Runtime/AIModule/Classes/Perception/AISense_Hearing.h b/Engine/Source/Runtime/AIModule/Classes/Perception/AISense_Hearing.h
deleted file mode 100644
diff --git a/Engine/Source/Runtime/AIModule/Classes/Perception/AISense_Prediction.h b/Engine/Source/Runtime/AIModule/Classes/Perception/AISense_Prediction.h
deleted file mode 100644
diff --git a/Engine/Source/Runtime/AIModule/Classes/Perception/AISense_Sight.h b/Engine/Source/Runtime/AIModule/Classes/Perception/AISense_Sight.h
deleted file mode 100644
diff --git a/Engine/Source/Runtime/AIModule/Classes/Perception/AISense_Team.h b/Engine/Source/Runtime/AIModule/Classes/Perception/AISense_Team.h
deleted file mode 100644
diff --git a/Engine/Source/Runtime/AIModule/Classes/Perception/AISense_Touch.h b/Engine/Source/Runtime/AIModule/Classes/Perception/AISense_Touch.h
deleted file mode 100644
diff --git a/Engine/Source/Runtime/AIModule/Private/AIModulePrivate.h b/Engine/Source/Runtime/AIModule/Private/AIModulePrivate.h
--- a/Engine/Source/Runtime/AIModule/Private/AIModulePrivate.h
+++ b/Engine/Source/Runtime/AIModule/Private/AIModulePrivate.h
@@ -56,7 +56,7 @@
 
 #include "GenericTeamAgentInterface.h"
 
-#include "Perception/AISense.h"
+#include "Perception/AISenseImplementation.h"
 #include "Perception/AIPerceptionComponent.h"
 
 #include "VisualLog.h"
diff --git a/Engine/Source/Runtime/AIModule/Private/Perception/AIPerceptionSystem.cpp b/Engine/Source/Runtime/AIModule/Private/Perception/AIPerceptionSystem.cpp
--- a/Engine/Source/Runtime/AIModule/Private/Perception/AIPerceptionSystem.cpp
+++ b/Engine/Source/Runtime/AIModule/Private/Perception/AIPerceptionSystem.cpp
@@ -2,12 +2,12 @@
 
 #include "AIModulePrivate.h"
 #include "Perception/AIPerceptionSystem.h"
-#include "Perception/AISense_Hearing.h"
-#include "Perception/AISense_Sight.h"
-#include "Perception/AISense_Team.h"
-#include "Perception/AISense_Touch.h"
-#include "Perception/AISense_Prediction.h"
-#include "Perception/AISense_Damage.h"
+#include "Perception/AISenseImplementation_Hearing.h"
+#include "Perception/AISenseImplementation_Sight.h"
+#include "Perception/AISenseImplementation_Team.h"
+#include "Perception/AISenseImplementation_Touch.h"
+#include "Perception/AISenseImplementation_Prediction.h"
+#include "Perception/AISenseImplementation_Damage.h"
 
 DECLARE_CYCLE_STAT(TEXT("Perception System"),STAT_AI_PerceptionSys,STATGROUP_AI);
 
@@ -25,12 +25,12 @@ UAIPerceptionSystem::UAIPerceptionSystem(const FObjectInitializer& ObjectInitial
 	if (HasAnyFlags(RF_ClassDefaultObject) == false)
 	{
 		Senses.AddZeroed(ECorePerceptionTypes::MAX);
-		Senses[ECorePerceptionTypes::Hearing] = ConstructObject<UAISense>(UAISense_Hearing::StaticClass(), this);
-		Senses[ECorePerceptionTypes::Sight] = ConstructObject<UAISense>(UAISense_Sight::StaticClass(), this);
-		Senses[ECorePerceptionTypes::Damage] = ConstructObject<UAISense>(UAISense_Damage::StaticClass(), this);
-		Senses[ECorePerceptionTypes::Touch] = ConstructObject<UAISense>(UAISense_Touch::StaticClass(), this);		
-		Senses[ECorePerceptionTypes::Team] = ConstructObject<UAISense>(UAISense_Team::StaticClass(), this);		
-		Senses[ECorePerceptionTypes::Prediction] = ConstructObject<UAISense>(UAISense_Prediction::StaticClass(), this);
+		Senses[ECorePerceptionTypes::Hearing] = ConstructObject<UAISenseImplementation>(UAISenseImplementation_Hearing::StaticClass(), this);
+		Senses[ECorePerceptionTypes::Sight] = ConstructObject<UAISenseImplementation>(UAISenseImplementation_Sight::StaticClass(), this);
+		Senses[ECorePerceptionTypes::Damage] = ConstructObject<UAISenseImplementation>(UAISenseImplementation_Damage::StaticClass(), this);
+		Senses[ECorePerceptionTypes::Touch] = ConstructObject<UAISenseImplementation>(UAISenseImplementation_Touch::StaticClass(), this);		
+		Senses[ECorePerceptionTypes::Team] = ConstructObject<UAISenseImplementation>(UAISenseImplementation_Team::StaticClass(), this);		
+		Senses[ECorePerceptionTypes::Prediction] = ConstructObject<UAISenseImplementation>(UAISenseImplementation_Prediction::StaticClass(), this);
 	}
 	else
 	{
@@ -75,7 +75,7 @@ void UAIPerceptionSystem::PerformSourceRegistration()
 
 void UAIPerceptionSystem::OnNewListener(const FPerceptionListener& NewListener)
 {
-	UAISense** SenseInstance = Senses.GetData();
+	UAISenseImplementation** SenseInstance = Senses.GetData();
 	for (int32 SenseIndex = 0; SenseIndex < Senses.Num(); ++SenseIndex, ++SenseInstance)
 	{
 		if (*SenseInstance != NULL)
@@ -87,7 +87,7 @@ void UAIPerceptionSystem::OnNewListener(const FPerceptionListener& NewListener)
 
 void UAIPerceptionSystem::OnListenerUpdate(const FPerceptionListener& UpdatedListener)
 {
-	UAISense** SenseInstance = Senses.GetData();
+	UAISenseImplementation** SenseInstance = Senses.GetData();
 	for (int32 SenseIndex = 0; SenseIndex < Senses.Num(); ++SenseIndex, ++SenseInstance)
 	{
 		if (*SenseInstance != NULL)
@@ -117,7 +117,7 @@ void UAIPerceptionSystem::ManagerTick(float DeltaSeconds)
 	}
 
 	{
-		UAISense** SenseInstance = Senses.GetData();
+		UAISenseImplementation** SenseInstance = Senses.GetData();
 		for (int32 SenseIndex = 0; SenseIndex < Senses.Num(); ++SenseIndex, ++SenseInstance)
 		{
 			bNeedsUpdate |= *SenseInstance != NULL && (*SenseInstance)->ProgressTime(DeltaSeconds);
@@ -140,7 +140,7 @@ void UAIPerceptionSystem::ManagerTick(float DeltaSeconds)
 			}
 		}
 
-		UAISense** SenseInstance = Senses.GetData();
+		UAISenseImplementation** SenseInstance = Senses.GetData();
 		for (int32 SenseIndex = 0; SenseIndex < Senses.Num(); ++SenseIndex, ++SenseInstance)
 		{
 			if (*SenseInstance != NULL)
@@ -253,7 +253,7 @@ void UAIPerceptionSystem::UnregisterListener(UAIPerceptionComponent* Listener)
 
 void UAIPerceptionSystem::OnListenerRemoved(const FPerceptionListener& NewListener)
 {
-	UAISense** SenseInstance = Senses.GetData();
+	UAISenseImplementation** SenseInstance = Senses.GetData();
 	for (int32 SenseIndex = 0; SenseIndex < Senses.Num(); ++SenseIndex, ++SenseInstance)
 	{
 		if (*SenseInstance != NULL)
diff --git a/Engine/Source/Runtime/AIModule/Private/Perception/AISense.cpp b/Engine/Source/Runtime/AIModule/Private/Perception/AISense.cpp
deleted file mode 100644
diff --git a/Engine/Source/Runtime/AIModule/Private/Perception/AISenseImplementation.cpp b/Engine/Source/Runtime/AIModule/Private/Perception/AISenseImplementation.cpp
new file mode 100644
--- /dev/null
+++ b/Engine/Source/Runtime/AIModule/Private/Perception/AISenseImplementation.cpp
@@ -0,0 +1,28 @@
+// Copyright 1998-2014 Epic Games, Inc. All Rights Reserved.
+
+#include "AIModulePrivate.h"
+#include "Perception/AISenseImplementation.h"
+
+const float UAISenseImplementation::SuspendNextUpdate = FLT_MAX;
+
+UAISenseImplementation::UAISenseImplementation(const FObjectInitializer& ObjectInitializer)
+	: Super(ObjectInitializer)
+	, TimeUntilNextUpdate(SuspendNextUpdate)
+{
+}
+
+void UAISenseImplementation::PostInitProperties() 
+{
+	Super::PostInitProperties();
+
+	if (HasAnyFlags(RF_ClassDefaultObject) == false) 
+	{
+		PerceptionSystemInstance = Cast<UAIPerceptionSystem>(GetOuter());
+	}
+}
+
+AIPerception::FListenerMap* UAISenseImplementation::GetListeners() 
+{
+	check(PerceptionSystemInstance);
+	return &(PerceptionSystemInstance->GetListenersMap());
+}
\ No newline at end of file
diff --git a/Engine/Source/Runtime/AIModule/Private/Perception/AISenseImplementation_Damage.cpp b/Engine/Source/Runtime/AIModule/Private/Perception/AISenseImplementation_Damage.cpp
new file mode 100644
--- /dev/null
+++ b/Engine/Source/Runtime/AIModule/Private/Perception/AISenseImplementation_Damage.cpp
@@ -0,0 +1,45 @@
+// Copyright 1998-2014 Epic Games, Inc. All Rights Reserved.
+
+#include "AIModulePrivate.h"
+#include "Perception/AISenseImplementation_Damage.h"
+
+UAISenseImplementation_Damage::UAISenseImplementation_Damage(const FObjectInitializer& ObjectInitializer) :
+	Super(ObjectInitializer)
+{
+	
+}
+
+float UAISenseImplementation_Damage::Update()
+{
+	AIPerception::FListenerMap& ListenersMap = *GetListeners();
+
+	for (int32 EventIndex = 0; EventIndex < RegisteredEvents.Num(); ++EventIndex)
+	{
+		const FAIDamageEvent& Event = RegisteredEvents[EventIndex];
+
+		IAIPerceptionListenerInterface* PerceptionListener = Cast<IAIPerceptionListenerInterface>(Event.DamagedActor);
+		if (PerceptionListener != NULL)
+		{
+			UAIPerceptionComponent* PerceptionComponent = PerceptionListener->GetPerceptionComponent();
+			if (PerceptionComponent != NULL)
+			{
+				// this has to succeed, will assert a failure
+				FPerceptionListener& Listener = ListenersMap[PerceptionComponent->GetListenerId()];
+
+				Listener.RegisterStimulus(Event.Instigator, FAIStimulus(GetSenseIndex(), Event.Amount, Event.Location, Event.HitLocation));
+			}
+		}
+	}
+
+	RegisteredEvents.Reset();
+
+	// return decides when next tick is going to happen
+	return SuspendNextUpdate;
+}
+
+void UAISenseImplementation_Damage::RegisterEvent(const FAIDamageEvent& Event)
+{
+	RegisteredEvents.Add(Event);
+
+	RequestImmediateUpdate();
+}
\ No newline at end of file
diff --git a/Engine/Source/Runtime/AIModule/Private/Perception/AISenseImplementation_Hearing.cpp b/Engine/Source/Runtime/AIModule/Private/Perception/AISenseImplementation_Hearing.cpp
new file mode 100644
--- /dev/null
+++ b/Engine/Source/Runtime/AIModule/Private/Perception/AISenseImplementation_Hearing.cpp
@@ -0,0 +1,68 @@
+// Copyright 1998-2014 Epic Games, Inc. All Rights Reserved.
+
+#include "AIModulePrivate.h"
+#include "Perception/AISenseImplementation_Hearing.h"
+
+//----------------------------------------------------------------------//
+// FAINoiseEvent
+//----------------------------------------------------------------------//
+FAINoiseEvent::FAINoiseEvent(AActor* InInstigator, const FVector& InNoiseLocation, float InLoudness)
+	: Age(0.f), NoiseLocation(InNoiseLocation), Loudness(InLoudness), TeamIdentifier(uint8(-1)), Instigator(InInstigator)
+{
+	TeamIdentifier = FGenericTeamId::GetTeamIdentifier(InInstigator);
+}
+
+//----------------------------------------------------------------------//
+// UAISense_Hearing
+//----------------------------------------------------------------------//
+UAISenseImplementation_Hearing::UAISenseImplementation_Hearing(const FObjectInitializer& ObjectInitializer) :
+	Super(ObjectInitializer)
+{
+	
+}
+
+float UAISenseImplementation_Hearing::Update()
+{
+	AIPerception::FListenerMap& ListenersMap = *GetListeners();
+	UAIPerceptionSystem* PerseptionSys = GetPerceptionSystem();
+
+	for (AIPerception::FListenerMap::TIterator ListenerIt(ListenersMap); ListenerIt; ++ListenerIt)
+	{
+		FPerceptionListener& Listener = ListenerIt->Value;
+
+		if (Listener.HasSense(GetSenseIndex()) == false)
+		{
+			// skip listeners not interested in this sense
+			continue;
+		}
+
+		for (int32 EventIndex = 0; EventIndex < NoiseEvents.Num(); ++EventIndex)
+		{
+			const FAINoiseEvent& Event = NoiseEvents[EventIndex];
+		
+			// @todo implement some kind of TeamIdentifierType that would supply comparison operator 
+			if (Listener.TeamIdentifier == Event.TeamIdentifier 
+				|| FVector::DistSquared(Event.NoiseLocation, Listener.CachedLocation) > Listener.HearingRangeSq * Event.Loudness)
+			{
+				continue;
+			}
+			// calculate delay and fake it with Age
+			const float Delay = FVector::DistSquared(Event.NoiseLocation, Listener.CachedLocation)/900000;
+			// pass over to listener to process 			
+			PerseptionSys->RegisterDelayedStimulus(Listener.GetListenerId(), Delay, Event.Instigator
+				, FAIStimulus(GetSenseIndex(), Event.Loudness, Event.NoiseLocation, Listener.CachedLocation) );
+		}
+	}
+
+	NoiseEvents.Reset();
+
+	// return decides when next tick is going to happen
+	return SuspendNextUpdate;
+}
+
+void UAISenseImplementation_Hearing::RegisterEvent(const FAINoiseEvent& Event)
+{
+	NoiseEvents.Add(Event);
+
+	RequestImmediateUpdate();
+}
diff --git a/Engine/Source/Runtime/AIModule/Private/Perception/AISenseImplementation_Prediction.cpp b/Engine/Source/Runtime/AIModule/Private/Perception/AISenseImplementation_Prediction.cpp
new file mode 100644
--- /dev/null
+++ b/Engine/Source/Runtime/AIModule/Private/Perception/AISenseImplementation_Prediction.cpp
@@ -0,0 +1,49 @@
+// Copyright 1998-2014 Epic Games, Inc. All Rights Reserved.
+
+#include "AIModulePrivate.h"
+#include "Perception/AISenseImplementation_Prediction.h"
+
+UAISenseImplementation_Prediction::UAISenseImplementation_Prediction(const FObjectInitializer& ObjectInitializer) :
+	Super(ObjectInitializer)
+{
+}
+
+float UAISenseImplementation_Prediction::Update()
+{
+	AIPerception::FListenerMap& ListenersMap = *GetListeners();
+
+	for (int32 EventIndex = 0; EventIndex < RegisteredEvents.Num(); ++EventIndex)
+	{
+		const FAIPredictionEvent& Event = RegisteredEvents[EventIndex];
+
+		if (Event.Requestor != NULL && Event.PredictedActor != NULL)
+		{
+			IAIPerceptionListenerInterface* PerceptionListener = Cast<IAIPerceptionListenerInterface>(Event.Requestor);
+			if (PerceptionListener != NULL)
+			{
+				UAIPerceptionComponent* PerceptionComponent = PerceptionListener->GetPerceptionComponent();
+				if (PerceptionComponent != NULL && ListenersMap.Contains(PerceptionComponent->GetListenerId()))
+				{
+					// this has to succeed, will assert a failure
+					FPerceptionListener& Listener = ListenersMap[PerceptionComponent->GetListenerId()];
+
+					// calculate the prediction here:
+					const FVector PredictedLocation = Event.PredictedActor->GetActorLocation() + Event.PredictedActor->GetVelocity() * Event.TimeToPredict;
+
+					Listener.RegisterStimulus(Event.PredictedActor, FAIStimulus(GetSenseIndex(), 1.f, PredictedLocation, Listener.CachedLocation));
+				}
+			}
+		}
+	}
+
+	RegisteredEvents.Reset();
+
+	// return decides when next tick is going to happen
+	return SuspendNextUpdate;
+}
+
+void UAISenseImplementation_Prediction::RegisterEvent(const FAIPredictionEvent& Event)
+{
+	RegisteredEvents.Add(Event);
+	RequestImmediateUpdate();
+}
\ No newline at end of file
diff --git a/Engine/Source/Runtime/AIModule/Private/Perception/AISenseImplementation_Sight.cpp b/Engine/Source/Runtime/AIModule/Private/Perception/AISenseImplementation_Sight.cpp
new file mode 100644
--- /dev/null
+++ b/Engine/Source/Runtime/AIModule/Private/Perception/AISenseImplementation_Sight.cpp
@@ -0,0 +1,421 @@
+// Copyright 1998-2014 Epic Games, Inc. All Rights Reserved.
+
+#include "AIModulePrivate.h"
+#include "Perception/AISightTargetInterface.h"
+#include "Perception/AISenseImplementation_Sight.h"
+
+DECLARE_CYCLE_STAT(TEXT("Perception Sense: Sight"),STAT_AI_Sense_Sight,STATGROUP_AI);
+
+static const int32 DefaultMaxTracesPerTick = 6;
+
+//----------------------------------------------------------------------//
+// helpers
+//----------------------------------------------------------------------//
+FORCEINLINE_DEBUGGABLE bool CheckIsTargetInSightPie(const FPerceptionListener& Listener, const FVector& TargetLocation, const float SightRadiusSq)
+{
+	if (FVector::DistSquared(Listener.CachedLocation, TargetLocation) <= SightRadiusSq) 
+	{
+		const FVector DirectionToTarget = (TargetLocation - Listener.CachedLocation).UnsafeNormal();
+		return FVector::DotProduct(DirectionToTarget, Listener.CachedDirection) > Listener.PeripheralVisionAngleCos;
+	}
+
+	return false;
+}
+
+//----------------------------------------------------------------------//
+// FAISightTarget
+//----------------------------------------------------------------------//
+const FAISightTarget::FTargetId FAISightTarget::InvalidTargetId = NAME_None;
+
+FAISightTarget::FAISightTarget(AActor* InTarget, FGenericTeamId InTeamId)
+	: Target(InTarget), SightTargetInterface(NULL), TeamId(InTeamId)
+{
+	if (InTarget)
+	{
+		TargetId = InTarget->GetFName();
+	}
+	else
+	{
+		TargetId = InvalidTargetId;
+	}
+}
+
+//----------------------------------------------------------------------//
+// UAISense_Sight
+//----------------------------------------------------------------------//
+UAISenseImplementation_Sight::UAISenseImplementation_Sight(const FObjectInitializer& ObjectInitializer)
+	: Super(ObjectInitializer)
+	, MaxTracesPerTick(DefaultMaxTracesPerTick)
+	, HighImportanceQueryDistanceThreshold(300.f)
+	, MaxQueryImportance(60.f)
+	, SightLimitQueryImportance(10.f)
+{
+	if (HasAnyFlags(RF_ClassDefaultObject) == false)
+	{
+		OnNewListenerDelegate.BindUObject(this, &UAISenseImplementation_Sight::OnNewListenerImpl);
+		OnListenerUpdateDelegate.BindUObject(this, &UAISenseImplementation_Sight::OnListenerUpdateImpl);
+		OnListenerRemovedDelegate.BindUObject(this, &UAISenseImplementation_Sight::OnListenerRemovedImpl);
+	}
+}
+
+FORCEINLINE_DEBUGGABLE float UAISenseImplementation_Sight::CalcQueryImportance(const FPerceptionListener& Listener, const FVector& TargetLocation, const float SightRadiusSq) const
+{
+	const float DistanceSq = FVector::DistSquared(Listener.CachedLocation, TargetLocation);
+	return DistanceSq <= HighImportanceDistanceSquare ? MaxQueryImportance
+		: FMath::Clamp((SightLimitQueryImportance - MaxQueryImportance) / SightRadiusSq * DistanceSq + MaxQueryImportance, 0.f, MaxQueryImportance);
+}
+
+void UAISenseImplementation_Sight::PostInitProperties()
+{
+	Super::PostInitProperties();
+	HighImportanceDistanceSquare = FMath::Square(HighImportanceQueryDistanceThreshold);
+}
+
+float UAISenseImplementation_Sight::Update()
+{
+	static const FName NAME_AILineOfSight = FName(TEXT("AILineOfSight"));
+
+	SCOPE_CYCLE_COUNTER(STAT_AI_Sense_Sight);
+
+	const UWorld* World = GEngine->GetWorldFromContextObject(GetPerceptionSystem()->GetOuter());
+
+	if (World == NULL)
+	{
+		return SuspendNextUpdate;
+	}
+
+	int32 TracesCount = 0;
+	static const int32 InitialInvalidItemsSize = 16;
+	TArray<int32> InvalidQueries;
+	TArray<FAISightTarget::FTargetId> InvalidTargets;
+	InvalidQueries.Reserve(InitialInvalidItemsSize);
+	InvalidTargets.Reserve(InitialInvalidItemsSize);
+
+	AIPerception::FListenerMap& ListenersMap = *GetListeners();
+
+	FAISightQuery* SightQuery = SightQueryQueue.GetData();
+	for (int32 QueryIndex = 0; QueryIndex < SightQueryQueue.Num(); ++QueryIndex, ++SightQuery)
+	{
+		if (TracesCount < MaxTracesPerTick)
+		{
+			FPerceptionListener& Listener = ListenersMap[SightQuery->ObserverId];
+			ensure(Listener.Listener.IsValid());
+			FAISightTarget& Target = ObservedTargets[SightQuery->TargetId];
+					
+			const bool bTargetValid = Target.Target.IsValid();
+			const bool bListenerValid = Listener.Listener.IsValid();
+
+			// @todo figure out what should we do if not valid
+			if (bTargetValid && bListenerValid)
+			{
+				AActor* TargetActor = Target.Target.Get();
+				const FVector TargetLocation = TargetActor->GetActorLocation();
+				const float SightRadiusSq = SightQuery->bLastResult ? Listener.LoseSightRadiusSq : Listener.SightRadiusSq;
+
+				if (CheckIsTargetInSightPie(Listener, TargetLocation, SightRadiusSq))
+				{
+//					UE_VLOG_SEGMENT(Listener.Listener.Get()->GetOwner(), Listener.CachedLocation, TargetLocation, FColor::Green, TEXT("%s"), *(Target.TargetId.ToString()));
+
+					FVector OutSeenLocation(0.f);
+					// do line checks
+					if (Target.SightTargetInterface != NULL)
+					{
+						int32 NumberOfLoSChecksPerformed = 0;
+						if (Target.SightTargetInterface->CanBeSeenFrom(Listener.CachedLocation, OutSeenLocation, NumberOfLoSChecksPerformed, Listener.Listener->GetBodyActor()) == true)
+						{
+							Listener.RegisterStimulus(TargetActor, FAIStimulus(ECorePerceptionTypes::Sight, 1.f, OutSeenLocation, Listener.CachedLocation));
+							SightQuery->bLastResult = true;
+						}
+						else
+						{
+//							UE_VLOG_LOCATION(Listener.Listener.Get()->GetOwner(), TargetLocation, 25.f, FColor::Red, TEXT(""));
+							Listener.RegisterStimulus(TargetActor, FAIStimulus(ECorePerceptionTypes::Sight, 0.f, TargetLocation, Listener.CachedLocation, FAIStimulus::SensingFailed));
+							SightQuery->bLastResult = false;
+						}
+
+						TracesCount += NumberOfLoSChecksPerformed;
+					}
+					else
+					{
+						// we need to do tests ourselves
+						/*const bool bHit = World->LineTraceTest(Listener.CachedLocation, TargetLocation
+							, FCollisionQueryParams(NAME_AILineOfSight, true, Listener.Listener->GetBodyActor())
+							, FCollisionObjectQueryParams(ECC_WorldStatic));*/
+						FHitResult HitResult;
+						const bool bHit = World->LineTraceSingle(HitResult, Listener.CachedLocation, TargetLocation
+							, FCollisionQueryParams(NAME_AILineOfSight, true, Listener.Listener->GetBodyActor())
+							, FCollisionObjectQueryParams(ECC_WorldStatic));
+
+						++TracesCount;
+
+						if (bHit == false || (HitResult.Actor.IsValid() && HitResult.Actor->IsOwnedBy(TargetActor)))
+						{
+							Listener.RegisterStimulus(TargetActor, FAIStimulus(ECorePerceptionTypes::Sight, 1.f, TargetLocation, Listener.CachedLocation));
+							SightQuery->bLastResult = true;
+						}
+						else
+						{
+//							UE_VLOG_LOCATION(Listener.Listener.Get()->GetOwner(), TargetLocation, 25.f, FColor::Red, TEXT(""));
+							Listener.RegisterStimulus(TargetActor, FAIStimulus(ECorePerceptionTypes::Sight, 0.f, TargetLocation, Listener.CachedLocation, FAIStimulus::SensingFailed));
+							SightQuery->bLastResult = false;
+						}
+					}
+				}
+				else
+				{
+//					UE_VLOG_SEGMENT(Listener.Listener.Get()->GetOwner(), Listener.CachedLocation, TargetLocation, FColor::Red, TEXT("%s"), *(Target.TargetId.ToString()));
+					Listener.RegisterStimulus(TargetActor, FAIStimulus(ECorePerceptionTypes::Sight, 0.f, TargetLocation, Listener.CachedLocation, FAIStimulus::SensingFailed));
+					SightQuery->bLastResult = false;
+				}
+
+				SightQuery->Importance = CalcQueryImportance(Listener, TargetLocation, SightRadiusSq);
+
+				// restart query
+				SightQuery->Age = 0.f;
+			}
+			else
+			{
+				// put this index to "to be removed" array
+				InvalidQueries.Add(QueryIndex);
+				if (bTargetValid == false)
+				{
+					InvalidTargets.AddUnique(SightQuery->TargetId);
+				}
+			}
+		}
+		else
+		{
+			// age unprocessed queries so that they can advance in the queue during next sort
+			SightQuery->Age += 1.f;
+		}
+
+		SightQuery->RecalcScore();
+	}
+
+	if (InvalidQueries.Num() > 0)
+	{
+		for (int32 Index = InvalidQueries.Num() - 1; Index >= 0; --Index)
+		{
+			// removing with swapping here, since queue is going to be sorted anyway
+			SightQueryQueue.RemoveAtSwap(InvalidQueries[Index], 1, /*bAllowShrinking*/false);
+		}
+
+		if (InvalidTargets.Num() > 0)
+		{
+			for (const auto& TargetId : InvalidTargets)
+			{
+				// remove affected queries
+				RemoveAllQueriesToTarget(TargetId, DontSort);
+				// remove target itself
+				ObservedTargets.Remove(TargetId);
+			}
+
+			// remove holes
+			ObservedTargets.Compact();
+		}
+	}
+
+	// sort Sight Queries
+	SortQueries();
+
+	//return SightQueryQueue.Num() > 0 ? 1.f/6 : FLT_MAX;
+	return 0.f;
+}
+
+void UAISenseImplementation_Sight::RegisterEvent(const FAISightEvent& Event)
+{
+
+}
+
+void UAISenseImplementation_Sight::RegisterSource(AActor& SourceActor) 
+{
+	RegisterTarget(SourceActor, Sort);
+}
+
+void UAISenseImplementation_Sight::RegisterTarget(AActor& TargetActor, FQueriesOperationPostProcess PostProcess)
+{
+	SCOPE_CYCLE_COUNTER(STAT_AI_Sense_Sight);
+	
+	FAISightTarget* SightTarget = ObservedTargets.Find(TargetActor.GetFName());
+	
+	if (SightTarget == NULL)
+	{
+		FAISightTarget NewSightTarget(&TargetActor);
+
+		SightTarget = &(ObservedTargets.Add(NewSightTarget.TargetId, NewSightTarget));
+		SightTarget->SightTargetInterface = Cast<IAISightTargetInterface>(&TargetActor);
+	}
+
+	// set/update data
+	SightTarget->TeamId = FGenericTeamId::GetTeamIdentifier(&TargetActor);
+	
+	// generate all pairs and add them to current Sight Queries
+	bool bNewQueriesAdded = false;
+	AIPerception::FListenerMap& ListenersMap = *GetListeners();
+	const FVector TargetLocation = TargetActor.GetActorLocation();
+
+	for (AIPerception::FListenerMap::TConstIterator ItListener(ListenersMap); ItListener; ++ItListener)
+	{
+		const FPerceptionListener& Listener = ItListener->Value;
+		const IGenericTeamAgentInterface* ListenersTeamAgent = Listener.GetTeamAgent();
+
+		// @todo add configuration here
+		if (Listener.HasSense(GetSenseIndex()) && (ListenersTeamAgent == NULL || ListenersTeamAgent->GetTeamAttitudeTowards(TargetActor) == ETeamAttitude::Hostile))
+		{
+			// create a sight query		
+			FAISightQuery SightQuery(ItListener->Key, SightTarget->TargetId);
+			SightQuery.Importance = CalcQueryImportance(ItListener->Value, TargetLocation, Listener.SightRadiusSq);
+
+			SightQueryQueue.Add(SightQuery);
+			bNewQueriesAdded = true;
+		}
+	}
+
+	// sort Sight Queries
+	if (PostProcess == Sort && bNewQueriesAdded)
+	{
+		SortQueries();
+		RequestImmediateUpdate();
+	}
+}
+
+void UAISenseImplementation_Sight::OnNewListenerImpl(const FPerceptionListener& NewListener)
+{
+	if (NewListener.HasSense(GetSenseIndex()) == false)
+	{
+		return;
+	}
+
+	bool bNewQueriesAdded = false;
+
+	const IGenericTeamAgentInterface* ListenersTeamAgent = NewListener.GetTeamAgent();
+
+	// create sight queries with all legal targets
+	for (TMap<FName, FAISightTarget>::TConstIterator ItTarget(ObservedTargets); ItTarget; ++ItTarget)
+	{
+		const AActor* TargetActor = ItTarget->Value.GetTargetActor();
+		if (TargetActor == NULL)
+		{
+			continue;
+		}
+
+		// @todo this should be configurable - some AI might want to observe Neutrals and Friendlies as well
+		if (ListenersTeamAgent == NULL || ListenersTeamAgent->GetTeamAttitudeTowards(*TargetActor) == ETeamAttitude::Hostile)
+		{
+			// create a sight query		
+			FAISightQuery SightQuery(NewListener.GetListenerId(), ItTarget->Key);
+			SightQuery.Importance = CalcQueryImportance(NewListener, ItTarget->Value.GetLocationSimple(), NewListener.SightRadiusSq);
+
+			SightQueryQueue.Add(SightQuery);
+			bNewQueriesAdded = true;
+		}
+	}
+
+	// sort Sight Queries
+	if (bNewQueriesAdded)
+	{
+		SortQueries();
+		RequestImmediateUpdate();
+	}
+}
+
+void UAISenseImplementation_Sight::OnListenerUpdateImpl(const FPerceptionListener& UpdatedListener)
+{
+	// first, naive implementation:
+	// 1. remove all queries by this listener
+	// 2. proceed as if it was a new listener
+	// @todo add stats here to check how bad it is
+
+	// remove all queries
+	RemoveAllQueriesByListener(UpdatedListener, DontSort);
+
+	// see if this listener is a Target as well
+	const FAISightTarget::FTargetId AsTargetId = UpdatedListener.GetBodyActorName();
+	FAISightTarget* AsTarget = ObservedTargets.Find(AsTargetId);
+	if (AsTarget != NULL)
+	{
+		RemoveAllQueriesToTarget(AsTargetId, DontSort);
+		if (AsTarget->Target.IsValid())
+		{
+			RegisterTarget(*(AsTarget->Target.Get()), DontSort);
+		}
+	}
+
+	// act as if it was a new listener
+	OnNewListenerImpl(UpdatedListener);
+}
+
+void UAISenseImplementation_Sight::OnListenerRemovedImpl(const FPerceptionListener& UpdatedListener)
+{
+	RemoveAllQueriesByListener(UpdatedListener, DontSort);
+
+	if (UpdatedListener.Listener.IsValid())
+	{
+		// see if this listener is a Target as well
+		const FAISightTarget::FTargetId AsTargetId = UpdatedListener.GetBodyActorName();
+		FAISightTarget* AsTarget = ObservedTargets.Find(AsTargetId);
+		if (AsTarget != NULL)
+		{
+			RemoveAllQueriesToTarget(AsTargetId, Sort);			
+		}
+	}
+	else
+	{
+		//@todo quite possible there are left over sight queries with this listener as target
+	}
+}
+
+void UAISenseImplementation_Sight::RemoveAllQueriesByListener(const FPerceptionListener& Listener, FQueriesOperationPostProcess PostProcess)
+{
+	SCOPE_CYCLE_COUNTER(STAT_AI_Sense_Sight);
+
+	if (SightQueryQueue.Num() == 0)
+	{
+		return;
+	}
+
+	const uint32 ListenerId = Listener.GetListenerId();
+	bool bQueriesRemoved = false;
+	
+	const FAISightQuery* SightQuery = &SightQueryQueue[SightQueryQueue.Num() - 1];
+	for (int32 QueryIndex = SightQueryQueue.Num() - 1; QueryIndex >= 0 ; --QueryIndex, --SightQuery)
+	{
+		if (SightQuery->ObserverId == ListenerId)
+		{
+			SightQueryQueue.RemoveAt(QueryIndex, 1, /*bAllowShrinking=*/false);
+			bQueriesRemoved = true;
+		}
+	}
+
+	if (PostProcess == Sort && bQueriesRemoved)
+	{
+		SortQueries();
+	}
+}
+
+void UAISenseImplementation_Sight::RemoveAllQueriesToTarget(const FName& TargetId, FQueriesOperationPostProcess PostProcess)
+{
+	SCOPE_CYCLE_COUNTER(STAT_AI_Sense_Sight);
+
+	if (SightQueryQueue.Num() == 0)
+	{
+		return;
+	}
+
+	bool bQueriesRemoved = false;
+
+	const FAISightQuery* SightQuery = &SightQueryQueue[SightQueryQueue.Num() - 1];
+	for (int32 QueryIndex = SightQueryQueue.Num() - 1; QueryIndex >= 0 ; --QueryIndex, --SightQuery)
+	{
+		if (SightQuery->TargetId == TargetId)
+		{
+			SightQueryQueue.RemoveAt(QueryIndex, 1, /*bAllowShrinking=*/false);
+			bQueriesRemoved = true;
+		}
+	}
+
+	if (PostProcess == Sort && bQueriesRemoved)
+	{
+		SortQueries();
+	}
+}
\ No newline at end of file
diff --git a/Engine/Source/Runtime/AIModule/Private/Perception/AISenseImplementation_Team.cpp b/Engine/Source/Runtime/AIModule/Private/Perception/AISenseImplementation_Team.cpp
new file mode 100644
--- /dev/null
+++ b/Engine/Source/Runtime/AIModule/Private/Perception/AISenseImplementation_Team.cpp
@@ -0,0 +1,66 @@
+// Copyright 1998-2014 Epic Games, Inc. All Rights Reserved.
+
+#include "AIModulePrivate.h"
+#include "Perception/AISenseImplementation_Team.h"
+
+//----------------------------------------------------------------------//
+// 
+//----------------------------------------------------------------------//
+FAITeamStimulusEvent::FAITeamStimulusEvent(AActor* InBroadcaster, AActor* InEnemy, const FVector& InLastKnowLocation, float EventRange, float PassedInfoAge)
+	: LastKnowLocation(InLastKnowLocation), RangeSq(FMath::Square(EventRange)), InformationAge(PassedInfoAge), Broadcaster(InBroadcaster), Enemy(InEnemy)
+{
+	CacheBroadcastLocation();
+
+	TeamIdentifier = FGenericTeamId::GetTeamIdentifier(InBroadcaster);
+}
+
+//----------------------------------------------------------------------//
+// 
+//----------------------------------------------------------------------//
+UAISenseImplementation_Team::UAISenseImplementation_Team(const FObjectInitializer& ObjectInitializer) :
+	Super(ObjectInitializer)
+{
+	
+}
+
+float UAISenseImplementation_Team::Update()
+{
+	AIPerception::FListenerMap& ListenersMap = *GetListeners();
+	
+	for (AIPerception::FListenerMap::TIterator ListenerIt(ListenersMap); ListenerIt; ++ListenerIt)
+	{
+		FPerceptionListener& Listener = ListenerIt->Value;
+
+		if (Listener.HasSense(GetSenseIndex()) == false)
+		{
+			// skip listeners not interested in this sense
+			continue;
+		}
+
+		for (int32 EventIndex = 0; EventIndex < RegisteredEvents.Num(); ++EventIndex)
+		{
+			const FAITeamStimulusEvent& Event = RegisteredEvents[EventIndex];
+
+			// @todo implement some kind of TeamIdentifierType that would supply comparison operator 
+			if (Listener.TeamIdentifier != Event.TeamIdentifier 
+				|| FVector::DistSquared(Event.GetBroadcastLocation(), Listener.CachedLocation) > Event.RangeSq)
+			{
+				continue;
+			}
+			
+			Listener.RegisterStimulus(Event.Enemy, FAIStimulus(GetSenseIndex(), 1.f, Event.LastKnowLocation, Event.GetBroadcastLocation(), FAIStimulus::SensingSucceeded, Event.InformationAge));
+		}
+	}
+
+	RegisteredEvents.Reset();
+
+	// return decides when next tick is going to happen
+	return SuspendNextUpdate;
+}
+
+void UAISenseImplementation_Team::RegisterEvent(const FAITeamStimulusEvent& Event)
+{
+	RegisteredEvents.Add(Event);
+
+	RequestImmediateUpdate();
+}
\ No newline at end of file
diff --git a/Engine/Source/Runtime/AIModule/Private/Perception/AISenseImplementation_Touch.cpp b/Engine/Source/Runtime/AIModule/Private/Perception/AISenseImplementation_Touch.cpp
new file mode 100644
--- /dev/null
+++ b/Engine/Source/Runtime/AIModule/Private/Perception/AISenseImplementation_Touch.cpp
@@ -0,0 +1,48 @@
+// Copyright 1998-2014 Epic Games, Inc. All Rights Reserved.
+
+#include "AIModulePrivate.h"
+#include "Perception/AISenseImplementation_Touch.h"
+
+UAISenseImplementation_Touch::UAISenseImplementation_Touch(const FObjectInitializer& ObjectInitializer) :
+	Super(ObjectInitializer)
+{
+	
+}
+
+float UAISenseImplementation_Touch::Update()
+{
+	AIPerception::FListenerMap& ListenersMap = *GetListeners();
+
+	for (int32 EventIndex = 0; EventIndex < RegisteredEvents.Num(); ++EventIndex)
+	{
+		const FAITouchEvent& Event = RegisteredEvents[EventIndex];
+
+		if (Event.TouchReceiver != NULL && Event.OtherActor != NULL)
+		{
+			IAIPerceptionListenerInterface* PerceptionListener = Cast<IAIPerceptionListenerInterface>(Event.TouchReceiver);
+			if (PerceptionListener != NULL)
+			{
+				UAIPerceptionComponent* PerceptionComponent = PerceptionListener->GetPerceptionComponent();
+				if (PerceptionComponent != NULL && ListenersMap.Contains(PerceptionComponent->GetListenerId()))
+				{
+					// this has to succeed, will assert a failure
+					FPerceptionListener& Listener = ListenersMap[PerceptionComponent->GetListenerId()];
+
+					Listener.RegisterStimulus(Event.OtherActor, FAIStimulus(GetSenseIndex(), 1.f, Event.Location, Event.Location));
+				}
+			}
+		}
+	}
+
+	RegisteredEvents.Reset();
+
+	// return decides when next tick is going to happen
+	return SuspendNextUpdate;
+}
+
+void UAISenseImplementation_Touch::RegisterEvent(const FAITouchEvent& Event)
+{
+	RegisteredEvents.Add(Event);
+
+	RequestImmediateUpdate();
+}
\ No newline at end of file
diff --git a/Engine/Source/Runtime/AIModule/Private/Perception/AISense_Damage.cpp b/Engine/Source/Runtime/AIModule/Private/Perception/AISense_Damage.cpp
deleted file mode 100644
diff --git a/Engine/Source/Runtime/AIModule/Private/Perception/AISense_Hearing.cpp b/Engine/Source/Runtime/AIModule/Private/Perception/AISense_Hearing.cpp
deleted file mode 100644
diff --git a/Engine/Source/Runtime/AIModule/Private/Perception/AISense_Prediction.cpp b/Engine/Source/Runtime/AIModule/Private/Perception/AISense_Prediction.cpp
deleted file mode 100644
diff --git a/Engine/Source/Runtime/AIModule/Private/Perception/AISense_Sight.cpp b/Engine/Source/Runtime/AIModule/Private/Perception/AISense_Sight.cpp
deleted file mode 100644
diff --git a/Engine/Source/Runtime/AIModule/Private/Perception/AISense_Team.cpp b/Engine/Source/Runtime/AIModule/Private/Perception/AISense_Team.cpp
deleted file mode 100644
diff --git a/Engine/Source/Runtime/AIModule/Private/Perception/AISense_Touch.cpp b/Engine/Source/Runtime/AIModule/Private/Perception/AISense_Touch.cpp
deleted file mode 100644
