BP blackboard interface improvements #UE4

- Added GetBlackboardValueAsRotator and SetBlackboardValueAsRotator functions to UBTFunctionLibrary
- Added a single function for reseting value of indicated Key
- deprecated vector-specific clearing function

diff --git a/Engine/Source/Runtime/AIModule/Classes/BehaviorTree/BTFunctionLibrary.h b/Engine/Source/Runtime/AIModule/Classes/BehaviorTree/BTFunctionLibrary.h
--- a/Engine/Source/Runtime/AIModule/Classes/BehaviorTree/BTFunctionLibrary.h
+++ b/Engine/Source/Runtime/AIModule/Classes/BehaviorTree/BTFunctionLibrary.h
@@ -4,6 +4,7 @@
 #include "BTFunctionLibrary.generated.h"
 
 class UBlackboardComponent;
+struct FBlackboardKeySelector;
 
 UCLASS(meta=(RestrictedToClasses="BTNode"))
 class AIMODULE_API UBTFunctionLibrary : public UBlueprintFunctionLibrary
@@ -15,64 +16,75 @@ public:
 	static UBlackboardComponent* GetOwnersBlackboard(UBTNode* NodeOwner);
 
 	UFUNCTION(BlueprintPure, Category="AI|BehaviorTree", Meta=(HidePin="NodeOwner", DefaultToSelf="NodeOwner"))
-	static UObject* GetBlackboardValueAsObject(UBTNode* NodeOwner, const struct FBlackboardKeySelector& Key);
+	static UObject* GetBlackboardValueAsObject(UBTNode* NodeOwner, const FBlackboardKeySelector& Key);
 
 	UFUNCTION(BlueprintPure, Category="AI|BehaviorTree", Meta=(HidePin="NodeOwner", DefaultToSelf="NodeOwner"))
-	static AActor* GetBlackboardValueAsActor(UBTNode* NodeOwner, const struct FBlackboardKeySelector& Key);
+	static AActor* GetBlackboardValueAsActor(UBTNode* NodeOwner, const FBlackboardKeySelector& Key);
 
 	UFUNCTION(BlueprintPure, Category="AI|BehaviorTree", Meta=(HidePin="NodeOwner", DefaultToSelf="NodeOwner"))
-	static UClass* GetBlackboardValueAsClass(UBTNode* NodeOwner, const struct FBlackboardKeySelector& Key);
+	static UClass* GetBlackboardValueAsClass(UBTNode* NodeOwner, const FBlackboardKeySelector& Key);
 
 	UFUNCTION(BlueprintPure, Category="AI|BehaviorTree", Meta=(HidePin="NodeOwner", DefaultToSelf="NodeOwner"))
-	static uint8 GetBlackboardValueAsEnum(UBTNode* NodeOwner, const struct FBlackboardKeySelector& Key);
+	static uint8 GetBlackboardValueAsEnum(UBTNode* NodeOwner, const FBlackboardKeySelector& Key);
 
 	UFUNCTION(BlueprintPure, Category="AI|BehaviorTree", Meta=(HidePin="NodeOwner", DefaultToSelf="NodeOwner"))
-	static int32 GetBlackboardValueAsInt(UBTNode* NodeOwner, const struct FBlackboardKeySelector& Key);
+	static int32 GetBlackboardValueAsInt(UBTNode* NodeOwner, const FBlackboardKeySelector& Key);
 
 	UFUNCTION(BlueprintPure, Category="AI|BehaviorTree", Meta=(HidePin="NodeOwner", DefaultToSelf="NodeOwner"))
-	static float GetBlackboardValueAsFloat(UBTNode* NodeOwner, const struct FBlackboardKeySelector& Key);
+	static float GetBlackboardValueAsFloat(UBTNode* NodeOwner, const FBlackboardKeySelector& Key);
 
 	UFUNCTION(BlueprintPure, Category="AI|BehaviorTree", Meta=(HidePin="NodeOwner", DefaultToSelf="NodeOwner"))
-	static bool GetBlackboardValueAsBool(UBTNode* NodeOwner, const struct FBlackboardKeySelector& Key);
+	static bool GetBlackboardValueAsBool(UBTNode* NodeOwner, const FBlackboardKeySelector& Key);
 
 	UFUNCTION(BlueprintPure, Category="AI|BehaviorTree", Meta=(HidePin="NodeOwner", DefaultToSelf="NodeOwner"))
-	static FString GetBlackboardValueAsString(UBTNode* NodeOwner, const struct FBlackboardKeySelector& Key);
+	static FString GetBlackboardValueAsString(UBTNode* NodeOwner, const FBlackboardKeySelector& Key);
 
 	UFUNCTION(BlueprintPure, Category="AI|BehaviorTree", Meta=(HidePin="NodeOwner", DefaultToSelf="NodeOwner"))
-	static FName GetBlackboardValueAsName(UBTNode* NodeOwner, const struct FBlackboardKeySelector& Key);
+	static FName GetBlackboardValueAsName(UBTNode* NodeOwner, const FBlackboardKeySelector& Key);
 
 	UFUNCTION(BlueprintPure, Category="AI|BehaviorTree", Meta=(HidePin="NodeOwner", DefaultToSelf="NodeOwner"))
-	static FVector GetBlackboardValueAsVector(UBTNode* NodeOwner, const struct FBlackboardKeySelector& Key);
+	static FVector GetBlackboardValueAsVector(UBTNode* NodeOwner, const FBlackboardKeySelector& Key);
+
+	UFUNCTION(BlueprintPure, Category ="AI|BehaviorTree", Meta=(HidePin="NodeOwner", DefaultToSelf="NodeOwner"))
+	static FRotator GetBlackboardValueAsRotator(UBTNode* NodeOwner, const FBlackboardKeySelector& Key);
 
 	UFUNCTION(BlueprintCallable, Category="AI|BehaviorTree", Meta=(HidePin="NodeOwner", DefaultToSelf="NodeOwner"))
-	static void SetBlackboardValueAsObject(UBTNode* NodeOwner, const struct FBlackboardKeySelector& Key, UObject* Value);
+	static void SetBlackboardValueAsObject(UBTNode* NodeOwner, const FBlackboardKeySelector& Key, UObject* Value);
 
 	UFUNCTION(BlueprintCallable, Category="AI|BehaviorTree", Meta=(HidePin="NodeOwner", DefaultToSelf="NodeOwner"))
-	static void SetBlackboardValueAsClass(UBTNode* NodeOwner, const struct FBlackboardKeySelector& Key, UClass* Value);
+	static void SetBlackboardValueAsClass(UBTNode* NodeOwner, const FBlackboardKeySelector& Key, UClass* Value);
 
 	UFUNCTION(BlueprintCallable, Category="AI|BehaviorTree", Meta=(HidePin="NodeOwner", DefaultToSelf="NodeOwner"))
-	static void SetBlackboardValueAsEnum(UBTNode* NodeOwner, const struct FBlackboardKeySelector& Key, uint8 Value);
+	static void SetBlackboardValueAsEnum(UBTNode* NodeOwner, const FBlackboardKeySelector& Key, uint8 Value);
 
 	UFUNCTION(BlueprintCallable, Category="AI|BehaviorTree", Meta=(HidePin="NodeOwner", DefaultToSelf="NodeOwner"))
-	static void SetBlackboardValueAsInt(UBTNode* NodeOwner, const struct FBlackboardKeySelector& Key, int32 Value);
+	static void SetBlackboardValueAsInt(UBTNode* NodeOwner, const FBlackboardKeySelector& Key, int32 Value);
 
 	UFUNCTION(BlueprintCallable, Category="AI|BehaviorTree", Meta=(HidePin="NodeOwner", DefaultToSelf="NodeOwner"))
-	static void SetBlackboardValueAsFloat(UBTNode* NodeOwner, const struct FBlackboardKeySelector& Key, float Value);
+	static void SetBlackboardValueAsFloat(UBTNode* NodeOwner, const FBlackboardKeySelector& Key, float Value);
 
 	UFUNCTION(BlueprintCallable, Category="AI|BehaviorTree", Meta=(HidePin="NodeOwner", DefaultToSelf="NodeOwner"))
-	static void SetBlackboardValueAsBool(UBTNode* NodeOwner, const struct FBlackboardKeySelector& Key, bool Value);
+	static void SetBlackboardValueAsBool(UBTNode* NodeOwner, const FBlackboardKeySelector& Key, bool Value);
 
 	UFUNCTION(BlueprintCallable, Category="AI|BehaviorTree", Meta=(HidePin="NodeOwner", DefaultToSelf="NodeOwner"))
-	static void SetBlackboardValueAsString(UBTNode* NodeOwner, const struct FBlackboardKeySelector& Key, FString Value);
+	static void SetBlackboardValueAsString(UBTNode* NodeOwner, const FBlackboardKeySelector& Key, FString Value);
 
 	UFUNCTION(BlueprintCallable, Category="AI|BehaviorTree", Meta=(HidePin="NodeOwner", DefaultToSelf="NodeOwner"))
-	static void SetBlackboardValueAsName(UBTNode* NodeOwner, const struct FBlackboardKeySelector& Key, FName Value);
+	static void SetBlackboardValueAsName(UBTNode* NodeOwner, const FBlackboardKeySelector& Key, FName Value);
 
 	UFUNCTION(BlueprintCallable, Category="AI|BehaviorTree", Meta=(HidePin="NodeOwner", DefaultToSelf="NodeOwner"))
-	static void SetBlackboardValueAsVector(UBTNode* NodeOwner, const struct FBlackboardKeySelector& Key, FVector Value);
+	static void SetBlackboardValueAsVector(UBTNode* NodeOwner, const FBlackboardKeySelector& Key, FVector Value);
+
+	/** (DEPRECATED) Use ClearBlackboardValue instead */
+	UFUNCTION(BlueprintCallable, Category=BehaviorTree, Meta=(HidePin="NodeOwner", DefaultToSelf="NodeOwner", DeprecatedFunction, DeprecationMessage="Use ClearBlackboardValue instead."))
+	static void ClearBlackboardValueAsVector(UBTNode* NodeOwner, const FBlackboardKeySelector& Key);
+
+	UFUNCTION(BlueprintCallable, Category = "AI|BehaviorTree", Meta = (HidePin = "NodeOwner", DefaultToSelf = "NodeOwner"))
+	static void SetBlackboardValueAsRotator(UBTNode* NodeOwner, const FBlackboardKeySelector& Key, FRotator Value);
 
-	UFUNCTION(BlueprintCallable, Category=BehaviorTree, Meta=(HidePin="NodeOwner", DefaultToSelf="NodeOwner"))
-	static void ClearBlackboardValueAsVector(UBTNode* NodeOwner, const struct FBlackboardKeySelector& Key);
+	/** Resets indicated value to "not set" value, based on values type */
+	UFUNCTION(BlueprintCallable, Category = BehaviorTree, Meta = (HidePin = "NodeOwner", DefaultToSelf = "NodeOwner"))
+	static void ClearBlackboardValue(UBTNode* NodeOwner, const FBlackboardKeySelector& Key);
 
 	/** Initialize variables marked as "instance memory" and set owning actor for blackboard operations */
 	UFUNCTION(BlueprintCallable, Category="AI|BehaviorTree", Meta=(HidePin="NodeOwner", DefaultToSelf="NodeOwner", DeprecatedFunction, DeprecationMessage="No longer needed"))
diff --git a/Engine/Source/Runtime/AIModule/Classes/BehaviorTree/Blackboard/BlackboardKeyType.h b/Engine/Source/Runtime/AIModule/Classes/BehaviorTree/Blackboard/BlackboardKeyType.h
--- a/Engine/Source/Runtime/AIModule/Classes/BehaviorTree/Blackboard/BlackboardKeyType.h
+++ b/Engine/Source/Runtime/AIModule/Classes/BehaviorTree/Blackboard/BlackboardKeyType.h
@@ -83,6 +83,9 @@ class AIMODULE_API UBlackboardKeyType : public UObject
 	/** extract rotation from entry */
 	virtual bool GetRotation(const uint8* MemoryBlock, FRotator& Rotation) const;
 
+	/** sets value to the default */
+	virtual bool Clear(uint8* MemoryBlock) const;
+
 	/** interprets given two memory blocks as "my" type and compares them */
 	virtual EBlackboardCompare::Type Compare(const uint8* MemoryBlockA, const uint8* MemoryBlockB) const;
 
diff --git a/Engine/Source/Runtime/AIModule/Classes/BehaviorTree/Blackboard/BlackboardKeyType_Rotator.h b/Engine/Source/Runtime/AIModule/Classes/BehaviorTree/Blackboard/BlackboardKeyType_Rotator.h
--- a/Engine/Source/Runtime/AIModule/Classes/BehaviorTree/Blackboard/BlackboardKeyType_Rotator.h
+++ b/Engine/Source/Runtime/AIModule/Classes/BehaviorTree/Blackboard/BlackboardKeyType_Rotator.h
@@ -13,6 +13,7 @@ class AIMODULE_API UBlackboardKeyType_Rotator : public UBlackboardKeyType
 	virtual void Initialize(uint8* RawData) const override;
 	virtual FString DescribeValue(const uint8* RawData) const override;
 	virtual bool GetRotation(const uint8* RawData, FRotator& Rotation) const override;
+	virtual bool Clear(uint8* MemoryBlock) const override;
 	virtual EBlackboardCompare::Type Compare(const uint8* MemoryBlockA, const uint8* MemoryBlockB) const override;
 	virtual bool TestBasicOperation(const uint8* MemoryBlock, EBasicKeyOperation::Type Op) const override;
 };
diff --git a/Engine/Source/Runtime/AIModule/Classes/BehaviorTree/Blackboard/BlackboardKeyType_Vector.h b/Engine/Source/Runtime/AIModule/Classes/BehaviorTree/Blackboard/BlackboardKeyType_Vector.h
--- a/Engine/Source/Runtime/AIModule/Classes/BehaviorTree/Blackboard/BlackboardKeyType_Vector.h
+++ b/Engine/Source/Runtime/AIModule/Classes/BehaviorTree/Blackboard/BlackboardKeyType_Vector.h
@@ -14,6 +14,7 @@ class AIMODULE_API UBlackboardKeyType_Vector : public UBlackboardKeyType
 	virtual void Initialize(uint8* RawData) const override;
 	virtual FString DescribeValue(const uint8* RawData) const override;
 	virtual bool GetLocation(const uint8* RawData, FVector& Location) const override;
+	virtual bool Clear(uint8* RawData) const override;
 	virtual EBlackboardCompare::Type Compare(const uint8* MemoryBlockA, const uint8* MemoryBlockB) const override;
 	virtual bool TestBasicOperation(const uint8* MemoryBlock, EBasicKeyOperation::Type Op) const override;
 };
diff --git a/Engine/Source/Runtime/AIModule/Classes/BehaviorTree/BlackboardComponent.h b/Engine/Source/Runtime/AIModule/Classes/BehaviorTree/BlackboardComponent.h
--- a/Engine/Source/Runtime/AIModule/Classes/BehaviorTree/BlackboardComponent.h
+++ b/Engine/Source/Runtime/AIModule/Classes/BehaviorTree/BlackboardComponent.h
@@ -120,7 +120,7 @@ class AIMODULE_API UBlackboardComponent : public UActorComponent
 
 	UFUNCTION(BlueprintCallable, Category="AI|Components|Blackboard")
 	FRotator GetValueAsRotator(const FName& KeyName) const;
-	FRotator GetValueAsRotator(uint8 KeyID) const;
+	FRotator GetValueAsRotator(FBlackboard::FKey KeyID) const;
 
 	UFUNCTION(BlueprintCallable, Category="AI|Components|Blackboard")
 	void SetValueAsObject(const FName& KeyName, UObject* ObjectValue);
@@ -169,11 +169,11 @@ class AIMODULE_API UBlackboardComponent : public UActorComponent
 
 	UFUNCTION(BlueprintCallable, Category="AI|Components|Blackboard")
 	void SetValueAsRotator(const FName& KeyName, FRotator VectorValue);
-	void SetValueAsRotator(uint8 KeyID, const FRotator& VectorValue);
+	void SetValueAsRotator(FBlackboard::FKey KeyID, const FRotator& VectorValue);
 
 	UFUNCTION(BlueprintCallable, Category="AI|Components|Blackboard")
 	void ClearValueAsRotator(const FName& KeyName);
-	void ClearValueAsRotator(uint8 KeyID);
+	void ClearValueAsRotator(FBlackboard::FKey KeyID);
 
 	/** return false if call failed (most probably no such entry in BB) */
 	UFUNCTION(BlueprintCallable, Category="AI|Components|Blackboard")
@@ -185,6 +185,10 @@ class AIMODULE_API UBlackboardComponent : public UActorComponent
 	bool GetRotationFromEntry(const FName& KeyName, FRotator& ResultRotation) const;
 	bool GetRotationFromEntry(FBlackboard::FKey KeyID, FRotator& ResultRotation) const;
 
+	UFUNCTION(BlueprintCallable, Category = "AI|Components|Blackboard")
+	void ClearValue(const FName& KeyName);
+	void ClearValue(FBlackboard::FKey KeyID);
+
 	/** get pointer to raw data for given key */
 	FORCEINLINE uint8* GetKeyRawData(const FName& KeyName) { return GetKeyRawData(GetKeyID(KeyName)); }
 	FORCEINLINE uint8* GetKeyRawData(FBlackboard::FKey KeyID) { return ValueMemory.Num() && ValueOffsets.IsValidIndex(KeyID) ? (ValueMemory.GetData() + ValueOffsets[KeyID]) : NULL; }
diff --git a/Engine/Source/Runtime/AIModule/Private/BehaviorTree/BTFunctionLibrary.cpp b/Engine/Source/Runtime/AIModule/Private/BehaviorTree/BTFunctionLibrary.cpp
--- a/Engine/Source/Runtime/AIModule/Private/BehaviorTree/BTFunctionLibrary.cpp
+++ b/Engine/Source/Runtime/AIModule/Private/BehaviorTree/BTFunctionLibrary.cpp
@@ -24,77 +24,84 @@ UBlackboardComponent* UBTFunctionLibrary::GetOwnersBlackboard(UBTNode* NodeOwner
 	return GetBlackboard(*NodeOwner);
 }
 
-UObject* UBTFunctionLibrary::GetBlackboardValueAsObject(UBTNode* NodeOwner, const struct FBlackboardKeySelector& Key)
+UObject* UBTFunctionLibrary::GetBlackboardValueAsObject(UBTNode* NodeOwner, const FBlackboardKeySelector& Key)
 {
 	check(NodeOwner != NULL);
 	const UBlackboardComponent* BlackboardComp = GetBlackboard(*NodeOwner);
 	return BlackboardComp ? BlackboardComp->GetValueAsObject(Key.SelectedKeyName) : NULL;
 }
 
-AActor* UBTFunctionLibrary::GetBlackboardValueAsActor(UBTNode* NodeOwner, const struct FBlackboardKeySelector& Key)
+AActor* UBTFunctionLibrary::GetBlackboardValueAsActor(UBTNode* NodeOwner, const FBlackboardKeySelector& Key)
 {
 	check(NodeOwner != NULL);
 	const UBlackboardComponent* BlackboardComp = GetBlackboard(*NodeOwner);
 	return BlackboardComp ? Cast<AActor>(BlackboardComp->GetValueAsObject(Key.SelectedKeyName)) : NULL;
 }
 
-UClass* UBTFunctionLibrary::GetBlackboardValueAsClass(UBTNode* NodeOwner, const struct FBlackboardKeySelector& Key)
+UClass* UBTFunctionLibrary::GetBlackboardValueAsClass(UBTNode* NodeOwner, const FBlackboardKeySelector& Key)
 {
 	check(NodeOwner != NULL);
 	const UBlackboardComponent* BlackboardComp = GetBlackboard(*NodeOwner);
 	return BlackboardComp ? BlackboardComp->GetValueAsClass(Key.SelectedKeyName) : NULL;
 }
 
-uint8 UBTFunctionLibrary::GetBlackboardValueAsEnum(UBTNode* NodeOwner, const struct FBlackboardKeySelector& Key)
+uint8 UBTFunctionLibrary::GetBlackboardValueAsEnum(UBTNode* NodeOwner, const FBlackboardKeySelector& Key)
 {
 	check(NodeOwner != NULL);
 	const UBlackboardComponent* BlackboardComp = GetBlackboard(*NodeOwner);
 	return BlackboardComp ? BlackboardComp->GetValueAsEnum(Key.SelectedKeyName) : 0;
 }
 
-int32 UBTFunctionLibrary::GetBlackboardValueAsInt(UBTNode* NodeOwner, const struct FBlackboardKeySelector& Key)
+int32 UBTFunctionLibrary::GetBlackboardValueAsInt(UBTNode* NodeOwner, const FBlackboardKeySelector& Key)
 {
 	check(NodeOwner != NULL);
 	const UBlackboardComponent* BlackboardComp = GetBlackboard(*NodeOwner);
 	return BlackboardComp ? BlackboardComp->GetValueAsInt(Key.SelectedKeyName) : 0;
 }
 
-float UBTFunctionLibrary::GetBlackboardValueAsFloat(UBTNode* NodeOwner, const struct FBlackboardKeySelector& Key)
+float UBTFunctionLibrary::GetBlackboardValueAsFloat(UBTNode* NodeOwner, const FBlackboardKeySelector& Key)
 {
 	check(NodeOwner != NULL);
 	const UBlackboardComponent* BlackboardComp = GetBlackboard(*NodeOwner);
 	return BlackboardComp ? BlackboardComp->GetValueAsFloat(Key.SelectedKeyName) : 0.0f;
 }
 
-bool UBTFunctionLibrary::GetBlackboardValueAsBool(UBTNode* NodeOwner, const struct FBlackboardKeySelector& Key)
+bool UBTFunctionLibrary::GetBlackboardValueAsBool(UBTNode* NodeOwner, const FBlackboardKeySelector& Key)
 {
 	check(NodeOwner != NULL);
 	const UBlackboardComponent* BlackboardComp = GetBlackboard(*NodeOwner);
 	return BlackboardComp ? BlackboardComp->GetValueAsBool(Key.SelectedKeyName) : false;
 }
 
-FString UBTFunctionLibrary::GetBlackboardValueAsString(UBTNode* NodeOwner, const struct FBlackboardKeySelector& Key)
+FString UBTFunctionLibrary::GetBlackboardValueAsString(UBTNode* NodeOwner, const FBlackboardKeySelector& Key)
 {
 	check(NodeOwner != NULL);
 	const UBlackboardComponent* BlackboardComp = GetBlackboard(*NodeOwner);
 	return BlackboardComp ? BlackboardComp->GetValueAsString(Key.SelectedKeyName) : FString();
 }
 
-FName UBTFunctionLibrary::GetBlackboardValueAsName(UBTNode* NodeOwner, const struct FBlackboardKeySelector& Key)
+FName UBTFunctionLibrary::GetBlackboardValueAsName(UBTNode* NodeOwner, const FBlackboardKeySelector& Key)
 {
 	check(NodeOwner != NULL);
 	const UBlackboardComponent* BlackboardComp = GetBlackboard(*NodeOwner);
 	return BlackboardComp ? BlackboardComp->GetValueAsName(Key.SelectedKeyName) : NAME_None;
 }
 
-FVector UBTFunctionLibrary::GetBlackboardValueAsVector(UBTNode* NodeOwner, const struct FBlackboardKeySelector& Key)
+FVector UBTFunctionLibrary::GetBlackboardValueAsVector(UBTNode* NodeOwner, const FBlackboardKeySelector& Key)
 {
 	check(NodeOwner != NULL);
 	const UBlackboardComponent* BlackboardComp = GetBlackboard(*NodeOwner);
 	return BlackboardComp ? BlackboardComp->GetValueAsVector(Key.SelectedKeyName) : FVector::ZeroVector;
 }
 
-void UBTFunctionLibrary::SetBlackboardValueAsObject(UBTNode* NodeOwner, const struct FBlackboardKeySelector& Key, UObject* Value)
+FRotator UBTFunctionLibrary::GetBlackboardValueAsRotator(UBTNode* NodeOwner, const FBlackboardKeySelector& Key)
+{
+	check(NodeOwner != NULL);
+	const UBlackboardComponent* BlackboardComp = GetBlackboard(*NodeOwner);
+	return BlackboardComp ? BlackboardComp->GetValueAsRotator(Key.SelectedKeyName) : FRotator::ZeroRotator;
+}
+
+void UBTFunctionLibrary::SetBlackboardValueAsObject(UBTNode* NodeOwner, const FBlackboardKeySelector& Key, UObject* Value)
 {
 	check(NodeOwner != NULL);
 	UBlackboardComponent* BlackboardComp = GetBlackboard(*NodeOwner);
@@ -104,7 +111,7 @@ void UBTFunctionLibrary::SetBlackboardValueAsObject(UBTNode* NodeOwner, const st
 	}
 }
 
-void UBTFunctionLibrary::SetBlackboardValueAsClass(UBTNode* NodeOwner, const struct FBlackboardKeySelector& Key, UClass* Value)
+void UBTFunctionLibrary::SetBlackboardValueAsClass(UBTNode* NodeOwner, const FBlackboardKeySelector& Key, UClass* Value)
 {
 	check(NodeOwner != NULL);
 	UBlackboardComponent* BlackboardComp = GetBlackboard(*NodeOwner);
@@ -114,7 +121,7 @@ void UBTFunctionLibrary::SetBlackboardValueAsClass(UBTNode* NodeOwner, const str
 	}
 }
 
-void UBTFunctionLibrary::SetBlackboardValueAsEnum(UBTNode* NodeOwner, const struct FBlackboardKeySelector& Key, uint8 Value)
+void UBTFunctionLibrary::SetBlackboardValueAsEnum(UBTNode* NodeOwner, const FBlackboardKeySelector& Key, uint8 Value)
 {
 	check(NodeOwner != NULL);
 	UBlackboardComponent* BlackboardComp = GetBlackboard(*NodeOwner);
@@ -124,7 +131,7 @@ void UBTFunctionLibrary::SetBlackboardValueAsEnum(UBTNode* NodeOwner, const stru
 	}
 }
 
-void UBTFunctionLibrary::SetBlackboardValueAsInt(UBTNode* NodeOwner, const struct FBlackboardKeySelector& Key, int32 Value)
+void UBTFunctionLibrary::SetBlackboardValueAsInt(UBTNode* NodeOwner, const FBlackboardKeySelector& Key, int32 Value)
 {
 	check(NodeOwner != NULL);
 	UBlackboardComponent* BlackboardComp = GetBlackboard(*NodeOwner);
@@ -134,7 +141,7 @@ void UBTFunctionLibrary::SetBlackboardValueAsInt(UBTNode* NodeOwner, const struc
 	}
 }
 
-void UBTFunctionLibrary::SetBlackboardValueAsFloat(UBTNode* NodeOwner, const struct FBlackboardKeySelector& Key, float Value)
+void UBTFunctionLibrary::SetBlackboardValueAsFloat(UBTNode* NodeOwner, const FBlackboardKeySelector& Key, float Value)
 {
 	check(NodeOwner != NULL);
 	UBlackboardComponent* BlackboardComp = GetBlackboard(*NodeOwner);
@@ -144,7 +151,7 @@ void UBTFunctionLibrary::SetBlackboardValueAsFloat(UBTNode* NodeOwner, const str
 	}
 }
 
-void UBTFunctionLibrary::SetBlackboardValueAsBool(UBTNode* NodeOwner, const struct FBlackboardKeySelector& Key, bool Value)
+void UBTFunctionLibrary::SetBlackboardValueAsBool(UBTNode* NodeOwner, const FBlackboardKeySelector& Key, bool Value)
 {
 	check(NodeOwner != NULL);
 	UBlackboardComponent* BlackboardComp = GetBlackboard(*NodeOwner);
@@ -154,7 +161,7 @@ void UBTFunctionLibrary::SetBlackboardValueAsBool(UBTNode* NodeOwner, const stru
 	}
 }
 
-void UBTFunctionLibrary::SetBlackboardValueAsString(UBTNode* NodeOwner, const struct FBlackboardKeySelector& Key, FString Value)
+void UBTFunctionLibrary::SetBlackboardValueAsString(UBTNode* NodeOwner, const FBlackboardKeySelector& Key, FString Value)
 {
 	check(NodeOwner != NULL);
 	UBlackboardComponent* BlackboardComp = GetBlackboard(*NodeOwner);
@@ -164,7 +171,7 @@ void UBTFunctionLibrary::SetBlackboardValueAsString(UBTNode* NodeOwner, const st
 	}
 }
 
-void UBTFunctionLibrary::SetBlackboardValueAsName(UBTNode* NodeOwner, const struct FBlackboardKeySelector& Key, FName Value)
+void UBTFunctionLibrary::SetBlackboardValueAsName(UBTNode* NodeOwner, const FBlackboardKeySelector& Key, FName Value)
 {
 	check(NodeOwner != NULL);
 	UBlackboardComponent* BlackboardComp = GetBlackboard(*NodeOwner);
@@ -174,7 +181,7 @@ void UBTFunctionLibrary::SetBlackboardValueAsName(UBTNode* NodeOwner, const stru
 	}
 }
 
-void UBTFunctionLibrary::SetBlackboardValueAsVector(UBTNode* NodeOwner, const struct FBlackboardKeySelector& Key, FVector Value)
+void UBTFunctionLibrary::SetBlackboardValueAsVector(UBTNode* NodeOwner, const FBlackboardKeySelector& Key, FVector Value)
 {
 	check(NodeOwner != NULL);
 	UBlackboardComponent* BlackboardComp = GetBlackboard(*NodeOwner);
@@ -184,7 +191,7 @@ void UBTFunctionLibrary::SetBlackboardValueAsVector(UBTNode* NodeOwner, const st
 	}
 }
 
-void UBTFunctionLibrary::ClearBlackboardValueAsVector(UBTNode* NodeOwner, const struct FBlackboardKeySelector& Key)
+void UBTFunctionLibrary::ClearBlackboardValueAsVector(UBTNode* NodeOwner, const FBlackboardKeySelector& Key)
 {
 	check(NodeOwner != NULL);
 	UBlackboardComponent* BlackboardComp = GetBlackboard(*NodeOwner);
@@ -194,6 +201,26 @@ void UBTFunctionLibrary::ClearBlackboardValueAsVector(UBTNode* NodeOwner, const
 	}
 }
 
+void UBTFunctionLibrary::SetBlackboardValueAsRotator(UBTNode* NodeOwner, const FBlackboardKeySelector& Key, FRotator Value)
+{
+	check(NodeOwner != NULL);
+	UBlackboardComponent* BlackboardComp = GetBlackboard(*NodeOwner);
+	if (BlackboardComp != NULL)
+	{
+		BlackboardComp->SetValueAsRotator(Key.SelectedKeyName, Value);
+	}
+}
+
+void UBTFunctionLibrary::ClearBlackboardValue(UBTNode* NodeOwner, const FBlackboardKeySelector& Key)
+{
+	check(NodeOwner != NULL);
+	UBlackboardComponent* BlackboardComp = GetBlackboard(*NodeOwner);
+	if (BlackboardComp != NULL)
+	{
+		BlackboardComp->ClearValue(Key.SelectedKeyName);
+	}
+}
+
 void UBTFunctionLibrary::StartUsingExternalEvent(UBTNode* NodeOwner, AActor* OwningActor)
 {
 	// deprecated, not removed yet
diff --git a/Engine/Source/Runtime/AIModule/Private/BehaviorTree/Blackboard/BlackboardKeyType.cpp b/Engine/Source/Runtime/AIModule/Private/BehaviorTree/Blackboard/BlackboardKeyType.cpp
--- a/Engine/Source/Runtime/AIModule/Private/BehaviorTree/Blackboard/BlackboardKeyType.cpp
+++ b/Engine/Source/Runtime/AIModule/Private/BehaviorTree/Blackboard/BlackboardKeyType.cpp
@@ -45,6 +45,21 @@ bool UBlackboardKeyType::GetRotation(const uint8* MemoryBlock, FRotator& Rotatio
 	return false;
 }
 
+bool UBlackboardKeyType::Clear(uint8* MemoryBlock) const
+{
+	for (int32 ByteIndex = 0; ByteIndex < GetValueSize(); ++ByteIndex)
+	{
+		if (MemoryBlock[ByteIndex] != uint8(0))
+		{
+			FMemory::Memzero(MemoryBlock, GetValueSize());
+			return true;
+		}		
+	}
+
+	// value already "0", return false to indicate it had not changed
+	return false;
+}
+
 EBlackboardCompare::Type UBlackboardKeyType::Compare(const uint8* MemoryBlockA, const uint8* MemoryBlockB) const
 { 
 	return MemoryBlockA == MemoryBlockB ? EBlackboardCompare::Equal : EBlackboardCompare::NotEqual;
diff --git a/Engine/Source/Runtime/AIModule/Private/BehaviorTree/Blackboard/BlackboardKeyType_Rotator.cpp b/Engine/Source/Runtime/AIModule/Private/BehaviorTree/Blackboard/BlackboardKeyType_Rotator.cpp
--- a/Engine/Source/Runtime/AIModule/Private/BehaviorTree/Blackboard/BlackboardKeyType_Rotator.cpp
+++ b/Engine/Source/Runtime/AIModule/Private/BehaviorTree/Blackboard/BlackboardKeyType_Rotator.cpp
@@ -19,6 +19,11 @@ bool UBlackboardKeyType_Rotator::SetValue(uint8* RawData, const FRotator& Value)
 	return SetValueInMemory<FRotator>(RawData, Value);
 }
 
+bool UBlackboardKeyType_Rotator::Clear(uint8* RawData) const 
+{
+	return SetValueInMemory<FRotator>(RawData, FAISystem::InvalidRotation);
+}
+
 FString UBlackboardKeyType_Rotator::DescribeValue(const uint8* RawData) const
 {
 	const FRotator Rotation = GetValue(RawData);
diff --git a/Engine/Source/Runtime/AIModule/Private/BehaviorTree/Blackboard/BlackboardKeyType_Vector.cpp b/Engine/Source/Runtime/AIModule/Private/BehaviorTree/Blackboard/BlackboardKeyType_Vector.cpp
--- a/Engine/Source/Runtime/AIModule/Private/BehaviorTree/Blackboard/BlackboardKeyType_Vector.cpp
+++ b/Engine/Source/Runtime/AIModule/Private/BehaviorTree/Blackboard/BlackboardKeyType_Vector.cpp
@@ -19,6 +19,11 @@ bool UBlackboardKeyType_Vector::SetValue(uint8* RawData, const FVector& Value)
 	return SetValueInMemory<FVector>(RawData, Value);
 }
 
+bool UBlackboardKeyType_Vector::Clear(uint8* RawData) const 
+{
+	return SetValueInMemory<FVector>(RawData, FAISystem::InvalidLocation);
+}
+
 FString UBlackboardKeyType_Vector::DescribeValue(const uint8* RawData) const
 {
 	const FVector Location = GetValue(RawData);
diff --git a/Engine/Source/Runtime/AIModule/Private/BehaviorTree/BlackboardComponent.cpp b/Engine/Source/Runtime/AIModule/Private/BehaviorTree/BlackboardComponent.cpp
--- a/Engine/Source/Runtime/AIModule/Private/BehaviorTree/BlackboardComponent.cpp
+++ b/Engine/Source/Runtime/AIModule/Private/BehaviorTree/BlackboardComponent.cpp
@@ -194,7 +194,7 @@ void UBlackboardComponent::NotifyObservers(FBlackboard::FKey KeyID) const
 	}
 }
 
-bool UBlackboardComponent::IsCompatibleWith(class UBlackboardData* TestAsset) const
+bool UBlackboardComponent::IsCompatibleWith(UBlackboardData* TestAsset) const
 {
 	for (UBlackboardData* It = BlackboardAsset; It; It = It->Parent)
 	{
@@ -212,10 +212,9 @@ bool UBlackboardComponent::IsCompatibleWith(class UBlackboardData* TestAsset) co
 	return false;
 }
 
-EBlackboardCompare::Type UBlackboardComponent::CompareKeyValues(TSubclassOf<class UBlackboardKeyType> KeyType, FBlackboard::FKey KeyA, FBlackboard::FKey KeyB) const
-{
-	UBlackboardKeyType* KeyCDO = KeyType->GetDefaultObject<UBlackboardKeyType>();
-	return KeyCDO->Compare(GetKeyRawData(KeyA), GetKeyRawData(KeyB));
+EBlackboardCompare::Type UBlackboardComponent::CompareKeyValues(TSubclassOf<UBlackboardKeyType> KeyType, FBlackboard::FKey KeyA, FBlackboard::FKey KeyB) const
+{	
+	return GetDefault<UBlackboardKeyType>()->Compare(GetKeyRawData(KeyA), GetKeyRawData(KeyB));
 }
 
 FString UBlackboardComponent::GetDebugInfoString(EBlackboardDescription::Type Mode) const 
@@ -304,7 +303,7 @@ FString UBlackboardComponent::DescribeKeyValue(FBlackboard::FKey KeyID, EBlackbo
 
 #if ENABLE_VISUAL_LOG
 
-void UBlackboardComponent::DescribeSelfToVisLog(struct FVisualLogEntry* Snapshot) const
+void UBlackboardComponent::DescribeSelfToVisLog(FVisualLogEntry* Snapshot) const
 {
 	FVisualLogEntry::FStatusCategory Category;
 	Category.Category = FString::Printf(TEXT("Blackboard (asset: %s)"), *GetNameSafe(BlackboardAsset));
@@ -720,7 +719,7 @@ void UBlackboardComponent::SetValueAsVector(FBlackboard::FKey KeyID, const FVect
 
 void UBlackboardComponent::SetValueAsRotator(const FName& KeyName, FRotator RotatorValue)
 {
-	const uint8 KeyID = GetKeyID(KeyName);
+	const FBlackboard::FKey KeyID = GetKeyID(KeyName);
 	SetValueAsRotator(KeyID, RotatorValue);
 }
 
@@ -785,7 +784,7 @@ bool UBlackboardComponent::IsVectorValueSet(FBlackboard::FKey KeyID) const
 
 void UBlackboardComponent::ClearValueAsRotator(const FName& KeyName)
 {
-	const uint8 KeyID = GetKeyID(KeyName);
+	const FBlackboard::FKey KeyID = GetKeyID(KeyName);
 	ClearValueAsRotator(KeyID);
 }
 
@@ -807,6 +806,29 @@ void UBlackboardComponent::ClearValueAsRotator(uint8 KeyID)
 	}
 }
 
+void UBlackboardComponent::ClearValue(const FName& KeyName)
+{
+	const FBlackboard::FKey KeyID = GetKeyID(KeyName);
+	ClearValue(KeyID);
+}
+
+void UBlackboardComponent::ClearValue(FBlackboard::FKey KeyID)
+{
+	uint8* RawData = GetKeyRawData(KeyID);
+	if (RawData)
+	{
+		TSubclassOf<UBlackboardKeyType> ValueType = GetKeyType(KeyID);
+		if (ValueType)
+		{
+			// c-cast is ok here since ValueType is of at least UBlackboardKeyType and it's default object has to be UBlackboardKeyType instance
+			const bool bChanged = GetDefault<UBlackboardKeyType>()->Clear(RawData);
+			if (bChanged)
+			{
+				NotifyObservers(KeyID);
+			}
+		}
+	}
+}
 
 bool UBlackboardComponent::GetLocationFromEntry(const FName& KeyName, FVector& ResultLocation) const
 {
