fixed initialization order of blackboard's parent chain

diff --git a/Engine/Source/Runtime/AIModule/Classes/BehaviorTree/BlackboardComponent.h b/Engine/Source/Runtime/AIModule/Classes/BehaviorTree/BlackboardComponent.h
--- a/Engine/Source/Runtime/AIModule/Classes/BehaviorTree/BlackboardComponent.h
+++ b/Engine/Source/Runtime/AIModule/Classes/BehaviorTree/BlackboardComponent.h
@@ -239,6 +239,9 @@ protected:
 
 	/** notifies behavior tree decorators about change in blackboard */
 	void NotifyObservers(FBlackboard::FKey KeyID) const;
+
+	/** initializes parent chain in asset */
+	void InitializeParentChain(UBlackboardData* NewAsset);
 };
 
 //////////////////////////////////////////////////////////////////////////
diff --git a/Engine/Source/Runtime/AIModule/Classes/BehaviorTree/BlackboardData.h b/Engine/Source/Runtime/AIModule/Classes/BehaviorTree/BlackboardData.h
--- a/Engine/Source/Runtime/AIModule/Classes/BehaviorTree/BlackboardData.h
+++ b/Engine/Source/Runtime/AIModule/Classes/BehaviorTree/BlackboardData.h
@@ -109,6 +109,9 @@ class AIMODULE_API UBlackboardData : public UDataAsset
 	/** updates parent key cache for editor */
 	void UpdateParentKeys();
 
+	/** forces update of FirstKeyID, which depends on parent chain */
+	void UpdateKeyIDs();
+
 protected:
 
 	enum EKeyLookupMode
diff --git a/Engine/Source/Runtime/AIModule/Private/BehaviorTree/BlackboardComponent.cpp b/Engine/Source/Runtime/AIModule/Private/BehaviorTree/BlackboardComponent.cpp
--- a/Engine/Source/Runtime/AIModule/Private/BehaviorTree/BlackboardComponent.cpp
+++ b/Engine/Source/Runtime/AIModule/Private/BehaviorTree/BlackboardComponent.cpp
@@ -59,6 +59,15 @@ struct FBlackboardInitializationData
 	};
 };
 
+void UBlackboardComponent::InitializeParentChain(UBlackboardData* NewAsset)
+{
+	if (NewAsset)
+	{
+		InitializeParentChain(NewAsset->Parent);
+		NewAsset->UpdateKeyIDs();
+	}
+}
+
 void UBlackboardComponent::InitializeBlackboard(UBlackboardData* NewAsset)
 {
 	// if we re-initialize with the same asset then there's no point
@@ -76,6 +85,8 @@ void UBlackboardComponent::InitializeBlackboard(UBlackboardData* NewAsset)
 	{
 		if (BlackboardAsset->IsValid())
 		{
+			InitializeParentChain(BlackboardAsset);
+
 			TArray<FBlackboardInitializationData> InitList;
 			const int32 NumKeys = BlackboardAsset->GetNumKeys();
 			InitList.Reserve(NumKeys);
diff --git a/Engine/Source/Runtime/AIModule/Private/BehaviorTree/BlackboardData.cpp b/Engine/Source/Runtime/AIModule/Private/BehaviorTree/BlackboardData.cpp
--- a/Engine/Source/Runtime/AIModule/Private/BehaviorTree/BlackboardData.cpp
+++ b/Engine/Source/Runtime/AIModule/Private/BehaviorTree/BlackboardData.cpp
@@ -185,6 +185,11 @@ void UBlackboardData::UpdateParentKeys()
 	}
 #endif // WITH_EDITORONLY_DATA
 
-	FirstKeyID = Parent ? Parent->GetNumKeys() : 0;
+	UpdateKeyIDs();
 	OnUpdateKeys.Broadcast(this);
 }
+
+void UBlackboardData::UpdateKeyIDs()
+{
+	FirstKeyID = Parent ? Parent->GetNumKeys() : 0;
+}
