fixed crowd pathfollowing going to center point of custom navlink instead using edge portals

diff --git a/Engine/Source/Runtime/AIModule/Private/Navigation/CrowdFollowingComponent.cpp b/Engine/Source/Runtime/AIModule/Private/Navigation/CrowdFollowingComponent.cpp
--- a/Engine/Source/Runtime/AIModule/Private/Navigation/CrowdFollowingComponent.cpp
+++ b/Engine/Source/Runtime/AIModule/Private/Navigation/CrowdFollowingComponent.cpp
@@ -551,7 +551,19 @@ void UCrowdFollowingComponent::SetMoveSegment(int32 SegmentStartIndex)
 
 		const int32 PathPartSize = 15;
 		const int32 LastPolyIdx = NavMeshPath->PathCorridor.Num() - 1;
-		const int32 PathPartEndIdx = FMath::Min(PathStartIndex + PathPartSize, LastPolyIdx);
+		int32 PathPartEndIdx = FMath::Min(PathStartIndex + PathPartSize, LastPolyIdx);
+
+		FVector PtA, PtB;
+		const bool bStartIsNavLink = RecastNavData->GetLinkEndPoints(NavMeshPath->PathCorridor[PathStartIndex], PtA, PtB);
+		const bool bEndIsNavLink = RecastNavData->GetLinkEndPoints(NavMeshPath->PathCorridor[PathPartEndIdx], PtA, PtB);
+		if (bStartIsNavLink)
+		{
+			PathStartIndex = FMath::Max(0, PathStartIndex - 1);
+		}
+		if (bEndIsNavLink)
+		{
+			PathPartEndIdx = FMath::Max(0, PathPartEndIdx - 1);
+		}
 
 		bFinalPathPart = (PathPartEndIdx == LastPolyIdx);
 		if (!bFinalPathPart)
