Fix for BT aux nodes not being cleared when BT finishes merged over from dev branch #UE4

diff --git a/Engine/Source/Runtime/AIModule/Private/BehaviorTree/BehaviorTreeComponent.cpp b/Engine/Source/Runtime/AIModule/Private/BehaviorTree/BehaviorTreeComponent.cpp
--- a/Engine/Source/Runtime/AIModule/Private/BehaviorTree/BehaviorTreeComponent.cpp
+++ b/Engine/Source/Runtime/AIModule/Private/BehaviorTree/BehaviorTreeComponent.cpp
@@ -269,6 +269,10 @@ void UBehaviorTreeComponent::OnTreeFinished()
 		TopInstance.ActiveNode = NULL;
 		TopInstance.ActiveNodeType = EBTActiveNode::Composite;
 
+		// make sure that all active aux nodes will be removed
+		// root level services are being handled on applying search data
+		UnregisterAuxNodesUpTo(FBTNodeIndex(0, 0));
+
 		// result doesn't really matter, root node will be reset and start iterating child nodes from scratch
 		// although it shouldn't be set to Aborted, as it has special meaning in RequestExecution (switch to higher priority)
 		RequestExecution(TopInstance.RootNode, 0, TopInstance.RootNode, 0, EBTNodeResult::InProgress);
