Changed the way BP-implemented BT nodes get instantiated #UE4

- related issues: UEBP-15; UEPLAT-204

diff --git a/Engine/Source/Runtime/AIModule/Private/BehaviorTree/BTNode.cpp b/Engine/Source/Runtime/AIModule/Private/BehaviorTree/BTNode.cpp
--- a/Engine/Source/Runtime/AIModule/Private/BehaviorTree/BTNode.cpp
+++ b/Engine/Source/Runtime/AIModule/Private/BehaviorTree/BTNode.cpp
@@ -74,7 +74,7 @@ void UBTNode::InitializeInSubtree(UBehaviorTreeComponent* OwnerComp, uint8* Node
 		UBTNode* NodeInstance = OwnerComp->NodeInstances.IsValidIndex(NextInstancedIndex) ? OwnerComp->NodeInstances[NextInstancedIndex] : NULL;
 		if (NodeInstance == NULL)
 		{
-			NodeInstance = DuplicateObject<UBTNode>(this, OwnerComp);
+			NodeInstance = ConstructObject<UBTNode>(GetClass(), OwnerComp, GetFName(), RF_NoFlags, (UObject*)(this));
 			NodeInstance->InitializeNode(GetParentNode(), GetExecutionIndex(), GetMemoryOffset(), GetTreeDepth());
 			NodeInstance->bIsInstanced = true;
 
