- Non-unity fix up found while testing

diff --git a/Engine/Source/Runtime/AIModule/Classes/Perception/AIPerceptionTypes.h b/Engine/Source/Runtime/AIModule/Classes/Perception/AIPerceptionTypes.h
--- a/Engine/Source/Runtime/AIModule/Classes/Perception/AIPerceptionTypes.h
+++ b/Engine/Source/Runtime/AIModule/Classes/Perception/AIPerceptionTypes.h
@@ -3,6 +3,7 @@
 #pragma once
 
 #include "GenericTeamAgentInterface.h"
+#include "AITypes.h"
 #include "AIPerceptionTypes.generated.h"
 
 class UAIPerceptionComponent;
