#UE4 - Added blackboard key based version of behavior tree Wait task

diff --git a/Engine/Source/Runtime/AIModule/Classes/BehaviorTree/Tasks/BTTask_Wait.h b/Engine/Source/Runtime/AIModule/Classes/BehaviorTree/Tasks/BTTask_Wait.h
--- a/Engine/Source/Runtime/AIModule/Classes/BehaviorTree/Tasks/BTTask_Wait.h
+++ b/Engine/Source/Runtime/AIModule/Classes/BehaviorTree/Tasks/BTTask_Wait.h
@@ -19,7 +19,7 @@ class AIMODULE_API UBTTask_Wait : public UBTTaskNode
 	GENERATED_UCLASS_BODY()
 
 	/** wait time in seconds */
-	UPROPERTY(Category=Node, EditAnywhere)
+	UPROPERTY(Category=Wait, EditAnywhere)
 	float WaitTime;
 
 	virtual EBTNodeResult::Type ExecuteTask(class UBehaviorTreeComponent* OwnerComp, uint8* NodeMemory) override;
diff --git a/Engine/Source/Runtime/AIModule/Classes/BehaviorTree/Tasks/BTTask_WaitBlackboardTime.h b/Engine/Source/Runtime/AIModule/Classes/BehaviorTree/Tasks/BTTask_WaitBlackboardTime.h
new file mode 100644
--- /dev/null
+++ b/Engine/Source/Runtime/AIModule/Classes/BehaviorTree/Tasks/BTTask_WaitBlackboardTime.h
@@ -0,0 +1,38 @@
+// Copyright 1998-2014 Epic Games, Inc. All Rights Reserved.
+#pragma once
+
+#include "BehaviorTree/Tasks/BTTask_Wait.h"
+#include "BTTask_WaitBlackboardTime.generated.h"
+
+/**
+ * Wait task node.
+ * Wait for the time specified by a Blackboard key when executed.
+ */
+UCLASS(hidecategories=Wait)
+class AIMODULE_API UBTTask_WaitBlackboardTime : public UBTTask_Wait
+{
+	GENERATED_UCLASS_BODY()
+
+	virtual void InitializeFromAsset(UBehaviorTree& Asset) override;
+	virtual EBTNodeResult::Type ExecuteTask(class UBehaviorTreeComponent* OwnerComp, uint8* NodeMemory) override;
+	virtual FString GetStaticDescription() const override;
+
+	/** get name of selected blackboard key */
+	FName GetSelectedBlackboardKey() const;
+
+
+protected:
+
+	/** blackboard key selector */
+	UPROPERTY(EditAnywhere, Category=Blackboard)
+	struct FBlackboardKeySelector BlackboardKey;
+	
+};
+
+//////////////////////////////////////////////////////////////////////////
+// Inlines
+
+FORCEINLINE FName UBTTask_WaitBlackboardTime::GetSelectedBlackboardKey() const
+{
+	return BlackboardKey.SelectedKeyName;
+}
diff --git a/Engine/Source/Runtime/AIModule/Private/BehaviorTree/Tasks/BTTask_WaitBlackboardTime.cpp b/Engine/Source/Runtime/AIModule/Private/BehaviorTree/Tasks/BTTask_WaitBlackboardTime.cpp
new file mode 100644
--- /dev/null
+++ b/Engine/Source/Runtime/AIModule/Private/BehaviorTree/Tasks/BTTask_WaitBlackboardTime.cpp
@@ -0,0 +1,37 @@
+// Copyright 1998-2014 Epic Games, Inc. All Rights Reserved.
+
+#include "AIModulePrivate.h"
+#include "BehaviorTree/BlackboardComponent.h"
+#include "BehaviorTree/Blackboard/BlackboardKeyType_Float.h"
+#include "BehaviorTree/Tasks/BTTask_WaitBlackboardTime.h"
+
+UBTTask_WaitBlackboardTime::UBTTask_WaitBlackboardTime(const FObjectInitializer& ObjectInitializer) : Super(ObjectInitializer)
+{
+	NodeName = "Wait Blackboard Time";
+
+	BlackboardKey.AddFloatFilter(this);
+}
+
+void UBTTask_WaitBlackboardTime::InitializeFromAsset(UBehaviorTree& Asset)
+{
+	Super::InitializeFromAsset(Asset);
+	BlackboardKey.CacheSelectedKey(GetBlackboardAsset());
+}
+
+EBTNodeResult::Type UBTTask_WaitBlackboardTime::ExecuteTask(UBehaviorTreeComponent* OwnerComp, uint8* NodeMemory)
+{
+	// Update wait time based on current blackboard key value
+	const UBlackboardComponent* MyBlackboard = OwnerComp->GetBlackboardComponent();
+	if (MyBlackboard && BlackboardKey.SelectedKeyType == UBlackboardKeyType_Float::StaticClass())
+	{
+		WaitTime = MyBlackboard->GetValueAsFloat(BlackboardKey.GetSelectedKeyID());
+	}
+	
+	return Super::ExecuteTask(OwnerComp, NodeMemory);
+}
+
+FString UBTTask_WaitBlackboardTime::GetStaticDescription() const
+{
+	return FString::Printf(TEXT("%s: %s"), *UBTTaskNode::GetStaticDescription(), *GetSelectedBlackboardKey().ToString());
+}
+
