Added blueprint functions for checking AI-validity of vectors and rotators #UE4

diff --git a/Engine/Source/Runtime/AIModule/Classes/Blueprint/AIBlueprintHelperLibrary.h b/Engine/Source/Runtime/AIModule/Classes/Blueprint/AIBlueprintHelperLibrary.h
--- a/Engine/Source/Runtime/AIModule/Classes/Blueprint/AIBlueprintHelperLibrary.h
+++ b/Engine/Source/Runtime/AIModule/Classes/Blueprint/AIBlueprintHelperLibrary.h
@@ -32,4 +32,13 @@ class AIMODULE_API UAIBlueprintHelperLibrary : public UBlueprintFunctionLibrary
 	/** unlocks indicated AI resources of animated pawn. Will unlock only animation-locked resources */
 	UFUNCTION(BlueprintCallable, Category = "Animation", BlueprintAuthorityOnly, meta = (DefaultToSelf = "AnimInstance"))
 	static void UnlockAIResourcesWithAnimation(class UAnimInstance* AnimInstance, bool bUnlockMovement, bool UnlockAILogic);
+
+	UFUNCTION(BlueprintPure, Category = "AI")
+	static bool IsValidAILocation(FVector Location);
+
+	UFUNCTION(BlueprintPure, Category = "AI")
+	static bool IsValidAIDirection(FVector DirectionVector);
+
+	UFUNCTION(BlueprintPure, Category = "AI")
+	static bool IsValidAIRotation(FRotator Rotation);
 };
diff --git a/Engine/Source/Runtime/AIModule/Private/Blueprint/AIBlueprintHelperLibrary.cpp b/Engine/Source/Runtime/AIModule/Private/Blueprint/AIBlueprintHelperLibrary.cpp
--- a/Engine/Source/Runtime/AIModule/Private/Blueprint/AIBlueprintHelperLibrary.cpp
+++ b/Engine/Source/Runtime/AIModule/Private/Blueprint/AIBlueprintHelperLibrary.cpp
@@ -214,3 +214,18 @@ void UAIBlueprintHelperLibrary::UnlockAIResourcesWithAnimation(UAnimInstance* An
 		}
 	}
 }
+
+bool UAIBlueprintHelperLibrary::IsValidAILocation(FVector Location)
+{
+	return FAISystem::IsValidLocation(Location);
+}
+
+bool UAIBlueprintHelperLibrary::IsValidAIDirection(FVector DirectionVector)
+{
+	return FAISystem::IsValidDirection(DirectionVector);
+}
+
+bool UAIBlueprintHelperLibrary::IsValidAIRotation(FRotator Rotation)
+{
+	return FAISystem::IsValidRotation(Rotation);
+}
