Added alternative BP-implementable events to blueprint BT nodes that allow users to get AIController and Pawn right off the bat #UE4

- also added caching of Actor owner and AIController owner to instanced BP BT nodes

diff --git a/Engine/Source/Runtime/AIModule/Classes/BehaviorTree/BTFunctionLibrary.h b/Engine/Source/Runtime/AIModule/Classes/BehaviorTree/BTFunctionLibrary.h
--- a/Engine/Source/Runtime/AIModule/Classes/BehaviorTree/BTFunctionLibrary.h
+++ b/Engine/Source/Runtime/AIModule/Classes/BehaviorTree/BTFunctionLibrary.h
@@ -6,6 +6,16 @@
 class UBlackboardComponent;
 struct FBlackboardKeySelector;
 
+namespace FBTNodeBPImplementationHelper
+{
+	static const int32 NoImplementation = 0;
+	static const int32 Generic = 1 << 0;
+	static const int32 AISpecific = 1 << 1;
+	static const int32 All = Generic | AISpecific;
+
+	AIMODULE_API int32 CheckEventImplementationVersion(FName GenericEventName, FName AIEventName, const UObject* Ob, const UClass* StopAtClass);
+}
+
 UCLASS(meta=(RestrictedToClasses="BTNode"))
 class AIMODULE_API UBTFunctionLibrary : public UBlueprintFunctionLibrary
 {
diff --git a/Engine/Source/Runtime/AIModule/Classes/BehaviorTree/BTNode.h b/Engine/Source/Runtime/AIModule/Classes/BehaviorTree/BTNode.h
--- a/Engine/Source/Runtime/AIModule/Classes/BehaviorTree/BTNode.h
+++ b/Engine/Source/Runtime/AIModule/Classes/BehaviorTree/BTNode.h
@@ -6,7 +6,13 @@
 #include "BehaviorTreeTypes.h"
 #include "BTNode.generated.h"
 
+class AAIController;
+class UWorld;
 class UBehaviorTree;
+class UBehaviorTreeComponent;
+class UBTCompositeNode;
+class UBlackboardData;
+struct FBehaviorTreeSearchData;
 
 struct FBTInstancedNodeMemory
 {
@@ -18,67 +24,67 @@ class AIMODULE_API UBTNode : public UObject
 {
 	GENERATED_UCLASS_BODY()
 
-	virtual class UWorld* GetWorld() const override;
+	virtual UWorld* GetWorld() const override;
 
 	/** fill in data about tree structure */
-	void InitializeNode(class UBTCompositeNode* InParentNode, uint16 InExecutionIndex, uint16 InMemoryOffset, uint8 InTreeDepth);
+	void InitializeNode(UBTCompositeNode* InParentNode, uint16 InExecutionIndex, uint16 InMemoryOffset, uint8 InTreeDepth);
 
 	/** initialize any asset related data */
 	virtual void InitializeFromAsset(UBehaviorTree& Asset);
 	
 	/** initialize memory block */
-	virtual void InitializeMemory(class UBehaviorTreeComponent* OwnerComp, uint8* NodeMemory, EBTMemoryInit::Type InitType) const;
+	virtual void InitializeMemory(UBehaviorTreeComponent* OwnerComp, uint8* NodeMemory, EBTMemoryInit::Type InitType) const;
 
 	/** cleanup memory block */
-	virtual void CleanupMemory(class UBehaviorTreeComponent* OwnerComp, uint8* NodeMemory, EBTMemoryClear::Type CleanupType) const;
+	virtual void CleanupMemory(UBehaviorTreeComponent* OwnerComp, uint8* NodeMemory, EBTMemoryClear::Type CleanupType) const;
 
 	/** gathers description of all runtime parameters */
-	virtual void DescribeRuntimeValues(const class UBehaviorTreeComponent* OwnerComp, uint8* NodeMemory, EBTDescriptionVerbosity::Type Verbosity, TArray<FString>& Values) const;
+	virtual void DescribeRuntimeValues(const UBehaviorTreeComponent* OwnerComp, uint8* NodeMemory, EBTDescriptionVerbosity::Type Verbosity, TArray<FString>& Values) const;
 
 	/** size of instance memory */
 	virtual uint16 GetInstanceMemorySize() const;
 
 	/** called when node instance is added to tree */
-	virtual void OnInstanceCreated(class UBehaviorTreeComponent* OwnerComp);
+	virtual void OnInstanceCreated(UBehaviorTreeComponent* OwnerComp);
 
 	/** called when node instance is removed from tree */
-	virtual void OnInstanceDestroyed(class UBehaviorTreeComponent* OwnerComp);
+	virtual void OnInstanceDestroyed(UBehaviorTreeComponent* OwnerComp);
 
 	/** called on creating subtree to set up memory and instancing */
-	void InitializeInSubtree(class UBehaviorTreeComponent* OwnerComp, uint8* NodeMemory, int32& NextInstancedIndex, EBTMemoryInit::Type InitType) const;
+	void InitializeInSubtree(UBehaviorTreeComponent* OwnerComp, uint8* NodeMemory, int32& NextInstancedIndex, EBTMemoryInit::Type InitType) const;
 
 	/** called on removing subtree to cleanup memory */
-	void CleanupInSubtree(class UBehaviorTreeComponent* OwnerComp, uint8* NodeMemory, EBTMemoryClear::Type CleanupType) const;
+	void CleanupInSubtree(UBehaviorTreeComponent* OwnerComp, uint8* NodeMemory, EBTMemoryClear::Type CleanupType) const;
 
 	/** size of special, hidden memory block for internal mechanics */
 	virtual uint16 GetSpecialMemorySize() const;
 
 #if USE_BEHAVIORTREE_DEBUGGER
 	/** fill in data about execution order */
-	void InitializeExecutionOrder(class UBTNode* NextNode);
+	void InitializeExecutionOrder(UBTNode* NextNode);
 
 	/** @return next node in execution order */
-	class UBTNode* GetNextNode() const;
+	UBTNode* GetNextNode() const;
 #endif
 
 	template<typename T>
-	T* GetNodeMemory(struct FBehaviorTreeSearchData& SearchData) const;
+	T* GetNodeMemory(FBehaviorTreeSearchData& SearchData) const;
 
 	template<typename T>
-	const T* GetNodeMemory(const struct FBehaviorTreeSearchData& SearchData) const;
+	const T* GetNodeMemory(const FBehaviorTreeSearchData& SearchData) const;
 
 	template<typename T>
-	T* GetNodeMemory(struct FBehaviorTreeInstance& BTInstance) const;
+	T* GetNodeMemory(FBehaviorTreeInstance& BTInstance) const;
 
 	template<typename T>
-	const T* GetNodeMemory(const struct FBehaviorTreeInstance& BTInstance) const;
+	const T* GetNodeMemory(const FBehaviorTreeInstance& BTInstance) const;
 
 	/** get special memory block used for hidden shared data (e.g. node instancing) */
 	template<typename T>
 	T* GetSpecialNodeMemory(uint8* NodeMemory) const;
 
 	/** @return parent node */
-	class UBTCompositeNode* GetParentNode() const;
+	UBTCompositeNode* GetParentNode() const;
 
 	/** @return name of node */
 	FString GetNodeName() const;
@@ -108,17 +114,17 @@ class AIMODULE_API UBTNode : public UObject
 	bool IsInstanced() const;
 
 	/** @return tree asset */
-	class UBehaviorTree* GetTreeAsset() const;
+	UBehaviorTree* GetTreeAsset() const;
 
 	/** @return blackboard asset */
-	class UBlackboardData* GetBlackboardAsset() const;
+	UBlackboardData* GetBlackboardAsset() const;
 
 	/** @return node instance if bCreateNodeInstance was set */
-	class UBTNode* GetNodeInstance(const class UBehaviorTreeComponent* OwnerComp, uint8* NodeMemory) const;
-	class UBTNode* GetNodeInstance(struct FBehaviorTreeSearchData& SearchData) const;
+	UBTNode* GetNodeInstance(const UBehaviorTreeComponent* OwnerComp, uint8* NodeMemory) const;
+	UBTNode* GetNodeInstance(FBehaviorTreeSearchData& SearchData) const;
 
 	/** @return string containing description of this node instance with all relevant runtime values */
-	FString GetRuntimeDescription(const class UBehaviorTreeComponent* OwnerComp, uint8* NodeMemory, EBTDescriptionVerbosity::Type Verbosity) const;
+	FString GetRuntimeDescription(const UBehaviorTreeComponent* OwnerComp, uint8* NodeMemory, EBTDescriptionVerbosity::Type Verbosity) const;
 
 	/** @return string containing description of this node with all setup values */
 	virtual FString GetStaticDescription() const;
@@ -131,25 +137,28 @@ class AIMODULE_API UBTNode : public UObject
 	virtual bool UsesBlueprint() const;
 #endif
 
+	/** Gets called only for instanced nodes(bCreateNodeInstance == true). In practive overriden by BP-implemented BT nodes */
+	virtual void SetOwner(AActor* ActorOwner) {}
+
 protected:
 
 	/** node name */
 	UPROPERTY(Category=Description, EditInstanceOnly)
 	FString NodeName;
-
+	
 private:
 
 	/** source asset */
 	UPROPERTY()
-	class UBehaviorTree* TreeAsset;
+	UBehaviorTree* TreeAsset;
 
 	/** parent node */
 	UPROPERTY()
-	class UBTCompositeNode* ParentNode;
+	UBTCompositeNode* ParentNode;
 
 #if USE_BEHAVIORTREE_DEBUGGER
 	/** next node in execution order */
-	class UBTNode* NextExecutionNode;
+	UBTNode* NextExecutionNode;
 #endif
 
 	/** depth first index (execution order) */
@@ -176,18 +185,18 @@ protected:
 //////////////////////////////////////////////////////////////////////////
 // Inlines
 
-FORCEINLINE class UBehaviorTree* UBTNode::GetTreeAsset() const
+FORCEINLINE UBehaviorTree* UBTNode::GetTreeAsset() const
 {
 	return TreeAsset;
 }
 
-FORCEINLINE class UBTCompositeNode* UBTNode::GetParentNode() const
+FORCEINLINE UBTCompositeNode* UBTNode::GetParentNode() const
 {
 	return ParentNode;
 }
 
 #if USE_BEHAVIORTREE_DEBUGGER
-FORCEINLINE class UBTNode* UBTNode::GetNextNode() const
+FORCEINLINE UBTNode* UBTNode::GetNextNode() const
 {
 	return NextExecutionNode;
 }
@@ -242,25 +251,25 @@ FORCEINLINE bool UBTNode::IsInstanced() const
 }
 
 template<typename T>
-T* UBTNode::GetNodeMemory(struct FBehaviorTreeSearchData& SearchData) const
+T* UBTNode::GetNodeMemory(FBehaviorTreeSearchData& SearchData) const
 {
 	return GetNodeMemory<T>(SearchData.OwnerComp->InstanceStack[SearchData.OwnerComp->GetActiveInstanceIdx()]);
 }
 
 template<typename T>
-const T* UBTNode::GetNodeMemory(const struct FBehaviorTreeSearchData& SearchData) const
+const T* UBTNode::GetNodeMemory(const FBehaviorTreeSearchData& SearchData) const
 {
 	return GetNodeMemory<T>(SearchData.OwnerComp->InstanceStack[SearchData.OwnerComp->GetActiveInstanceIdx()]);
 }
 
 template<typename T>
-T* UBTNode::GetNodeMemory(struct FBehaviorTreeInstance& BTInstance) const
+T* UBTNode::GetNodeMemory(FBehaviorTreeInstance& BTInstance) const
 {
 	return (T*)(BTInstance.InstanceMemory.GetData() + MemoryOffset);
 }
 
 template<typename T>
-const T* UBTNode::GetNodeMemory(const struct FBehaviorTreeInstance& BTInstance) const
+const T* UBTNode::GetNodeMemory(const FBehaviorTreeInstance& BTInstance) const
 {
 	return (const T*)(BTInstance.InstanceMemory.GetData() + MemoryOffset);
 }
diff --git a/Engine/Source/Runtime/AIModule/Classes/BehaviorTree/Decorators/BTDecorator_BlueprintBase.h b/Engine/Source/Runtime/AIModule/Classes/BehaviorTree/Decorators/BTDecorator_BlueprintBase.h
--- a/Engine/Source/Runtime/AIModule/Classes/BehaviorTree/Decorators/BTDecorator_BlueprintBase.h
+++ b/Engine/Source/Runtime/AIModule/Classes/BehaviorTree/Decorators/BTDecorator_BlueprintBase.h
@@ -37,12 +37,21 @@ class AIMODULE_API UBTDecorator_BlueprintBase : public UBTDecorator
 	virtual bool CalculateRawConditionValue(class UBehaviorTreeComponent* OwnerComp, uint8* NodeMemory) const override;
 	virtual void OnInstanceDestroyed(class UBehaviorTreeComponent* OwnerComp) override;
 
+	virtual void SetOwner(AActor* ActorOwner) override;
+
 #if WITH_EDITOR
 	virtual FName GetNodeIconName() const override;
 	virtual bool UsesBlueprint() const override;
 #endif
 
 protected:
+	/** Cached AIController owner of BehaviorTreeComponent. */
+	UPROPERTY(Transient)
+	AAIController* AIOwner;
+
+	/** Cached AIController owner of BehaviorTreeComponent. */
+	UPROPERTY(Transient)
+	AActor* ActorOwner;
 
 	/** blackboard key names that should be observed */
 	TArray<FName> ObservedKeyNames;
@@ -58,55 +67,109 @@ protected:
 	mutable uint32 CurrentCallResult : 1;
 
 	/** set if ReceiveTick is implemented by blueprint */
-	uint32 bImplementsReceiveTick : 1;
+	uint32 ReceiveTickImplementations : 2;
 
 	/** set if ReceiveExecutionStart is implemented by blueprint */
-	uint32 bImplementsReceiveExecutionStart : 1;
+	uint32 ReceiveExecutionStartImplementations : 2;
 
 	/** set if ReceiveExecutionFinish is implemented by blueprint */
-	uint32 bImplementsReceiveExecutionFinish : 1;
+	uint32 ReceiveExecutionFinishImplementations : 2;
 
 	/** set if ReceiveObserverActivated is implemented by blueprint */
-	uint32 bImplementsReceiveObserverActivated : 1;
+	uint32 ReceiveObserverActivatedImplementations : 2;
 
 	/** set if ReceiveObserverDeactivated is implemented by blueprint */
-	uint32 bImplementsReceiveObserverDeactivated : 1;
+	uint32 ReceiveObserverDeactivatedImplementations : 2;
 
 	/** set if ReceiveConditionCheck is implemented by blueprint */
-	uint32 bImplementsReceiveConditionCheck : 1;
+	uint32 ReceiveConditionCheckImplementations : 2;
 
 	FOnBlackboardChange BBKeyObserver;
-
+	
 	virtual void OnBecomeRelevant(UBehaviorTreeComponent* OwnerComp, uint8* NodeMemory) override;
 	virtual void OnCeaseRelevant(UBehaviorTreeComponent* OwnerComp, uint8* NodeMemory) override;
 	virtual void OnNodeActivation(struct FBehaviorTreeSearchData& SearchData) override;
 	virtual void OnNodeDeactivation(struct FBehaviorTreeSearchData& SearchData, EBTNodeResult::Type NodeResult) override;
 	virtual void TickNode(UBehaviorTreeComponent* OwnerComp, uint8* NodeMemory, float DeltaSeconds) override;
 
-	/** tick function */
+	/** tick function
+	 *	@Note that if both generic and AI event versions are implemented only the more
+	 *	suitable one will be called, meaning the AI version if called for AI, generic one otherwise */
 	UFUNCTION(BlueprintImplementableEvent)
 	virtual void ReceiveTick(AActor* OwnerActor, float DeltaSeconds);
 
-	/** called on execution of underlying node */
+	/** called on execution of underlying node 
+	 *	@Note that if both generic and AI event versions are implemented only the more
+	 *	suitable one will be called, meaning the AI version if called for AI, generic one otherwise */
 	UFUNCTION(BlueprintImplementableEvent)
 	virtual void ReceiveExecutionStart(AActor* OwnerActor);
 
-	/** called when execution of underlying node is finished */
+	/** called when execution of underlying node is finished 
+	 *	@Note that if both generic and AI event versions are implemented only the more
+	 *	suitable one will be called, meaning the AI version if called for AI, generic one otherwise */
 	UFUNCTION(BlueprintImplementableEvent)
 	virtual void ReceiveExecutionFinish(AActor* OwnerActor, enum EBTNodeResult::Type NodeResult);
 
-	/** called when observer is activated (flow controller) */
+	/** called when observer is activated (flow controller) 
+	 *	@Note that if both generic and AI event versions are implemented only the more
+	 *	suitable one will be called, meaning the AI version if called for AI, generic one otherwise */
 	UFUNCTION(BlueprintImplementableEvent)
 	virtual void ReceiveObserverActivated(AActor* OwnerActor);
 
-	/** called when observer is deactivated (flow controller) */
+	/** called when observer is deactivated (flow controller) 
+	 *	@Note that if both generic and AI event versions are implemented only the more
+	 *	suitable one will be called, meaning the AI version if called for AI, generic one otherwise */
 	UFUNCTION(BlueprintImplementableEvent)
 	virtual void ReceiveObserverDeactivated(AActor* OwnerActor);
 
-	/** called when testing if underlying node can be executed, must call FinishConditionCheck */
+	/** called when testing if underlying node can be executed, must call FinishConditionCheck
+	 *	@Note that if both generic and AI event versions are implemented only the more
+	 *	suitable one will be called, meaning the AI version if called for AI, generic one otherwise */
 	UFUNCTION(BlueprintImplementableEvent)
 	virtual void ReceiveConditionCheck(AActor* OwnerActor);
 
+	/** Alternative AI version of ReceiveTick
+	 *	@see ReceiveTick for more details
+	 *	@Note that if both generic and AI event versions are implemented only the more
+	 *	suitable one will be called, meaning the AI version if called for AI, generic one otherwise */
+	UFUNCTION(BlueprintImplementableEvent, Category = AI)
+	virtual void ReceiveTickAI(AAIController* OwnerController, APawn* ControlledPawn, float DeltaSeconds);
+
+	/** Alternative AI version of ReceiveExecutionStart
+	 *	@see ReceiveExecutionStart for more details
+	 *	@Note that if both generic and AI event versions are implemented only the more
+	 *	suitable one will be called, meaning the AI version if called for AI, generic one otherwise */
+	UFUNCTION(BlueprintImplementableEvent, Category = AI)
+	virtual void ReceiveExecutionStartAI(AAIController* OwnerController, APawn* ControlledPawn);
+
+	/** Alternative AI version of ReceiveExecutionFinish
+	 *	@see ReceiveExecutionFinish for more details
+	 *	@Note that if both generic and AI event versions are implemented only the more
+	 *	suitable one will be called, meaning the AI version if called for AI, generic one otherwise */
+	UFUNCTION(BlueprintImplementableEvent, Category = AI)
+	virtual void ReceiveExecutionFinishAI(AAIController* OwnerController, APawn* ControlledPawn, enum EBTNodeResult::Type NodeResult);
+
+	/** Alternative AI version of ReceiveObserverActivated
+	 *	@see ReceiveObserverActivated for more details
+	 *	@Note that if both generic and AI event versions are implemented only the more
+	 *	suitable one will be called, meaning the AI version if called for AI, generic one otherwise */
+	UFUNCTION(BlueprintImplementableEvent, Category = AI)
+	virtual void ReceiveObserverActivatedAI(AAIController* OwnerController, APawn* ControlledPawn);
+
+	/** Alternative AI version of ReceiveObserverDeactivated
+	 *	@see ReceiveObserverDeactivated for more details
+	 *	@Note that if both generic and AI event versions are implemented only the more
+	 *	suitable one will be called, meaning the AI version if called for AI, generic one otherwise */
+	UFUNCTION(BlueprintImplementableEvent, Category = AI)
+	virtual void ReceiveObserverDeactivatedAI(AAIController* OwnerController, APawn* ControlledPawn);
+
+	/** Alternative AI version of ReceiveConditionCheck
+	 *	@see ReceiveConditionCheck for more details
+	 *	@Note that if both generic and AI event versions are implemented only the more
+	 *	suitable one will be called, meaning the AI version if called for AI, generic one otherwise */
+	UFUNCTION(BlueprintImplementableEvent, Category = AI)
+	virtual void ReceiveConditionCheckAI(AAIController* OwnerController, APawn* ControlledPawn);
+
 	/** finishes condition check */
 	UFUNCTION(BlueprintCallable, Category="AI|BehaviorTree")
 	void FinishConditionCheck(bool bAllowExecution);
diff --git a/Engine/Source/Runtime/AIModule/Classes/BehaviorTree/Services/BTService_BlueprintBase.h b/Engine/Source/Runtime/AIModule/Classes/BehaviorTree/Services/BTService_BlueprintBase.h
--- a/Engine/Source/Runtime/AIModule/Classes/BehaviorTree/Services/BTService_BlueprintBase.h
+++ b/Engine/Source/Runtime/AIModule/Classes/BehaviorTree/Services/BTService_BlueprintBase.h
@@ -22,11 +22,20 @@ class AIMODULE_API UBTService_BlueprintBase : public UBTService
 	virtual void DescribeRuntimeValues(const class UBehaviorTreeComponent* OwnerComp, uint8* NodeMemory, EBTDescriptionVerbosity::Type Verbosity, TArray<FString>& Values) const override;
 	virtual void OnInstanceDestroyed(class UBehaviorTreeComponent* OwnerComp) override;
 
+	virtual void SetOwner(AActor* ActorOwner) override;
+
 #if WITH_EDITOR
 	virtual bool UsesBlueprint() const override;
 #endif
 
 protected:
+	/** Cached AIController owner of BehaviorTreeComponent. */
+	UPROPERTY(Transient)
+	AAIController* AIOwner;
+
+	/** Cached actor owner of BehaviorTreeComponent. */
+	UPROPERTY(Transient)
+	AActor* ActorOwner;
 
 	// Gets the description for our service
 	virtual FString GetStaticServiceDescription() const override;
@@ -39,38 +48,74 @@ protected:
 	uint32 bShowPropertyDetails : 1;
 
 	/** set if ReceiveTick is implemented by blueprint */
-	uint32 bImplementsReceiveTick : 1;
+	uint32 ReceiveTickImplementations : 2;
 	
 	/** set if ReceiveActivation is implemented by blueprint */
-	uint32 bImplementsReceiveActivation : 1;
+	uint32 ReceiveActivationImplementations : 2;
 
 	/** set if ReceiveDeactivation is implemented by blueprint */
-	uint32 bImplementsReceiveDeactivation : 1;
+	uint32 ReceiveDeactivationImplementations : 2;
 
 	/** set if ReceiveSearchStart is implemented by blueprint */
-	uint32 bImplementsReceiveSearchStart : 1;
+	uint32 ReceiveSearchStartImplementations : 2;
 
 	virtual void OnBecomeRelevant(class UBehaviorTreeComponent* OwnerComp, uint8* NodeMemory) override;
 	virtual void OnCeaseRelevant(class UBehaviorTreeComponent* OwnerComp, uint8* NodeMemory) override;
 	virtual void TickNode(class UBehaviorTreeComponent* OwnerComp, uint8* NodeMemory, float DeltaSeconds) override;
 	virtual void OnSearchStart(struct FBehaviorTreeSearchData& SearchData) override;
 
-	/** tick function */
+	/** tick function
+	 *	@Note that if both generic and AI event versions are implemented only the more
+	 *	suitable one will be called, meaning the AI version if called for AI, generic one otherwise */
 	UFUNCTION(BlueprintImplementableEvent)
 	virtual void ReceiveTick(AActor* OwnerActor, float DeltaSeconds);
 
-	/** task search enters branch of tree, initial ReceiveTick won't be called if you implement this event! */
+	/** task search enters branch of tree, initial ReceiveTick won't be called if you implement this event!
+	 *	@Note that if both generic and AI event versions are implemented only the more
+	 *	suitable one will be called, meaning the AI version if called for AI, generic one otherwise */
 	UFUNCTION(BlueprintImplementableEvent)
 	virtual void ReceiveSearchStart(AActor* OwnerActor);
 
-	/** service became active */
+	/** service became active
+	 *	@Note that if both generic and AI event versions are implemented only the more
+	 *	suitable one will be called, meaning the AI version if called for AI, generic one otherwise */
 	UFUNCTION(BlueprintImplementableEvent)
 	virtual void ReceiveActivation(AActor* OwnerActor);
 
-	/** service became inactive */
+	/** service became inactive
+	 *	@Note that if both generic and AI event versions are implemented only the more
+	 *	suitable one will be called, meaning the AI version if called for AI, generic one otherwise */
 	UFUNCTION(BlueprintImplementableEvent)
 	virtual void ReceiveDeactivation(AActor* OwnerActor);
 
+	/** Alternative AI version of ReceiveTick function.
+	 *	@see ReceiveTick for more details
+	 *	@Note that if both generic and AI event versions are implemented only the more
+	 *	suitable one will be called, meaning the AI version if called for AI, generic one otherwise */
+	UFUNCTION(BlueprintImplementableEvent, Category = AI)
+	virtual void ReceiveTickAI(AAIController* OwnerController, APawn* ControlledPawn, float DeltaSeconds);
+
+	/** Alternative AI version of ReceiveSearchStart function.
+	 *	@see ReceiveSearchStart for more details
+	 *	@Note that if both generic and AI event versions are implemented only the more
+	 *	suitable one will be called, meaning the AI version if called for AI, generic one otherwise */
+	UFUNCTION(BlueprintImplementableEvent, Category = AI)
+	virtual void ReceiveSearchStartAI(AAIController* OwnerController, APawn* ControlledPawn);
+
+	/** Alternative AI version of ReceiveActivation function.
+	 *	@see ReceiveActivation for more details
+	 *	@Note that if both generic and AI event versions are implemented only the more
+	 *	suitable one will be called, meaning the AI version if called for AI, generic one otherwise */
+	UFUNCTION(BlueprintImplementableEvent, Category = AI)
+	virtual void ReceiveActivationAI(AAIController* OwnerController, APawn* ControlledPawn);
+
+	/** Alternative AI version of ReceiveDeactivation function.
+	 *	@see ReceiveDeactivation for more details
+	 *	@Note that if both generic and AI event versions are implemented only the more
+	 *	suitable one will be called, meaning the AI version if called for AI, generic one otherwise */
+	UFUNCTION(BlueprintImplementableEvent, Category = AI)
+	virtual void ReceiveDeactivationAI(AAIController* OwnerController, APawn* ControlledPawn);
+
 	/** check if service is currently being active */
 	UFUNCTION(BlueprintCallable, Category="AI|BehaviorTree")
 	bool IsServiceActive() const;
diff --git a/Engine/Source/Runtime/AIModule/Classes/BehaviorTree/Tasks/BTTask_BlueprintBase.h b/Engine/Source/Runtime/AIModule/Classes/BehaviorTree/Tasks/BTTask_BlueprintBase.h
--- a/Engine/Source/Runtime/AIModule/Classes/BehaviorTree/Tasks/BTTask_BlueprintBase.h
+++ b/Engine/Source/Runtime/AIModule/Classes/BehaviorTree/Tasks/BTTask_BlueprintBase.h
@@ -26,11 +26,20 @@ class AIMODULE_API UBTTask_BlueprintBase : public UBTTaskNode
 	virtual void DescribeRuntimeValues(const class UBehaviorTreeComponent* OwnerComp, uint8* NodeMemory, EBTDescriptionVerbosity::Type Verbosity, TArray<FString>& Values) const override;
 	virtual void OnInstanceDestroyed(class UBehaviorTreeComponent* OwnerComp) override;
 
+	virtual void SetOwner(AActor* ActorOwner) override;
+
 #if WITH_EDITOR
 	virtual bool UsesBlueprint() const override;
 #endif
 
 protected:
+	/** Cached AIController owner of BehaviorTreeComponent. */
+	UPROPERTY(Transient)
+	AAIController* AIOwner;
+
+	/** Cached actor owner of BehaviorTreeComponent. */
+	UPROPERTY(Transient)
+	AActor* ActorOwner;
 
 	/** temporary variable for ReceiveExecute(Abort)-FinishExecute(Abort) chain */
 	mutable TEnumAsByte<EBTNodeResult::Type> CurrentCallResult;
@@ -43,30 +52,57 @@ protected:
 	uint32 bShowPropertyDetails : 1;
 
 	/** set if ReceiveTick is implemented by blueprint */
-	uint32 bImplementsReceiveTick : 1;
+	uint32 ReceiveTickImplementations : 2;
 
 	/** set if ReceiveExecute is implemented by blueprint */
-	uint32 bImplementsReceiveExecute : 1;
+	uint32 ReceiveExecuteImplementations : 2;
 
 	/** set if ReceiveAbort is implemented by blueprint */
-	uint32 bImplementsReceiveAbort : 1;
+	uint32 ReceiveAbortImplementations : 2;
 
 	/** if set, execution is inside blueprint's ReceiveExecute(Abort) event
 	  * FinishExecute(Abort) function should store their result in CurrentCallResult variable */
 	mutable uint32 bStoreFinishResult : 1;
 
-	/** entry point, task will stay active until FinishExecute is called */
+	/** entry point, task will stay active until FinishExecute is called.
+	 *	@Note that if both generic and AI event versions are implemented only the more 
+	 *	suitable one will be called, meaning the AI version if called for AI, generic one otherwise */
 	UFUNCTION(BlueprintImplementableEvent)
 	virtual void ReceiveExecute(AActor* OwnerActor);
 
-	/** if blueprint graph contains this event, task will stay active until FinishAbort is called */
+	/** if blueprint graph contains this event, task will stay active until FinishAbort is called
+	 *	@Note that if both generic and AI event versions are implemented only the more
+	 *	suitable one will be called, meaning the AI version if called for AI, generic one otherwise */
 	UFUNCTION(BlueprintImplementableEvent)
 	virtual void ReceiveAbort(AActor* OwnerActor);
 
-	/** tick function */
+	/** tick function
+	 *	@Note that if both generic and AI event versions are implemented only the more
+	 *	suitable one will be called, meaning the AI version if called for AI, generic one otherwise */
 	UFUNCTION(BlueprintImplementableEvent)
 	virtual void ReceiveTick(AActor* OwnerActor, float DeltaSeconds);
 
+	/** Alternative AI version of ReceiveExecute
+	*	@see ReceiveExecute for more details
+	 *	@Note that if both generic and AI event versions are implemented only the more
+	 *	suitable one will be called, meaning the AI version if called for AI, generic one otherwise */
+	UFUNCTION(BlueprintImplementableEvent, Category = AI)
+	virtual void ReceiveExecuteAI(AAIController* OwnerController, APawn* ControlledPawn);
+
+	/** Alternative AI version of ReceiveAbort
+	 *	@see ReceiveAbort for more details
+	 *	@Note that if both generic and AI event versions are implemented only the more
+	 *	suitable one will be called, meaning the AI version if called for AI, generic one otherwise */
+	UFUNCTION(BlueprintImplementableEvent, Category = AI)
+	virtual void ReceiveAbortAI(AAIController* OwnerController, APawn* ControlledPawn);
+
+	/** Alternative AI version of tick function.
+	 *	@see ReceiveTick for more details
+	 *	@Note that if both generic and AI event versions are implemented only the more
+	 *	suitable one will be called, meaning the AI version if called for AI, generic one otherwise */
+	UFUNCTION(BlueprintImplementableEvent, Category = AI)
+	virtual void ReceiveTickAI(AAIController* OwnerController, APawn* ControlledPawn, float DeltaSeconds);
+
 	/** finishes task execution with Success or Fail result */
 	UFUNCTION(BlueprintCallable, Category="AI|BehaviorTree")
 	void FinishExecute(bool bSuccess);
@@ -86,7 +122,7 @@ protected:
 	/** check if task is currently being executed */
 	UFUNCTION(BlueprintCallable, Category="AI|BehaviorTree")
 	bool IsTaskExecuting() const;
-
+	
 	/** ticks this task */
 	virtual void TickTask(class UBehaviorTreeComponent* OwnerComp, uint8* NodeMemory, float DeltaSeconds) override;
 
diff --git a/Engine/Source/Runtime/AIModule/Private/BehaviorTree/BTFunctionLibrary.cpp b/Engine/Source/Runtime/AIModule/Private/BehaviorTree/BTFunctionLibrary.cpp
--- a/Engine/Source/Runtime/AIModule/Private/BehaviorTree/BTFunctionLibrary.cpp
+++ b/Engine/Source/Runtime/AIModule/Private/BehaviorTree/BTFunctionLibrary.cpp
@@ -6,6 +6,22 @@
 #include "BehaviorTree/Blackboard/BlackboardKeyAllTypes.h"
 #include "BehaviorTree/BTFunctionLibrary.h"
 
+//----------------------------------------------------------------------//
+// 
+//----------------------------------------------------------------------//
+namespace FBTNodeBPImplementationHelper
+{
+	int32 CheckEventImplementationVersion(FName GenericEventName, FName AIEventName, const UObject* Ob, const UClass* StopAtClass)
+	{
+		const bool bGeneric = BlueprintNodeHelpers::HasBlueprintFunction(GenericEventName, Ob, StopAtClass);
+		const bool bAI = BlueprintNodeHelpers::HasBlueprintFunction(AIEventName, Ob, StopAtClass);
+		
+		return (bGeneric ? Generic : NoImplementation) | (bAI ? AISpecific : NoImplementation);
+	}
+}
+//----------------------------------------------------------------------//
+// 
+//----------------------------------------------------------------------//
 namespace
 {
 	FORCEINLINE UBlackboardComponent* GetBlackboard(UBTNode& BTNode)
diff --git a/Engine/Source/Runtime/AIModule/Private/BehaviorTree/BTNode.cpp b/Engine/Source/Runtime/AIModule/Private/BehaviorTree/BTNode.cpp
--- a/Engine/Source/Runtime/AIModule/Private/BehaviorTree/BTNode.cpp
+++ b/Engine/Source/Runtime/AIModule/Private/BehaviorTree/BTNode.cpp
@@ -80,6 +80,9 @@ void UBTNode::InitializeInSubtree(UBehaviorTreeComponent* OwnerComp, uint8* Node
 
 			OwnerComp->NodeInstances.Add(NodeInstance);
 		}
+		check(NodeInstance);
+
+		NodeInstance->SetOwner(OwnerComp->GetOwner());
 
 		FBTInstancedNodeMemory* MyMemory = GetSpecialNodeMemory<FBTInstancedNodeMemory>(NodeMemory);
 		MyMemory->NodeIdx = NextInstancedIndex;
diff --git a/Engine/Source/Runtime/AIModule/Private/BehaviorTree/Decorators/BTDecorator_BlueprintBase.cpp b/Engine/Source/Runtime/AIModule/Private/BehaviorTree/Decorators/BTDecorator_BlueprintBase.cpp
--- a/Engine/Source/Runtime/AIModule/Private/BehaviorTree/Decorators/BTDecorator_BlueprintBase.cpp
+++ b/Engine/Source/Runtime/AIModule/Private/BehaviorTree/Decorators/BTDecorator_BlueprintBase.cpp
@@ -9,18 +9,18 @@
 UBTDecorator_BlueprintBase::UBTDecorator_BlueprintBase(const FObjectInitializer& ObjectInitializer) : Super(ObjectInitializer)
 {
 	UClass* StopAtClass = UBTDecorator_BlueprintBase::StaticClass();
-	bImplementsReceiveTick = BlueprintNodeHelpers::HasBlueprintFunction(TEXT("ReceiveTick"), this, StopAtClass);
-	bImplementsReceiveExecutionStart = BlueprintNodeHelpers::HasBlueprintFunction(TEXT("ReceiveExecutionStart"), this, StopAtClass);
-	bImplementsReceiveExecutionFinish = BlueprintNodeHelpers::HasBlueprintFunction(TEXT("ReceiveExecutionFinish"), this, StopAtClass);
-	bImplementsReceiveObserverActivated = BlueprintNodeHelpers::HasBlueprintFunction(TEXT("ReceiveObserverActivated"), this, StopAtClass);
-	bImplementsReceiveObserverDeactivated = BlueprintNodeHelpers::HasBlueprintFunction(TEXT("ReceiveObserverDeactivated"), this, StopAtClass);
-	bImplementsReceiveConditionCheck = BlueprintNodeHelpers::HasBlueprintFunction(TEXT("ReceiveConditionCheck"), this, StopAtClass);
-
-	bNotifyBecomeRelevant = bImplementsReceiveObserverActivated;
-	bNotifyCeaseRelevant = bImplementsReceiveObserverDeactivated;
-	bNotifyTick = bImplementsReceiveTick;
-	bNotifyActivation = bImplementsReceiveExecutionStart;
-	bNotifyDeactivation = bImplementsReceiveExecutionFinish;
+	ReceiveTickImplementations = FBTNodeBPImplementationHelper::CheckEventImplementationVersion(TEXT("ReceiveTick"), TEXT("ReceiveTickAI"), this, StopAtClass);
+	ReceiveExecutionStartImplementations = FBTNodeBPImplementationHelper::CheckEventImplementationVersion(TEXT("ReceiveExecutionStart"), TEXT("ReceiveExecutionStartAI"), this, StopAtClass);
+	ReceiveExecutionFinishImplementations = FBTNodeBPImplementationHelper::CheckEventImplementationVersion(TEXT("ReceiveExecutionFinish"), TEXT("ReceiveExecutionFinishAI"), this, StopAtClass);
+	ReceiveObserverActivatedImplementations = FBTNodeBPImplementationHelper::CheckEventImplementationVersion(TEXT("ReceiveObserverActivated"), TEXT("ReceiveObserverActivatedAI"), this, StopAtClass);
+	ReceiveObserverDeactivatedImplementations = FBTNodeBPImplementationHelper::CheckEventImplementationVersion(TEXT("ReceiveObserverDeactivated"), TEXT("ReceiveObserverDeactivatedAI"), this, StopAtClass);
+	ReceiveConditionCheckImplementations = FBTNodeBPImplementationHelper::CheckEventImplementationVersion(TEXT("ReceiveConditionCheck"), TEXT("ReceiveConditionCheckAI"), this, StopAtClass);
+
+	bNotifyBecomeRelevant = ReceiveObserverActivatedImplementations != 0;
+	bNotifyCeaseRelevant = ReceiveObserverDeactivatedImplementations != 0;
+	bNotifyTick = ReceiveTickImplementations != 0;
+	bNotifyActivation = ReceiveExecutionStartImplementations != 0;
+	bNotifyDeactivation = ReceiveExecutionFinishImplementations != 0;
 	bShowPropertyDetails = true;
 
 	// all blueprint based nodes must create instances
@@ -60,11 +60,21 @@ void UBTDecorator_BlueprintBase::PostInitProperties()
 	BBKeyObserver = FOnBlackboardChange::CreateUObject(this, &UBTDecorator_BlueprintBase::OnBlackboardChange);
 }
 
+void UBTDecorator_BlueprintBase::SetOwner(AActor* InActorOwner)
+{
+	ActorOwner = InActorOwner;
+	AIOwner = Cast<AAIController>(InActorOwner);
+}
+
 void UBTDecorator_BlueprintBase::OnBecomeRelevant(UBehaviorTreeComponent* OwnerComp, uint8* NodeMemory)
 {
-	if (bImplementsReceiveObserverActivated)
+	if (AIOwner != nullptr && ReceiveObserverActivatedImplementations & FBTNodeBPImplementationHelper::AISpecific)
 	{
-		ReceiveObserverActivated(OwnerComp->GetOwner());
+		ReceiveObserverActivatedAI(AIOwner, AIOwner->GetPawn());
+	}
+	else if (ReceiveObserverActivatedImplementations & FBTNodeBPImplementationHelper::Generic)
+	{
+		ReceiveObserverActivated(ActorOwner);
 	}
 
 	UBlackboardComponent* BlackboardComp = OwnerComp->GetBlackboardComponent();
@@ -98,41 +108,68 @@ void UBTDecorator_BlueprintBase::OnCeaseRelevant(UBehaviorTreeComponent* OwnerCo
 		}
 	}
 		
-	if (bImplementsReceiveObserverDeactivated)
+	if (AIOwner != nullptr && ReceiveObserverDeactivatedImplementations & FBTNodeBPImplementationHelper::AISpecific)
+	{
+		ReceiveObserverDeactivatedAI(AIOwner, AIOwner->GetPawn());
+	}
+	else if (ReceiveObserverDeactivatedImplementations & FBTNodeBPImplementationHelper::Generic)
 	{
-		ReceiveObserverDeactivated(OwnerComp->GetOwner());
+		ReceiveObserverDeactivated(ActorOwner);
 	}
 }
 
 void UBTDecorator_BlueprintBase::OnNodeActivation(FBehaviorTreeSearchData& SearchData)
 {
-	// skip flag, will be handled by bNotifyActivation
-
-	ReceiveExecutionStart(SearchData.OwnerComp->GetOwner());
+	if (AIOwner != nullptr && ReceiveExecutionStartImplementations & FBTNodeBPImplementationHelper::AISpecific)
+	{
+		ReceiveExecutionStartAI(AIOwner, AIOwner->GetPawn());
+	}
+	else if (ReceiveExecutionStartImplementations & FBTNodeBPImplementationHelper::Generic)
+	{
+		ReceiveExecutionStart(ActorOwner);
+	}
 }
 
 void UBTDecorator_BlueprintBase::OnNodeDeactivation(FBehaviorTreeSearchData& SearchData, EBTNodeResult::Type NodeResult)
 {
-	// skip flag, will be handled by bNotifyDeactivation
-
-	ReceiveExecutionFinish(SearchData.OwnerComp->GetOwner(), NodeResult);
+	if (AIOwner != nullptr && ReceiveExecutionFinishImplementations & FBTNodeBPImplementationHelper::AISpecific)
+	{
+		ReceiveExecutionFinishAI(AIOwner, AIOwner->GetPawn(), NodeResult);
+	}
+	else if (ReceiveExecutionFinishImplementations & FBTNodeBPImplementationHelper::Generic)
+	{
+		ReceiveExecutionFinish(ActorOwner, NodeResult);
+	}
 }
 
 void UBTDecorator_BlueprintBase::TickNode(UBehaviorTreeComponent* OwnerComp, uint8* NodeMemory, float DeltaSeconds)
 {
-	// skip flag, will be handled by bNotifyTick
-
-	ReceiveTick(OwnerComp->GetOwner(), DeltaSeconds);
+	if (AIOwner != nullptr && ReceiveTickImplementations & FBTNodeBPImplementationHelper::AISpecific)
+	{
+		ReceiveTickAI(AIOwner, AIOwner->GetPawn(), DeltaSeconds);
+	}
+	else if (ReceiveTickImplementations & FBTNodeBPImplementationHelper::Generic)
+	{
+		ReceiveTick(ActorOwner, DeltaSeconds);
+	}
 }
 
 bool UBTDecorator_BlueprintBase::CalculateRawConditionValue(UBehaviorTreeComponent* OwnerComp, uint8* NodeMemory) const
 {
 	CurrentCallResult = false;
-	if (bImplementsReceiveConditionCheck)
+	if (ReceiveConditionCheckImplementations != 0)
 	{
 		// can't use const functions with blueprints
 		UBTDecorator_BlueprintBase* MyNode = (UBTDecorator_BlueprintBase*)this;
-		MyNode->ReceiveConditionCheck(OwnerComp->GetOwner());
+
+		if (AIOwner != nullptr && ReceiveConditionCheckImplementations & FBTNodeBPImplementationHelper::AISpecific)
+		{
+			MyNode->ReceiveConditionCheckAI(MyNode->AIOwner, MyNode->AIOwner->GetPawn());
+		}
+		else if (ReceiveConditionCheckImplementations & FBTNodeBPImplementationHelper::Generic)
+		{
+			MyNode->ReceiveConditionCheck(MyNode->ActorOwner);
+		}
 	}
 
 	return CurrentCallResult;
@@ -204,7 +241,7 @@ void UBTDecorator_BlueprintBase::OnInstanceDestroyed(UBehaviorTreeComponent* Own
 
 FName UBTDecorator_BlueprintBase::GetNodeIconName() const
 {
-	if(bImplementsReceiveConditionCheck)
+	if(ReceiveConditionCheckImplementations != 0)
 	{
 		return FName("BTEditor.Graph.BTNode.Decorator.Conditional.Icon");
 	}
diff --git a/Engine/Source/Runtime/AIModule/Private/BehaviorTree/Services/BTService_BlueprintBase.cpp b/Engine/Source/Runtime/AIModule/Private/BehaviorTree/Services/BTService_BlueprintBase.cpp
--- a/Engine/Source/Runtime/AIModule/Private/BehaviorTree/Services/BTService_BlueprintBase.cpp
+++ b/Engine/Source/Runtime/AIModule/Private/BehaviorTree/Services/BTService_BlueprintBase.cpp
@@ -7,15 +7,15 @@
 UBTService_BlueprintBase::UBTService_BlueprintBase(const FObjectInitializer& ObjectInitializer) : Super(ObjectInitializer)
 {
 	UClass* StopAtClass = UBTService_BlueprintBase::StaticClass();
-	bImplementsReceiveTick = BlueprintNodeHelpers::HasBlueprintFunction(TEXT("ReceiveTick"), this, StopAtClass);
-	bImplementsReceiveActivation = BlueprintNodeHelpers::HasBlueprintFunction(TEXT("ReceiveActivation"), this, StopAtClass);
-	bImplementsReceiveDeactivation = BlueprintNodeHelpers::HasBlueprintFunction(TEXT("ReceiveDeactivation"), this, StopAtClass);
-	bImplementsReceiveSearchStart = BlueprintNodeHelpers::HasBlueprintFunction(TEXT("ReceiveSearchStart"), this, StopAtClass);
+	ReceiveTickImplementations = FBTNodeBPImplementationHelper::CheckEventImplementationVersion(TEXT("ReceiveTick"), TEXT("ReceiveTickAI"), this, StopAtClass);
+	ReceiveActivationImplementations = FBTNodeBPImplementationHelper::CheckEventImplementationVersion(TEXT("ReceiveActivation"), TEXT("ReceiveActivationAI"), this, StopAtClass);
+	ReceiveDeactivationImplementations = FBTNodeBPImplementationHelper::CheckEventImplementationVersion(TEXT("ReceiveDeactivation"), TEXT("ReceiveDeactivationAI"), this, StopAtClass);
+	ReceiveSearchStartImplementations = FBTNodeBPImplementationHelper::CheckEventImplementationVersion(TEXT("ReceiveSearchStart"), TEXT("ReceiveSearchStartAI"), this, StopAtClass);
 
-	bNotifyBecomeRelevant = bImplementsReceiveActivation;
+	bNotifyBecomeRelevant = ReceiveActivationImplementations != 0;
 	bNotifyCeaseRelevant = bNotifyBecomeRelevant;
-	bNotifyOnSearch = bImplementsReceiveTick || bImplementsReceiveSearchStart;
-	bNotifyTick = bImplementsReceiveTick;
+	bNotifyOnSearch = ReceiveTickImplementations != 0 || ReceiveSearchStartImplementations != 0;
+	bNotifyTick = ReceiveTickImplementations != 0;
 	bShowPropertyDetails = true;
 
 	// all blueprint based nodes must create instances
@@ -33,14 +33,24 @@ void UBTService_BlueprintBase::PostInitProperties()
 	NodeName = BlueprintNodeHelpers::GetNodeName(this);
 }
 
+void UBTService_BlueprintBase::SetOwner(AActor* InActorOwner)
+{
+	ActorOwner = InActorOwner;
+	AIOwner = Cast<AAIController>(InActorOwner);
+}
+
 void UBTService_BlueprintBase::OnBecomeRelevant(UBehaviorTreeComponent* OwnerComp, uint8* NodeMemory)
 {
 	Super::OnBecomeRelevant(OwnerComp, NodeMemory);
 
 	// check flag, it could be used because user wants tick
-	if (bImplementsReceiveActivation)
+	if (AIOwner != nullptr && ReceiveActivationImplementations & FBTNodeBPImplementationHelper::AISpecific)
+	{
+		ReceiveActivationAI(AIOwner, AIOwner->GetPawn());
+	}
+	else if (ReceiveActivationImplementations & FBTNodeBPImplementationHelper::Generic)
 	{
-		ReceiveActivation(OwnerComp->GetOwner());
+		ReceiveActivation(ActorOwner);
 	}
 }
 
@@ -54,9 +64,13 @@ void UBTService_BlueprintBase::OnCeaseRelevant(UBehaviorTreeComponent* OwnerComp
 		// we can't have those resuming activity when node is/was aborted
 		BlueprintNodeHelpers::AbortLatentActions(OwnerComp, this);
 
-		if (bImplementsReceiveDeactivation)
+		if (AIOwner != nullptr && ReceiveActivationImplementations & FBTNodeBPImplementationHelper::AISpecific)
+		{
+			ReceiveDeactivationAI(AIOwner, AIOwner->GetPawn());
+		}
+		else if (ReceiveActivationImplementations & FBTNodeBPImplementationHelper::Generic)
 		{
-			ReceiveDeactivation(OwnerComp->GetOwner());
+			ReceiveDeactivation(ActorOwner);
 		}
 	}
 	else
@@ -75,9 +89,16 @@ void UBTService_BlueprintBase::OnSearchStart(FBehaviorTreeSearchData& SearchData
 {
 	// skip flag, will be handled by bNotifyOnSearch
 
-	if (bImplementsReceiveSearchStart)
+	if (ReceiveSearchStartImplementations != 0)
 	{
-		ReceiveSearchStart(SearchData.OwnerComp->GetOwner());
+		if (AIOwner != nullptr && ReceiveSearchStartImplementations & FBTNodeBPImplementationHelper::AISpecific)
+		{
+			ReceiveSearchStartAI(AIOwner, AIOwner->GetPawn());
+		}
+		else if (ReceiveSearchStartImplementations & FBTNodeBPImplementationHelper::Generic)
+		{
+			ReceiveSearchStart(ActorOwner);
+		}
 
 		const float NextTickTime = FMath::FRandRange(FMath::Max(0.0f, Interval - RandomDeviation), (Interval + RandomDeviation));
 		uint8* NodeMemory = GetNodeMemory<uint8>(SearchData);
@@ -93,9 +114,14 @@ void UBTService_BlueprintBase::TickNode(UBehaviorTreeComponent* OwnerComp, uint8
 {
 	Super::TickNode(OwnerComp, NodeMemory, DeltaSeconds);
 
-	// skip flag, will be handled by bNotifyTick
-
-	ReceiveTick(OwnerComp->GetOwner(), DeltaSeconds);
+	if (AIOwner != nullptr && ReceiveTickImplementations & FBTNodeBPImplementationHelper::AISpecific)
+	{
+		ReceiveTickAI(AIOwner, AIOwner->GetPawn(), DeltaSeconds);
+	}
+	else if (ReceiveTickImplementations & FBTNodeBPImplementationHelper::Generic)
+	{
+		ReceiveTick(ActorOwner, DeltaSeconds);
+	}
 }
 
 bool UBTService_BlueprintBase::IsServiceActive() const
@@ -113,10 +139,10 @@ FString UBTService_BlueprintBase::GetStaticServiceDescription() const
 	if (CDO)
 	{
 		ReturnDesc = FString::Printf(TEXT("%s, %s, %s, %s\n"),
-			bImplementsReceiveTick ? *GetStaticTickIntervalDescription() : TEXT("No tick"),
-			bImplementsReceiveActivation ? TEXT("Activation") : TEXT("No Activation"),
-			bImplementsReceiveDeactivation ? TEXT("Deactivation") : TEXT("No Deactivation"),
-			bImplementsReceiveSearchStart ? TEXT("Search Start") : TEXT("No Search Start"));
+			ReceiveTickImplementations != 0 ? *GetStaticTickIntervalDescription() : TEXT("No tick"),
+			ReceiveActivationImplementations != 0 ? TEXT("Activation") : TEXT("No Activation"),
+			ReceiveDeactivationImplementations != 0 ? TEXT("Deactivation") : TEXT("No Deactivation"),
+			ReceiveSearchStartImplementations != 0 ? TEXT("Search Start") : TEXT("No Search Start"));
 								
 		if (bShowPropertyDetails)
 		{
diff --git a/Engine/Source/Runtime/AIModule/Private/BehaviorTree/Tasks/BTTask_BlueprintBase.cpp b/Engine/Source/Runtime/AIModule/Private/BehaviorTree/Tasks/BTTask_BlueprintBase.cpp
--- a/Engine/Source/Runtime/AIModule/Private/BehaviorTree/Tasks/BTTask_BlueprintBase.cpp
+++ b/Engine/Source/Runtime/AIModule/Private/BehaviorTree/Tasks/BTTask_BlueprintBase.cpp
@@ -7,11 +7,11 @@
 UBTTask_BlueprintBase::UBTTask_BlueprintBase(const FObjectInitializer& ObjectInitializer) : Super(ObjectInitializer)
 {
 	UClass* StopAtClass = UBTTask_BlueprintBase::StaticClass();
-	bImplementsReceiveTick = BlueprintNodeHelpers::HasBlueprintFunction(TEXT("ReceiveTick"), this, StopAtClass);
-	bImplementsReceiveExecute = BlueprintNodeHelpers::HasBlueprintFunction(TEXT("ReceiveExecute"), this, StopAtClass);
-	bImplementsReceiveAbort = BlueprintNodeHelpers::HasBlueprintFunction(TEXT("ReceiveAbort"), this, StopAtClass);
+	ReceiveTickImplementations = FBTNodeBPImplementationHelper::CheckEventImplementationVersion(TEXT("ReceiveTick"), TEXT("ReceiveTickAI"), this, StopAtClass);
+	ReceiveExecuteImplementations = FBTNodeBPImplementationHelper::CheckEventImplementationVersion(TEXT("ReceiveExecute"), TEXT("ReceiveExecuteAI"), this, StopAtClass);
+	ReceiveAbortImplementations = FBTNodeBPImplementationHelper::CheckEventImplementationVersion(TEXT("ReceiveAbort"), TEXT("ReceiveAbortAI"), this, StopAtClass);
 
-	bNotifyTick = bImplementsReceiveTick;
+	bNotifyTick = ReceiveTickImplementations != FBTNodeBPImplementationHelper::NoImplementation;
 	bShowPropertyDetails = true;
 
 	// all blueprint based nodes must create instances
@@ -29,16 +29,29 @@ void UBTTask_BlueprintBase::PostInitProperties()
 	NodeName = BlueprintNodeHelpers::GetNodeName(this);
 }
 
+void UBTTask_BlueprintBase::SetOwner(AActor* InActorOwner) 
+{ 
+	ActorOwner = InActorOwner;
+	AIOwner = Cast<AAIController>(InActorOwner);
+}
+
 EBTNodeResult::Type UBTTask_BlueprintBase::ExecuteTask(UBehaviorTreeComponent* OwnerComp, uint8* NodeMemory)
 {
 	// fail when task doesn't react to execution (start or tick)
-	CurrentCallResult = (bImplementsReceiveExecute || bImplementsReceiveTick) ? EBTNodeResult::InProgress : EBTNodeResult::Failed;
+	CurrentCallResult = (ReceiveExecuteImplementations != 0 || ReceiveTickImplementations != 0) ? EBTNodeResult::InProgress : EBTNodeResult::Failed;
 
-	if (bImplementsReceiveExecute)
+	if (ReceiveExecuteImplementations != FBTNodeBPImplementationHelper::NoImplementation)
 	{
 		bStoreFinishResult = true;
 
-		ReceiveExecute(OwnerComp->GetOwner());
+		if (AIOwner != nullptr && ReceiveExecuteImplementations & FBTNodeBPImplementationHelper::AISpecific)
+		{
+			ReceiveExecuteAI(AIOwner, AIOwner->GetPawn());
+		}
+		else if (ReceiveExecuteImplementations & FBTNodeBPImplementationHelper::Generic)
+		{
+			ReceiveExecute(ActorOwner);
+		}
 
 		bStoreFinishResult = false;
 	}
@@ -52,12 +65,20 @@ EBTNodeResult::Type UBTTask_BlueprintBase::AbortTask(UBehaviorTreeComponent* Own
 	// we can't have those resuming activity when node is/was aborted
 	BlueprintNodeHelpers::AbortLatentActions(OwnerComp, this);
 
-	CurrentCallResult = bImplementsReceiveAbort ? EBTNodeResult::InProgress : EBTNodeResult::Aborted;
-	if (bImplementsReceiveAbort)
+	CurrentCallResult = ReceiveAbortImplementations != 0 ? EBTNodeResult::InProgress : EBTNodeResult::Aborted;
+
+	if (ReceiveAbortImplementations != FBTNodeBPImplementationHelper::NoImplementation)
 	{
 		bStoreFinishResult = true;
 
-		ReceiveAbort(OwnerComp->GetOwner());
+		if (AIOwner != nullptr && ReceiveAbortImplementations & FBTNodeBPImplementationHelper::AISpecific)
+		{
+			ReceiveAbortAI(AIOwner, AIOwner->GetPawn());
+		}
+		else if (ReceiveAbortImplementations & FBTNodeBPImplementationHelper::Generic)
+		{
+			ReceiveAbort(ActorOwner);
+		}
 
 		bStoreFinishResult = false;
 	}
@@ -67,9 +88,14 @@ EBTNodeResult::Type UBTTask_BlueprintBase::AbortTask(UBehaviorTreeComponent* Own
 
 void UBTTask_BlueprintBase::TickTask(UBehaviorTreeComponent* OwnerComp, uint8* NodeMemory, float DeltaSeconds) 
 {
-	// skip flag, will be handled by bNotifyTick
-
-	ReceiveTick(OwnerComp->GetOwner(), DeltaSeconds);
+	if (AIOwner != nullptr && ReceiveTickImplementations & FBTNodeBPImplementationHelper::AISpecific)
+	{
+		ReceiveTickAI(AIOwner, AIOwner->GetPawn(), DeltaSeconds);
+	}
+	else if (ReceiveTickImplementations & FBTNodeBPImplementationHelper::Generic)
+	{
+		ReceiveTick(ActorOwner, DeltaSeconds);
+	}
 }
 
 void UBTTask_BlueprintBase::FinishExecute(bool bSuccess)
