Non-unity compilation fix #UE4

diff --git a/Engine/Source/Runtime/AIModule/Classes/AITypes.h b/Engine/Source/Runtime/AIModule/Classes/AITypes.h
--- a/Engine/Source/Runtime/AIModule/Classes/AITypes.h
+++ b/Engine/Source/Runtime/AIModule/Classes/AITypes.h
@@ -180,7 +180,7 @@ public:
 	}
 
 	FAINamedID()
-		: Index(TCounter::Type(-1)), Name(TEXT("Invalid"))
+		: Index(typename TCounter::Type(-1)), Name(TEXT("Invalid"))
 	{}
 
 	operator typename TCounter::Type() const { return Index; }
@@ -223,7 +223,7 @@ public:
 	}
 
 	FAIGenericID()
-		: Index(TCounter::Type(-1))
+		: Index(typename TCounter::Type(-1))
 	{}
 
 	static FAIGenericID GetNextID() { return FAIGenericID(GetCounter().GetNextAvailableID()); }
diff --git a/Engine/Source/Runtime/AIModule/Classes/Perception/AIPerceptionSystem.h b/Engine/Source/Runtime/AIModule/Classes/Perception/AIPerceptionSystem.h
--- a/Engine/Source/Runtime/AIModule/Classes/Perception/AIPerceptionSystem.h
+++ b/Engine/Source/Runtime/AIModule/Classes/Perception/AIPerceptionSystem.h
@@ -74,8 +74,8 @@ public:
 	template<typename FEventClass>
 	void OnEvent(const FEventClass& Event)
 	{
-		check(Senses[UAISense::GetSenseID<FEventClass::FSenseClass>()]);
-		((typename FEventClass::FSenseClass*)Senses[UAISense::GetSenseID<FEventClass::FSenseClass>()])->RegisterEvent(Event);
+		check(Senses[UAISense::GetSenseID<typename FEventClass::FSenseClass>()]);
+		((typename FEventClass::FSenseClass*)Senses[UAISense::GetSenseID<typename FEventClass::FSenseClass>()])->RegisterEvent(Event);
 	}
 
 	template<typename FEventClass>
diff --git a/Engine/Source/Runtime/AIModule/Classes/Perception/AISenseBlueprintListener.h b/Engine/Source/Runtime/AIModule/Classes/Perception/AISenseBlueprintListener.h
--- a/Engine/Source/Runtime/AIModule/Classes/Perception/AISenseBlueprintListener.h
+++ b/Engine/Source/Runtime/AIModule/Classes/Perception/AISenseBlueprintListener.h
@@ -2,6 +2,7 @@
 
 #pragma once
 
+#include "Engine/UserDefinedStruct.h"
 #include "AISenseBlueprintListener.generated.h"
 
 UCLASS(ClassGroup=AI)
diff --git a/Engine/Source/Runtime/AIModule/Classes/Perception/AISenseEvent_Damage.h b/Engine/Source/Runtime/AIModule/Classes/Perception/AISenseEvent_Damage.h
--- a/Engine/Source/Runtime/AIModule/Classes/Perception/AISenseEvent_Damage.h
+++ b/Engine/Source/Runtime/AIModule/Classes/Perception/AISenseEvent_Damage.h
@@ -3,6 +3,7 @@
 #pragma once
 
 #include "Perception/AISenseEvent.h"
+#include "Perception/AISense_Damage.h"
 #include "AISenseEvent_Damage.generated.h"
 
 UCLASS()
diff --git a/Engine/Source/Runtime/AIModule/Private/AIModulePrivate.h b/Engine/Source/Runtime/AIModule/Private/AIModulePrivate.h
--- a/Engine/Source/Runtime/AIModule/Private/AIModulePrivate.h
+++ b/Engine/Source/Runtime/AIModule/Private/AIModulePrivate.h
@@ -57,6 +57,7 @@
 #include "GenericTeamAgentInterface.h"
 
 #include "Perception/AISense.h"
+#include "Perception/AISense_Sight.h"
 #include "Perception/AISense_Hearing.h"
 #include "Perception/AISenseEvent.h"
 #include "Perception/AISenseConfig.h"
@@ -65,6 +66,7 @@
 #include "VisualLogger/VisualLogger.h"
 
 #include "BlueprintNodeHelpers.h"
+#include "BehaviorTree/BTFunctionLibrary.h"
 
 DECLARE_LOG_CATEGORY_EXTERN(LogAINavigation, Warning, All);
 DECLARE_LOG_CATEGORY_EXTERN(LogBehaviorTree, Display, All);
diff --git a/Engine/Source/Runtime/AIModule/Private/Perception/AISenseEvent.cpp b/Engine/Source/Runtime/AIModule/Private/Perception/AISenseEvent.cpp
--- a/Engine/Source/Runtime/AIModule/Private/Perception/AISenseEvent.cpp
+++ b/Engine/Source/Runtime/AIModule/Private/Perception/AISenseEvent.cpp
@@ -4,6 +4,7 @@
 #include "Perception/AISenseEvent.h"
 #include "Perception/AISenseEvent_Hearing.h"
 #include "Perception/AISenseEvent_Damage.h"
+#include "Perception/AISense_Damage.h"
 
 //----------------------------------------------------------------------//
 // UAISenseEvent_Hearing
diff --git a/Engine/Source/Runtime/AIModule/Private/Perception/AISense_Damage.cpp b/Engine/Source/Runtime/AIModule/Private/Perception/AISense_Damage.cpp
--- a/Engine/Source/Runtime/AIModule/Private/Perception/AISense_Damage.cpp
+++ b/Engine/Source/Runtime/AIModule/Private/Perception/AISense_Damage.cpp
@@ -2,6 +2,7 @@
 
 #include "AIModulePrivate.h"
 #include "Perception/AISense_Damage.h"
+#include "Perception/AISenseEvent_Damage.h"
 
 //----------------------------------------------------------------------//
 // 
diff --git a/Engine/Source/Runtime/AIModule/Private/Perception/AISense_Hearing.cpp b/Engine/Source/Runtime/AIModule/Private/Perception/AISense_Hearing.cpp
--- a/Engine/Source/Runtime/AIModule/Private/Perception/AISense_Hearing.cpp
+++ b/Engine/Source/Runtime/AIModule/Private/Perception/AISense_Hearing.cpp
@@ -2,6 +2,8 @@
 
 #include "AIModulePrivate.h"
 #include "Perception/AISense_Hearing.h"
+#include "Perception/AISenseConfig_Hearing.h"
+#include "Perception/AISenseEvent_Hearing.h"
 
 //----------------------------------------------------------------------//
 // FAINoiseEvent
