UE-4496 Projects fail to package in Shipping for Windows 32-bit (Compile errors)

diff --git a/Engine/Source/Runtime/AIModule/Private/Perception/AIPerceptionComponent.cpp b/Engine/Source/Runtime/AIModule/Private/Perception/AIPerceptionComponent.cpp
--- a/Engine/Source/Runtime/AIModule/Private/Perception/AIPerceptionComponent.cpp
+++ b/Engine/Source/Runtime/AIModule/Private/Perception/AIPerceptionComponent.cpp
@@ -542,6 +542,7 @@ bool UAIPerceptionComponent::GetActorsPerception(AActor* Actor, FActorPerception
 //----------------------------------------------------------------------//
 void UAIPerceptionComponent::DrawDebugInfo(UCanvas* Canvas)
 {
+#if !UE_BUILD_SHIPPING
 	if (Canvas == nullptr)
 	{
 		return;
@@ -594,6 +595,7 @@ void UAIPerceptionComponent::DrawDebugInfo(UCanvas* Canvas)
 			Sense->DrawDebugInfo(*Canvas, *this);
 		}
 	}
+#endif // #if !UE_BUILD_SHIPPING
 }
 
 #if ENABLE_VISUAL_LOG
