Shipping compilation fixes #UE4

diff --git a/Engine/Source/Developer/GameplayDebugger/Private/GameplayDebuggingHUDComponent.cpp b/Engine/Source/Developer/GameplayDebugger/Private/GameplayDebuggingHUDComponent.cpp
--- a/Engine/Source/Developer/GameplayDebugger/Private/GameplayDebuggingHUDComponent.cpp
+++ b/Engine/Source/Developer/GameplayDebugger/Private/GameplayDebuggingHUDComponent.cpp
@@ -585,6 +585,7 @@ void AGameplayDebuggingHUDComponent::DrawEQSItemDetails(int32 ItemIdx, class UGa
 
 void AGameplayDebuggingHUDComponent::DrawPerception(APlayerController* PC, class UGameplayDebuggingComponent *DebugComponent)
 {
+#if !(UE_BUILD_SHIPPING || UE_BUILD_TEST)
 	if (!DebugComponent)
 	{
 		return;
@@ -625,6 +626,7 @@ void AGameplayDebuggingHUDComponent::DrawPerception(APlayerController* PC, class
 			}
 		}
 	}
+#endif // !(UE_BUILD_SHIPPING || UE_BUILD_TEST)
 }
 
 void AGameplayDebuggingHUDComponent::DrawNavMeshSnapshot(APlayerController* PC, class UGameplayDebuggingComponent *DebugComponent)
diff --git a/Engine/Source/Runtime/AIModule/Classes/Perception/AIPerceptionComponent.h b/Engine/Source/Runtime/AIModule/Classes/Perception/AIPerceptionComponent.h
--- a/Engine/Source/Runtime/AIModule/Classes/Perception/AIPerceptionComponent.h
+++ b/Engine/Source/Runtime/AIModule/Classes/Perception/AIPerceptionComponent.h
@@ -218,7 +218,9 @@ public:
 	bool HasAnyActiveStimulus(const AActor& Source) const;
 	bool HasActiveStimulus(const AActor& Source, FAISenseID Sense) const;
 
+#if !UE_BUILD_SHIPPING
 	void DrawDebugInfo(UCanvas* Canvas);
+#endif // !UE_BUILD_SHIPPING
 
 #if ENABLE_VISUAL_LOG
 	virtual void DescribeSelfToVisLog(FVisualLogEntry* Snapshot) const;
diff --git a/Engine/Source/Runtime/AIModule/Private/Perception/AIPerceptionComponent.cpp b/Engine/Source/Runtime/AIModule/Private/Perception/AIPerceptionComponent.cpp
--- a/Engine/Source/Runtime/AIModule/Private/Perception/AIPerceptionComponent.cpp
+++ b/Engine/Source/Runtime/AIModule/Private/Perception/AIPerceptionComponent.cpp
@@ -540,9 +540,9 @@ bool UAIPerceptionComponent::GetActorsPerception(AActor* Actor, FActorPerception
 //----------------------------------------------------------------------//
 // debug
 //----------------------------------------------------------------------//
+#if !UE_BUILD_SHIPPING
 void UAIPerceptionComponent::DrawDebugInfo(UCanvas* Canvas)
 {
-#if !UE_BUILD_SHIPPING
 	if (Canvas == nullptr)
 	{
 		return;
@@ -595,8 +595,8 @@ void UAIPerceptionComponent::DrawDebugInfo(UCanvas* Canvas)
 			Sense->DrawDebugInfo(*Canvas, *this);
 		}
 	}
-#endif // #if !UE_BUILD_SHIPPING
 }
+#endif // !UE_BUILD_SHIPPING
 
 #if ENABLE_VISUAL_LOG
 void UAIPerceptionComponent::DescribeSelfToVisLog(FVisualLogEntry* Snapshot) const
diff --git a/Engine/Source/Runtime/AIModule/Private/Perception/AISense.cpp b/Engine/Source/Runtime/AIModule/Private/Perception/AISense.cpp
--- a/Engine/Source/Runtime/AIModule/Private/Perception/AISense.cpp
+++ b/Engine/Source/Runtime/AIModule/Private/Perception/AISense.cpp
@@ -92,7 +92,7 @@ void UAISenseConfig::DrawDebugInfo(UCanvas& Canvas, UAIPerceptionComponent& Perc
 {
 
 }
-
+#endif // !UE_BUILD_SHIPPING
 //----------------------------------------------------------------------//
 // 
 //----------------------------------------------------------------------//
@@ -100,7 +100,7 @@ TSubclassOf<UAISense> UAISenseConfig_Sight::GetSenseImplementation() const
 { 
 	return *Implementation; 
 }
-
+#if !UE_BUILD_SHIPPING
 void UAISenseConfig_Sight::DrawDebugInfo(UCanvas& Canvas, UAIPerceptionComponent& PerceptionComponent) const
 {
 	/*PeripheralVisionAngleDegrees*/
@@ -119,7 +119,7 @@ void UAISenseConfig_Sight::DrawDebugInfo(UCanvas& Canvas, UAIPerceptionComponent
 		DrawDebugLine(World, BodyLocation, BodyLocation + (BodyFacing.RotateAngleAxis(-PeripheralVisionAngleDegrees, FVector::UpVector) * SightPieLength), UAISense_Sight::GetDebugSightRangeColor());
 	}
 }
-
+#endif // !UE_BUILD_SHIPPING
 //----------------------------------------------------------------------//
 // UAISenseConfig_Hearing
 //----------------------------------------------------------------------//
@@ -127,7 +127,7 @@ TSubclassOf<UAISense> UAISenseConfig_Hearing::GetSenseImplementation() const
 { 
 	return *Implementation; 
 }
-
+#if !UE_BUILD_SHIPPING
 void UAISenseConfig_Hearing::DrawDebugInfo(UCanvas& Canvas, UAIPerceptionComponent& PerceptionComponent) const
 {	
 	const AActor* BodyActor = PerceptionComponent.GetBodyActor();
