Fixed a crash in reporting AI perception event when no one is listening to given sense #UE4

diff --git a/Engine/Source/Runtime/AIModule/Classes/Perception/AIPerceptionSystem.h b/Engine/Source/Runtime/AIModule/Classes/Perception/AIPerceptionSystem.h
--- a/Engine/Source/Runtime/AIModule/Classes/Perception/AIPerceptionSystem.h
+++ b/Engine/Source/Runtime/AIModule/Classes/Perception/AIPerceptionSystem.h
@@ -74,8 +74,12 @@ public:
 	template<typename FEventClass>
 	void OnEvent(const FEventClass& Event)
 	{
-		check(Senses[UAISense::GetSenseID<typename FEventClass::FSenseClass>()]);
-		((typename FEventClass::FSenseClass*)Senses[UAISense::GetSenseID<typename FEventClass::FSenseClass>()])->RegisterEvent(Event);
+		const FAISenseID SenseID = UAISense::GetSenseID<typename FEventClass::FSenseClass>();
+		if (Senses.IsValidIndex(SenseID) && Senses[SenseID] != nullptr)
+		{
+			((typename FEventClass::FSenseClass*)Senses[SenseID])->RegisterEvent(Event);
+		}
+		// otherwise there's no one interested in this event, skip it.
 	}
 
 	template<typename FEventClass>
