Non-unity fixes #UE4

diff --git a/Engine/Source/Editor/UnrealEd/Public/Kismet2/KismetEditorUtilities.h b/Engine/Source/Editor/UnrealEd/Public/Kismet2/KismetEditorUtilities.h
--- a/Engine/Source/Editor/UnrealEd/Public/Kismet2/KismetEditorUtilities.h
+++ b/Engine/Source/Editor/UnrealEd/Public/Kismet2/KismetEditorUtilities.h
@@ -6,6 +6,8 @@
 
 #pragma once
 
+#include "Engine/Blueprint.h"
+
 class UBlueprintGeneratedClass;
 
 //////////////////////////////////////////////////////////////////////////
diff --git a/Engine/Source/Runtime/AIModule/Private/Perception/AISense.cpp b/Engine/Source/Runtime/AIModule/Private/Perception/AISense.cpp
--- a/Engine/Source/Runtime/AIModule/Private/Perception/AISense.cpp
+++ b/Engine/Source/Runtime/AIModule/Private/Perception/AISense.cpp
@@ -5,6 +5,10 @@
 #include "Perception/AISenseConfig_Sight.h"
 #include "Perception/AISenseConfig_Hearing.h"
 
+#if !UE_BUILD_SHIPPING
+#include "DrawDebugHelpers.h"
+#endif // !UE_BUILD_SHIPPING
+
 const float UAISense::SuspendNextUpdate = FLT_MAX;
 
 UAISense::UAISense(const FObjectInitializer& ObjectInitializer)
