AI Perception crash-fixes #UE4

diff --git a/Engine/Source/Runtime/AIModule/Private/Perception/AIPerceptionComponent.cpp b/Engine/Source/Runtime/AIModule/Private/Perception/AIPerceptionComponent.cpp
--- a/Engine/Source/Runtime/AIModule/Private/Perception/AIPerceptionComponent.cpp
+++ b/Engine/Source/Runtime/AIModule/Private/Perception/AIPerceptionComponent.cpp
@@ -381,7 +381,8 @@ void UAIPerceptionComponent::ProcessStimuli()
 
 				if (PerceptualInfo->LastSensedStimuli.Num() <= SourcedStimulus->Stimulus.Type)
 				{
-					for (int32 Index = 0; Index < SourcedStimulus->Stimulus.Type - PerceptualInfo->LastSensedStimuli.Num() + 1; ++Index)
+					const int32 NumberToAdd = SourcedStimulus->Stimulus.Type - PerceptualInfo->LastSensedStimuli.Num() + 1;
+					for (int32 Index = 0; Index < NumberToAdd; ++Index)
 					{
 						PerceptualInfo->LastSensedStimuli.Add(FAIStimulus());
 					}
diff --git a/Engine/Source/Runtime/AIModule/Private/Perception/AIPerceptionSystem.cpp b/Engine/Source/Runtime/AIModule/Private/Perception/AIPerceptionSystem.cpp
--- a/Engine/Source/Runtime/AIModule/Private/Perception/AIPerceptionSystem.cpp
+++ b/Engine/Source/Runtime/AIModule/Private/Perception/AIPerceptionSystem.cpp
@@ -113,7 +113,7 @@ void UAIPerceptionSystem::PerformSourceRegistration()
 	for (const auto& PercSource : SourcesToRegister)
 	{
 		AActor* SourceActor = PercSource.Source.Get();
-		if (SourceActor)
+		if (SourceActor && Senses.IsValidIndex(PercSource.SenseID))
 		{
 			Senses[PercSource.SenseID]->RegisterSource(*SourceActor);
 		}
