File missing from previous CL #UE4

diff --git a/Engine/Source/Runtime/AIModule/Private/Perception/AIPerceptionComponent.cpp b/Engine/Source/Runtime/AIModule/Private/Perception/AIPerceptionComponent.cpp
--- a/Engine/Source/Runtime/AIModule/Private/Perception/AIPerceptionComponent.cpp
+++ b/Engine/Source/Runtime/AIModule/Private/Perception/AIPerceptionComponent.cpp
@@ -378,15 +378,15 @@ void UAIPerceptionComponent::ProcessStimuli()
 				PerceptualInfo->DominantSense = DominantSenseID;
 
 				PerceptualInfo->bIsHostile = AIOwner != NULL && FGenericTeamId::GetAttitude(AIOwner, SourcedStimulus->Source) == ETeamAttitude::Hostile;
+			}
+		}
 
-				if (PerceptualInfo->LastSensedStimuli.Num() <= SourcedStimulus->Stimulus.Type)
-				{
-					const int32 NumberToAdd = SourcedStimulus->Stimulus.Type - PerceptualInfo->LastSensedStimuli.Num() + 1;
-					for (int32 Index = 0; Index < NumberToAdd; ++Index)
-					{
-						PerceptualInfo->LastSensedStimuli.Add(FAIStimulus());
-					}
-				}
+		if (PerceptualInfo->LastSensedStimuli.Num() <= SourcedStimulus->Stimulus.Type)
+		{
+			const int32 NumberToAdd = SourcedStimulus->Stimulus.Type - PerceptualInfo->LastSensedStimuli.Num() + 1;
+			for (int32 Index = 0; Index < NumberToAdd; ++Index)
+			{
+				PerceptualInfo->LastSensedStimuli.Add(FAIStimulus());
 			}
 		}
 
