Refactored INavAgentInterface::GetNavAgentProperties to return a reference rather then a pointer to FNavAgentProperties #UE4

The main reason being lots of code calling GetNavAgentProperties assumed it was retrieving valid data anyway, and besides it returning NULL was an odd case which still meant "use defaults" rather then "invalid".

diff --git a/Engine/Source/Runtime/AIModule/Private/EnvironmentQuery/EnvQueryTypes.cpp b/Engine/Source/Runtime/AIModule/Private/EnvironmentQuery/EnvQueryTypes.cpp
--- a/Engine/Source/Runtime/AIModule/Private/EnvironmentQuery/EnvQueryTypes.cpp
+++ b/Engine/Source/Runtime/AIModule/Private/EnvironmentQuery/EnvQueryTypes.cpp
@@ -238,11 +238,8 @@ namespace FEQSHelpers
 		INavAgentInterface* NavAgent = QueryInstance.Owner.IsValid() ? Cast<INavAgentInterface>(QueryInstance.Owner.Get()) : NULL;
 		if (NavAgent)
 		{
-			const FNavAgentProperties* NavAgentProps = NavAgent ? NavAgent->GetNavAgentProperties() : NULL;
-			if (NavAgentProps != NULL)
-			{
-				return Cast<const ARecastNavMesh>(NavSys->GetNavDataForProps(*NavAgentProps));
-			}
+			const FNavAgentProperties& NavAgentProps = NavAgent ? NavAgent->GetNavAgentProperties() : FNavAgentProperties::DefaultProperties;
+			return Cast<const ARecastNavMesh>(NavSys->GetNavDataForProps(NavAgentProps));
 		}
 
 		return Cast<const ARecastNavMesh>(NavSys->GetMainNavData());
diff --git a/Engine/Source/Runtime/AIModule/Private/EnvironmentQuery/Tests/EnvQueryTest_Pathfinding.cpp b/Engine/Source/Runtime/AIModule/Private/EnvironmentQuery/Tests/EnvQueryTest_Pathfinding.cpp
--- a/Engine/Source/Runtime/AIModule/Private/EnvironmentQuery/Tests/EnvQueryTest_Pathfinding.cpp
+++ b/Engine/Source/Runtime/AIModule/Private/EnvironmentQuery/Tests/EnvQueryTest_Pathfinding.cpp
@@ -245,7 +245,7 @@ ANavigationData* UEnvQueryTest_Pathfinding::FindNavigationData(UNavigationSystem
 	INavAgentInterface* NavAgent = Cast<INavAgentInterface>(Owner);
 	if (NavAgent)
 	{
-		return NavSys->GetNavDataForProps(*NavAgent->GetNavAgentProperties());
+		return NavSys->GetNavDataForProps(NavAgent->GetNavAgentProperties());
 	}
 
 	return NavSys->GetMainNavData(FNavigationSystem::DontCreate);
diff --git a/Engine/Source/Runtime/AIModule/Private/Navigation/NavigationComponent.cpp b/Engine/Source/Runtime/AIModule/Private/Navigation/NavigationComponent.cpp
--- a/Engine/Source/Runtime/AIModule/Private/Navigation/NavigationComponent.cpp
+++ b/Engine/Source/Runtime/AIModule/Private/Navigation/NavigationComponent.cpp
@@ -216,36 +216,29 @@ bool UNavigationComponent::AsyncGeneratePath(const FVector& FromLocation, const
 	// best to do so using a flag while keeping this as the default behavior.  NOTE: If any changes to this behavior
 	// are made, you should search for other places in UNavigationComponent using NavMeshGoalLocation and make sure the
 	// behavior remains consistent.
-	const FNavAgentProperties* NavAgentProps = MyNavAgent->GetNavAgentProperties();
-	if (NavAgentProps != NULL)
-	{
-		FVector NavMeshGoalLocation;
-
-		if (ProjectPointToNavigation(GoalLocation, NavMeshGoalLocation) == QuerySuccess)
-		{
-			const EPathFindingMode::Type Mode = bDoHierarchicalPathfinding ? EPathFindingMode::Hierarchical : EPathFindingMode::Regular;
-			FPathFindingQuery Query(GetOwner(), GetNavData(), FromLocation, NavMeshGoalLocation, QueryFilter);
-			Query.NavDataFlags = NavDataFlags;
+		
+	FVector NavMeshGoalLocation;
 
-			AsynPathQueryID = NavSys->FindPathAsync(*MyNavAgent->GetNavAgentProperties(), Query,
-				FNavPathQueryDelegate::CreateUObject(this, &UNavigationComponent::AsyncGeneratePath_OnCompleteCallback), Mode);
-		}
-		else
-		{
-			// projecting destination point failed
-			AsynPathQueryID = INVALID_NAVQUERYID;
+	if (ProjectPointToNavigation(GoalLocation, NavMeshGoalLocation) == QuerySuccess)
+	{
+		const EPathFindingMode::Type Mode = bDoHierarchicalPathfinding ? EPathFindingMode::Hierarchical : EPathFindingMode::Regular;
+		FPathFindingQuery Query(GetOwner(), GetNavData(), FromLocation, NavMeshGoalLocation, QueryFilter);
+		Query.NavDataFlags = NavDataFlags;
 
-			UE_VLOG(GetOwner(), LogAINavigation, Warning, TEXT("AI trying to generate path to point off of navmesh"));
-			UE_VLOG_LOCATION(GetOwner(), LogAINavigation, Warning, GoalLocation, 30, FColor::Red, TEXT("Invalid pathing GoalLocation"));
-			if (MyNavData != NULL)
-			{
-				UE_VLOG_BOX(GetOwner(), LogAINavigation, Warning, FBox(GoalLocation - GetQueryExtent(), GoalLocation + GetQueryExtent()), FColor::Red, TEXT_EMPTY);
-			}
-		}
+		AsynPathQueryID = NavSys->FindPathAsync(MyNavAgent->GetNavAgentProperties(), Query,
+			FNavPathQueryDelegate::CreateUObject(this, &UNavigationComponent::AsyncGeneratePath_OnCompleteCallback), Mode);
 	}
 	else
 	{
+		// projecting destination point failed
 		AsynPathQueryID = INVALID_NAVQUERYID;
+
+		UE_VLOG(GetOwner(), LogAINavigation, Warning, TEXT("AI trying to generate path to point off of navmesh"));
+		UE_VLOG_LOCATION(GetOwner(), LogAINavigation, Warning, GoalLocation, 30, FColor::Red, TEXT("Invalid pathing GoalLocation"));
+		if (MyNavData != NULL)
+		{
+			UE_VLOG_BOX(GetOwner(), LogAINavigation, Warning, FBox(GoalLocation - GetQueryExtent(), GoalLocation + GetQueryExtent()), FColor::Red, TEXT_EMPTY);
+		}
 	}
 
 	return AsynPathQueryID != INVALID_NAVQUERYID;
@@ -358,85 +351,77 @@ bool UNavigationComponent::GeneratePathTo(const FVector& GoalLocation, TSharedPt
 
 	check(MyNavAgent);
 
-	const FNavAgentProperties* NavAgentProps = MyNavAgent->GetNavAgentProperties();
-	if (NavAgentProps != NULL)
-	{
-		FVector NavMeshGoalLocation;
+	FVector NavMeshGoalLocation;
 
 #if ENABLE_VISUAL_LOG
-		UE_VLOG_LOCATION(GetOwner(), LogAINavigation, Log, GoalLocation, 34, FColor(0, 127, 14), TEXT("GoalLocation"));
-		UE_VLOG_BOX(GetOwner(), LogAINavigation, Log, FBox(GoalLocation - GetQueryExtent(), GoalLocation + GetQueryExtent()), FColor::Green, TEXT_EMPTY);
+	UE_VLOG_LOCATION(GetOwner(), LogAINavigation, Log, GoalLocation, 34, FColor(0, 127, 14), TEXT("GoalLocation"));
+	UE_VLOG_BOX(GetOwner(), LogAINavigation, Log, FBox(GoalLocation - GetQueryExtent(), GoalLocation + GetQueryExtent()), FColor::Green, TEXT_EMPTY);
 #endif
 
-		if (ProjectPointToNavigation(GoalLocation, NavMeshGoalLocation) == QuerySuccess)
-		{
-			UE_VLOG_SEGMENT(GetOwner(), LogAINavigation, Log, GoalLocation, NavMeshGoalLocation, FColor::Blue, TEXT_EMPTY);
-			UE_VLOG_LOCATION(GetOwner(), LogAINavigation, Log, NavMeshGoalLocation, 34, FColor::Blue, TEXT_EMPTY);
+	if (ProjectPointToNavigation(GoalLocation, NavMeshGoalLocation) == QuerySuccess)
+	{
+		UE_VLOG_SEGMENT(GetOwner(), LogAINavigation, Log, GoalLocation, NavMeshGoalLocation, FColor::Blue, TEXT_EMPTY);
+		UE_VLOG_LOCATION(GetOwner(), LogAINavigation, Log, NavMeshGoalLocation, 34, FColor::Blue, TEXT_EMPTY);
 			
-			UNavigationSystem* NavSys = GetWorld()->GetNavigationSystem();
-			FPathFindingQuery Query(GetOwner(), MyNavData, MyNavAgent->GetNavAgentLocation(), NavMeshGoalLocation, QueryFilter);
-			Query.NavDataFlags = NavDataFlags;
+		UNavigationSystem* NavSys = GetWorld()->GetNavigationSystem();
+		FPathFindingQuery Query(GetOwner(), MyNavData, MyNavAgent->GetNavAgentLocation(), NavMeshGoalLocation, QueryFilter);
+		Query.NavDataFlags = NavDataFlags;
 
-			const EPathFindingMode::Type Mode = bDoHierarchicalPathfinding ? EPathFindingMode::Hierarchical : EPathFindingMode::Regular;
-			FPathFindingResult Result = NavSys->FindPathSync(*MyNavAgent->GetNavAgentProperties(), Query, Mode);
+		const EPathFindingMode::Type Mode = bDoHierarchicalPathfinding ? EPathFindingMode::Hierarchical : EPathFindingMode::Regular;
+		FPathFindingResult Result = NavSys->FindPathSync(MyNavAgent->GetNavAgentProperties(), Query, Mode);
 
-			// try reversing direction
-			if (bSearchFromGoalWhenOutOfNodes && !bDoHierarchicalPathfinding &&
-				Result.Path.IsValid() && Result.Path->DidSearchReachedLimit())
+		// try reversing direction
+		if (bSearchFromGoalWhenOutOfNodes && !bDoHierarchicalPathfinding &&
+			Result.Path.IsValid() && Result.Path->DidSearchReachedLimit())
+		{
+			// quick check if path exists
+			bool bPathExists = false;
 			{
-				// quick check if path exists
-				bool bPathExists = false;
-				{
-					DECLARE_SCOPE_CYCLE_COUNTER(TEXT("Pathfinding: HPA* test time"), STAT_Navigation_PathVerifyTime, STATGROUP_Navigation);
-					bPathExists = NavSys->TestPathSync(Query, EPathFindingMode::Hierarchical);
-				}
+				DECLARE_SCOPE_CYCLE_COUNTER(TEXT("Pathfinding: HPA* test time"), STAT_Navigation_PathVerifyTime, STATGROUP_Navigation);
+				bPathExists = NavSys->TestPathSync(Query, EPathFindingMode::Hierarchical);
+			}
 
-				UE_VLOG(GetOwner(), LogAINavigation, Log, TEXT("GeneratePathTo: out of nodes, HPA* path: %s"),
-					bPathExists ? TEXT("exists, trying reversed search") : TEXT("doesn't exist"));
+			UE_VLOG(GetOwner(), LogAINavigation, Log, TEXT("GeneratePathTo: out of nodes, HPA* path: %s"),
+				bPathExists ? TEXT("exists, trying reversed search") : TEXT("doesn't exist"));
 
-				// reverse search
-				if (bPathExists)
-				{
-					DECLARE_SCOPE_CYCLE_COUNTER(TEXT("Pathfinding: reversed test time"), STAT_Navigation_PathReverseTime, STATGROUP_Navigation);
+			// reverse search
+			if (bPathExists)
+			{
+				DECLARE_SCOPE_CYCLE_COUNTER(TEXT("Pathfinding: reversed test time"), STAT_Navigation_PathReverseTime, STATGROUP_Navigation);
 
-					TSharedPtr<FNavigationQueryFilter> Filter = QueryFilter.IsValid() ? QueryFilter->GetCopy() : MyNavData->GetDefaultQueryFilter()->GetCopy();					
-					Filter->SetBacktrackingEnabled(true);
+				TSharedPtr<FNavigationQueryFilter> Filter = QueryFilter.IsValid() ? QueryFilter->GetCopy() : MyNavData->GetDefaultQueryFilter()->GetCopy();					
+				Filter->SetBacktrackingEnabled(true);
 
-					FPathFindingQuery ReversedQuery(GetOwner(), MyNavData, Query.EndLocation, Query.StartLocation, Filter);
-					ReversedQuery.NavDataFlags = NavDataFlags;
+				FPathFindingQuery ReversedQuery(GetOwner(), MyNavData, Query.EndLocation, Query.StartLocation, Filter);
+				ReversedQuery.NavDataFlags = NavDataFlags;
 
-					Result = NavSys->FindPathSync(*MyNavAgent->GetNavAgentProperties(), Query, Mode);
-				}
+				Result = NavSys->FindPathSync(MyNavAgent->GetNavAgentProperties(), Query, Mode);
 			}
+		}
 
-			if (Result.IsSuccessful() || Result.IsPartial())
-			{
-				CurrentGoal = NavMeshGoalLocation;
-				SetPath(Result.Path);
-
-				if (IsFollowing() == true)
-				{
-					// make sure ticking is enabled (and it shouldn't be enabled before _first_ async path finding)
-					//SetComponentTickEnabledAsync(true);
-					Path->SetGoalActorObservation(*GoalActor, FMath::Sqrt(RepathDistanceSq));
-				}
+		if (Result.IsSuccessful() || Result.IsPartial())
+		{
+			CurrentGoal = NavMeshGoalLocation;
+			SetPath(Result.Path);
 
-				// NotifyPathUpdate();
-				return true;
-			}
-			else
+			if (IsFollowing() == true)
 			{
-				UE_VLOG(GetOwner(), LogAINavigation, Display, TEXT("Failed to generate path to destination"));
+				// make sure ticking is enabled (and it shouldn't be enabled before _first_ async path finding)
+				//SetComponentTickEnabledAsync(true);
+				Path->SetGoalActorObservation(*GoalActor, FMath::Sqrt(RepathDistanceSq));
 			}
+
+			// NotifyPathUpdate();
+			return true;
 		}
 		else
 		{
-			UE_VLOG(GetOwner(), LogAINavigation, Display, TEXT("Destination not on navmesh (and nowhere near!)"));
+			UE_VLOG(GetOwner(), LogAINavigation, Display, TEXT("Failed to generate path to destination"));
 		}
 	}
 	else
 	{
-		UE_VLOG(GetOwner(), LogAINavigation, Display, TEXT("UNavigationComponent::GeneratePathTo: NavAgentProps == NULL (Probably Pawn died)"));
+		UE_VLOG(GetOwner(), LogAINavigation, Display, TEXT("Destination not on navmesh (and nowhere near!)"));
 	}
 
 	return false;
@@ -455,7 +440,7 @@ bool UNavigationComponent::GeneratePath(const FVector& FromLocation, const FVect
 	FPathFindingQuery Query(GetOwner(), GetNavData(), FromLocation, NavMeshGoalLocation.Location, QueryFilter);
 	Query.NavDataFlags = NavDataFlags;
 
-	FPathFindingResult Result = NavSys->FindPathSync(*MyNavAgent->GetNavAgentProperties(), Query, Mode);
+	FPathFindingResult Result = NavSys->FindPathSync(MyNavAgent->GetNavAgentProperties(), Query, Mode);
 	if (Result.IsSuccessful() || Result.IsPartial())
 	{
 		CurrentGoal = NavMeshGoalLocation.Location;
@@ -834,16 +819,13 @@ ANavigationData* UNavigationComponent::PickNavData() const
 
 	if (GetWorld() != NULL && GetWorld()->GetNavigationSystem() != NULL && MyNavAgent != NULL)
 	{
-		const FNavAgentProperties* MoveProps = MyNavAgent->GetNavAgentProperties();
+		const FNavAgentProperties& MoveProps = MyNavAgent->GetNavAgentProperties();
 
 		// Make sure any previous registration for the event is cleared.
 		GetWorld()->GetNavigationSystem()->OnNavDataRegisteredEvent.RemoveDynamic(this, &UNavigationComponent::OnNavDataRegistered);
 		
-		if (MoveProps != NULL)
-		{
-			// cache it
-			MyNavData = GetWorld()->GetNavigationSystem()->GetNavDataForProps(*MoveProps);
-		}
+		// cache it
+		MyNavData = GetWorld()->GetNavigationSystem()->GetNavDataForProps(MoveProps);
 
 		if (MyNavData == NULL)
 		{	// Failed to pick it, so register to try again.
@@ -862,11 +844,10 @@ void UNavigationComponent::CacheNavQueryExtent() const
 	NavigationQueryExtent = MyNavData ? MyNavData->GetDefaultQueryExtent() : FVector(DEFAULT_NAV_QUERY_EXTENT_HORIZONTAL, DEFAULT_NAV_QUERY_EXTENT_HORIZONTAL, DEFAULT_NAV_QUERY_EXTENT_VERTICAL);
 	if (MyNavAgent)
 	{
-		const FNavAgentProperties* AgentProperties = MyNavAgent->GetNavAgentProperties();
-		check(AgentProperties);
-		NavigationQueryExtent = FVector(FMath::Max(NavigationQueryExtent.X, AgentProperties->AgentRadius)
-			, FMath::Max(NavigationQueryExtent.Y, AgentProperties->AgentRadius)
-			, FMath::Max(NavigationQueryExtent.Z, AgentProperties->AgentHeight / 2));
+		const FNavAgentProperties& AgentProperties = MyNavAgent->GetNavAgentProperties();
+		NavigationQueryExtent = FVector(FMath::Max(NavigationQueryExtent.X, AgentProperties.AgentRadius)
+			, FMath::Max(NavigationQueryExtent.Y, AgentProperties.AgentRadius)
+			, FMath::Max(NavigationQueryExtent.Z, AgentProperties.AgentHeight / 2));
 	}
 }
 
@@ -930,7 +911,7 @@ void UNavigationComponent::OnCustomLinkBroadcast(UNavLinkCustomComponent* Nearby
 	FPathFindingQuery Query(GetOwner(), GetNavData(), MyNavAgent->GetNavAgentLocation(), GoalLocation, StoredQueryFilter);
 	Query.NavDataFlags = NavDataFlags;
 
-	FPathFindingResult Result = NavSys->FindPathSync(*MyNavAgent->GetNavAgentProperties(), Query, Mode);
+	FPathFindingResult Result = NavSys->FindPathSync(MyNavAgent->GetNavAgentProperties(), Query, Mode);
 	if (Result.IsSuccessful() || Result.IsPartial())
 	{
 		const bool bNewHasLink = Result.Path->ContainsCustomLink(NearbyLink->GetLinkId());
diff --git a/Engine/Source/Runtime/AIModule/Private/Navigation/PathFollowingComponent.cpp b/Engine/Source/Runtime/AIModule/Private/Navigation/PathFollowingComponent.cpp
--- a/Engine/Source/Runtime/AIModule/Private/Navigation/PathFollowingComponent.cpp
+++ b/Engine/Source/Runtime/AIModule/Private/Navigation/PathFollowingComponent.cpp
@@ -447,14 +447,13 @@ void UPathFollowingComponent::SetMovementComponent(UNavMovementComponent* MoveCo
 
 	if (MoveComp != NULL)
 	{
-		const FNavAgentProperties* NavAgentProps = MoveComp->GetNavAgentProperties();
-		check(NavAgentProps);
-		MyDefaultAcceptanceRadius = NavAgentProps->AgentRadius;
+		const FNavAgentProperties& NavAgentProps = MoveComp->GetNavAgentProperties();
+		MyDefaultAcceptanceRadius = NavAgentProps.AgentRadius;
 		MoveComp->PathFollowingComp = this;
 
 		if (GetWorld() && GetWorld()->GetNavigationSystem())
 		{	
-			MyNavData = GetWorld()->GetNavigationSystem()->GetNavDataForProps(*NavAgentProps);
+			MyNavData = GetWorld()->GetNavigationSystem()->GetNavDataForProps(NavAgentProps);
 		}
 	}
 }
diff --git a/Engine/Source/Runtime/Engine/Classes/AI/Navigation/NavAgentInterface.h b/Engine/Source/Runtime/Engine/Classes/AI/Navigation/NavAgentInterface.h
--- a/Engine/Source/Runtime/Engine/Classes/AI/Navigation/NavAgentInterface.h
+++ b/Engine/Source/Runtime/Engine/Classes/AI/Navigation/NavAgentInterface.h
@@ -1,6 +1,8 @@
 // Copyright 1998-2014 Epic Games, Inc. All Rights Reserved.
 
 #pragma once
+
+#include "AI/Navigation/NavigationTypes.h"
 #include "NavAgentInterface.generated.h"
 
 class AActor;
@@ -19,7 +21,7 @@ class INavAgentInterface
 	/**
 	 *	Retrieves FNavAgentProperties expressing navigation props and caps of represented agent
 	 */
-	virtual const FNavAgentProperties* GetNavAgentProperties() const { return NULL;} //PURE_VIRTUAL(INavAgentInterface::GetNavAgentProperties,return NULL;);
+	virtual const FNavAgentProperties& GetNavAgentProperties() const { return FNavAgentProperties::DefaultProperties; }
 
 	/**
 	 *	Retrieves Agent's location
diff --git a/Engine/Source/Runtime/Engine/Classes/AI/Navigation/NavigationTestingActor.h b/Engine/Source/Runtime/Engine/Classes/AI/Navigation/NavigationTestingActor.h
--- a/Engine/Source/Runtime/Engine/Classes/AI/Navigation/NavigationTestingActor.h
+++ b/Engine/Source/Runtime/Engine/Classes/AI/Navigation/NavigationTestingActor.h
@@ -151,7 +151,7 @@ public:
 #endif // WITH_EDITOR
 
 	// Begin INavAgentInterface Interface
-	virtual const struct FNavAgentProperties* GetNavAgentProperties() const override { return &NavAgentProps; }
+	virtual const FNavAgentProperties& GetNavAgentProperties() const override { return NavAgentProps; }
 	virtual FVector GetNavAgentLocation() const override;
 	virtual void GetMoveGoalReachTest(class AActor* MovingActor, const FVector& MoveOffset, FVector& GoalOffset, float& GoalRadius, float& GoalHalfHeight) const override {}
 	// End INavAgentInterface Interface
diff --git a/Engine/Source/Runtime/Engine/Classes/AI/Navigation/NavigationTypes.h b/Engine/Source/Runtime/Engine/Classes/AI/Navigation/NavigationTypes.h
--- a/Engine/Source/Runtime/Engine/Classes/AI/Navigation/NavigationTypes.h
+++ b/Engine/Source/Runtime/Engine/Classes/AI/Navigation/NavigationTypes.h
@@ -326,7 +326,7 @@ struct FMovementProperties
 
 /** Properties of representation of an 'agent' (or Pawn) used by AI navigation/pathfinding. */
 USTRUCT()
-struct FNavAgentProperties : public FMovementProperties
+struct ENGINE_API FNavAgentProperties : public FMovementProperties
 {
 	GENERATED_USTRUCT_BODY()
 
@@ -349,7 +349,7 @@ struct FNavAgentProperties : public FMovementProperties
 
 	void UpdateWithCollisionComponent(class UShapeComponent* CollisionComponent);
 
-	FORCEINLINE bool IsValid() const { return AgentRadius >= 0 && AgentHeight >= 0; }
+	FORCEINLINE bool IsValid() const { return AgentRadius >= 0.f && AgentHeight >= 0.f; }
 	FORCEINLINE bool HasStepHeightOverride() const { return AgentStepHeight >= 0.0f; }
 
 	FORCEINLINE bool IsEquivalent(const FNavAgentProperties& Other, float Precision = 5.f) const
@@ -368,6 +368,8 @@ struct FNavAgentProperties : public FMovementProperties
 	{
 		return FVector(AgentRadius, AgentRadius, AgentHeight / 2);
 	}
+
+	static const FNavAgentProperties DefaultProperties;
 };
 
 inline uint32 GetTypeHash(const FNavAgentProperties& A)
diff --git a/Engine/Source/Runtime/Engine/Classes/GameFramework/Controller.h b/Engine/Source/Runtime/Engine/Classes/GameFramework/Controller.h
--- a/Engine/Source/Runtime/Engine/Classes/GameFramework/Controller.h
+++ b/Engine/Source/Runtime/Engine/Classes/GameFramework/Controller.h
@@ -264,7 +264,7 @@ public:
 	virtual void FailedToSpawnPawn();
 
 	// Begin INavAgentInterface Interface
-	virtual const struct FNavAgentProperties* GetNavAgentProperties() const override;
+	virtual const struct FNavAgentProperties& GetNavAgentProperties() const override;
 	virtual FVector GetNavAgentLocation() const override;
 	virtual void GetMoveGoalReachTest(class AActor* MovingActor, const FVector& MoveOffset, FVector& GoalOffset, float& GoalRadius, float& GoalHalfHeight) const override;
 	// End INavAgentInterface Interface
diff --git a/Engine/Source/Runtime/Engine/Classes/GameFramework/NavMovementComponent.h b/Engine/Source/Runtime/Engine/Classes/GameFramework/NavMovementComponent.h
--- a/Engine/Source/Runtime/Engine/Classes/GameFramework/NavMovementComponent.h
+++ b/Engine/Source/Runtime/Engine/Classes/GameFramework/NavMovementComponent.h
@@ -76,9 +76,9 @@ public:
 	virtual bool CanStopPathFollowing() const;
 
 	/** @returns the NavAgentProps(const) */
-	FORCEINLINE const FNavAgentProperties* GetNavAgentProperties() const { return &NavAgentProps; }
+	FORCEINLINE const FNavAgentProperties& GetNavAgentProperties() const { return NavAgentProps; }
 	/** @returns the NavAgentProps */
-	FORCEINLINE FNavAgentProperties* GetNavAgentProperties() { return &NavAgentProps; }
+	FORCEINLINE FNavAgentProperties& GetNavAgentProperties() { return NavAgentProps; }
 
 	/** Resets runtime movement state to character's movement capabilities */
 	void ResetMoveState() { MovementState = NavAgentProps; }
diff --git a/Engine/Source/Runtime/Engine/Classes/GameFramework/Pawn.h b/Engine/Source/Runtime/Engine/Classes/GameFramework/Pawn.h
--- a/Engine/Source/Runtime/Engine/Classes/GameFramework/Pawn.h
+++ b/Engine/Source/Runtime/Engine/Classes/GameFramework/Pawn.h
@@ -19,6 +19,7 @@ class UPawnMovementComponent;
 class UPawnNoiseEmitterComponent;
 class UPlayer;
 class UPrimitiveComponent;
+struct FNavAgentProperties;
 
 /** 
  * Pawn is the base class of all actors that can be possessed by players or AI.
@@ -229,7 +230,7 @@ public:
 	virtual void UpdateNavigationRelevance() {}
 
 	// Begin INavAgentInterface Interface
-	virtual const struct FNavAgentProperties* GetNavAgentProperties() const override;
+	virtual const FNavAgentProperties& GetNavAgentProperties() const override;
 	/** Basically retrieved pawn's location on navmesh */
 	UFUNCTION(BlueprintCallable, Category="Pawn")
 	virtual FVector GetNavAgentLocation() const override { return GetActorLocation() - FVector(0.f, 0.f, BaseEyeHeight); }
diff --git a/Engine/Source/Runtime/Engine/Private/AI/Navigation/NavigationSystem.cpp b/Engine/Source/Runtime/Engine/Private/AI/Navigation/NavigationSystem.cpp
--- a/Engine/Source/Runtime/Engine/Private/AI/Navigation/NavigationSystem.cpp
+++ b/Engine/Source/Runtime/Engine/Private/AI/Navigation/NavigationSystem.cpp
@@ -946,12 +946,9 @@ UNavigationPath* UNavigationSystem::FindPathToLocationSynchronously(UObject* Wor
 			
 			if (NavAgent != NULL)
 			{
-				const FNavAgentProperties* AgentProps = NavAgent->GetNavAgentProperties();
-				if (AgentProps)
-				{
-					NavigationData = NavSys->GetNavDataForProps(*AgentProps);
-					bValidPathContext = true;
-				}
+				const FNavAgentProperties& AgentProps = NavAgent->GetNavAgentProperties();
+				NavigationData = NavSys->GetNavDataForProps(AgentProps);
+				bValidPathContext = true;
 			}
 			else if (Cast<ANavigationData>(PathfindingContext))
 			{
@@ -1005,8 +1002,8 @@ bool UNavigationSystem::NavigationRaycast(UObject* WorldContextObject, const FVe
 		INavAgentInterface* MyNavAgent = Cast<INavAgentInterface>(Querier);
 		if (MyNavAgent)
 		{
-			const FNavAgentProperties* AgentProps = MyNavAgent->GetNavAgentProperties();
-			NavData = NavSys->GetNavDataForProps(*AgentProps);
+			const FNavAgentProperties& AgentProps = MyNavAgent->GetNavAgentProperties();
+			NavData = NavSys->GetNavDataForProps(AgentProps);
 		}
 		if (NavData == NULL)
 		{
diff --git a/Engine/Source/Runtime/Engine/Private/Controller.cpp b/Engine/Source/Runtime/Engine/Private/Controller.cpp
--- a/Engine/Source/Runtime/Engine/Private/Controller.cpp
+++ b/Engine/Source/Runtime/Engine/Private/Controller.cpp
@@ -532,9 +532,9 @@ APawn* AController::K2_GetPawn() const
 	return GetPawn();
 }
 
-const FNavAgentProperties* AController::GetNavAgentProperties() const
+const FNavAgentProperties& AController::GetNavAgentProperties() const
 {
-	return Pawn && Pawn->GetMovementComponent() ? Pawn->GetMovementComponent()->GetNavAgentProperties() : NULL;
+	return Pawn && Pawn->GetMovementComponent() ? Pawn->GetMovementComponent()->GetNavAgentProperties() : FNavAgentProperties::DefaultProperties;
 }
 
 FVector AController::GetNavAgentLocation() const
diff --git a/Engine/Source/Runtime/Engine/Private/NavMovementComponent.cpp b/Engine/Source/Runtime/Engine/Private/NavMovementComponent.cpp
--- a/Engine/Source/Runtime/Engine/Private/NavMovementComponent.cpp
+++ b/Engine/Source/Runtime/Engine/Private/NavMovementComponent.cpp
@@ -8,6 +8,8 @@
 //----------------------------------------------------------------------//
 // FNavAgentProperties
 //----------------------------------------------------------------------//
+const FNavAgentProperties FNavAgentProperties::DefaultProperties;
+
 void FNavAgentProperties::UpdateWithCollisionComponent(UShapeComponent* CollisionComponent)
 {
 	check( CollisionComponent != NULL);
diff --git a/Engine/Source/Runtime/Engine/Private/Pawn.cpp b/Engine/Source/Runtime/Engine/Private/Pawn.cpp
--- a/Engine/Source/Runtime/Engine/Private/Pawn.cpp
+++ b/Engine/Source/Runtime/Engine/Private/Pawn.cpp
@@ -1080,7 +1080,7 @@ void APawn::PawnMakeNoise(float Loudness, FVector NoiseLocation, bool bUseNoiseM
 	NoiseMaker->MakeNoise(Loudness, this, bUseNoiseMakerLocation ? NoiseMaker->GetActorLocation() : NoiseLocation);
 }
 
-const struct FNavAgentProperties* APawn::GetNavAgentProperties() const
+const FNavAgentProperties& APawn::GetNavAgentProperties() const
 {
-	return GetMovementComponent() ? GetMovementComponent()->GetNavAgentProperties() : NULL;
+	return GetMovementComponent() ? GetMovementComponent()->GetNavAgentProperties() : FNavAgentProperties::DefaultProperties;
 }
