Cleaned up redundant usage of 'class' and 'struct' keywords out of AI sources #UE4

diff --git a/Engine/Source/Runtime/AIModule/Classes/BehaviorTree/Decorators/BTDecorator_CheckGameplayTagsOnActor.h b/Engine/Source/Runtime/AIModule/Classes/BehaviorTree/Decorators/BTDecorator_CheckGameplayTagsOnActor.h
--- a/Engine/Source/Runtime/AIModule/Classes/BehaviorTree/Decorators/BTDecorator_CheckGameplayTagsOnActor.h
+++ b/Engine/Source/Runtime/AIModule/Classes/BehaviorTree/Decorators/BTDecorator_CheckGameplayTagsOnActor.h
@@ -1,6 +1,7 @@
 // Copyright 1998-2014 Epic Games, Inc. All Rights Reserved.
 
 #pragma once
+
 #include "BehaviorTree/BTDecorator.h"
 #include "GameplayTagContainer.h"
 #include "BTDecorator_CheckGameplayTagsOnActor.generated.h"
@@ -15,8 +16,8 @@ class AIMODULE_API UBTDecorator_CheckGameplayTagsOnActor : public UBTDecorator
 {
 	GENERATED_UCLASS_BODY()
 
-	virtual bool CalculateRawConditionValue(class UBehaviorTreeComponent* OwnerComp, uint8* NodeMemory) const override;
-	virtual void DescribeRuntimeValues(const class UBehaviorTreeComponent* OwnerComp, uint8* NodeMemory, EBTDescriptionVerbosity::Type Verbosity, TArray<FString>& Values) const override;
+	virtual bool CalculateRawConditionValue(UBehaviorTreeComponent* OwnerComp, uint8* NodeMemory) const override;
+	virtual void DescribeRuntimeValues(const UBehaviorTreeComponent* OwnerComp, uint8* NodeMemory, EBTDescriptionVerbosity::Type Verbosity, TArray<FString>& Values) const override;
 	virtual FString GetStaticDescription() const override;
 
 protected:
diff --git a/Engine/Source/Runtime/AIModule/Classes/Navigation/CrowdManager.h b/Engine/Source/Runtime/AIModule/Classes/Navigation/CrowdManager.h
--- a/Engine/Source/Runtime/AIModule/Classes/Navigation/CrowdManager.h
+++ b/Engine/Source/Runtime/AIModule/Classes/Navigation/CrowdManager.h
@@ -4,12 +4,19 @@
 #include "Tickable.h"
 #include "CrowdManager.generated.h"
 
+class UWorld;
 struct FNavMeshPath;
+class ANavigationData;
 class UCrowdFollowingComponent;
 class ICrowdAgentInterface;
+class UCrowdManager;
 #if WITH_RECAST
 struct dtQuerySpecialLinkFilter;
 struct dtCrowdAgentParams;
+class dtCrowd;
+struct dtCrowdAgent;
+struct dtCrowdAgentDebugInfo;
+class dtObstacleAvoidanceDebugData;
 #endif
 
 /**
@@ -133,7 +140,7 @@ struct AIMODULE_API FCrowdAgentData
 
 struct FCrowdTickHelper : FTickableGameObject
 {
-	TWeakObjectPtr<class UCrowdManager> Owner;
+	TWeakObjectPtr<UCrowdManager> Owner;
 
 	FCrowdTickHelper() : Owner(NULL) {}
 	virtual void Tick(float DeltaTime);
@@ -155,16 +162,16 @@ class AIMODULE_API UCrowdManager : public UObject
 	void RegisterAgent(ICrowdAgentInterface* Agent);
 
 	/** removes agent from crowd */
-	void UnregisterAgent(const class ICrowdAgentInterface* Agent);
+	void UnregisterAgent(const ICrowdAgentInterface* Agent);
 
 	/** updates agent data */
-	void UpdateAgentParams(const class ICrowdAgentInterface* Agent) const;
+	void UpdateAgentParams(const ICrowdAgentInterface* Agent) const;
 
 	/** refresh agent state */
-	void UpdateAgentState(const class ICrowdAgentInterface* Agent) const;
+	void UpdateAgentState(const ICrowdAgentInterface* Agent) const;
 
 	/** update agent after using custom link */
-	void OnAgentFinishedCustomLink(const class ICrowdAgentInterface* Agent) const;
+	void OnAgentFinishedCustomLink(const ICrowdAgentInterface* Agent) const;
 
 	/** sets move target for crowd agent (only for fully simulated) */
 	bool SetAgentMoveTarget(const UCrowdFollowingComponent* AgentComponent, const FVector& MoveTarget, TSharedPtr<const FNavigationQueryFilter> Filter) const;
@@ -185,11 +192,11 @@ class AIMODULE_API UCrowdManager : public UObject
 	void ResumeAgent(const UCrowdFollowingComponent* AgentComponent, bool bForceReplanPath = true) const;
 
 	/** check if object is a valid crowd agent */
-	bool IsAgentValid(const class UCrowdFollowingComponent* AgentComponent) const;
-	bool IsAgentValid(const class ICrowdAgentInterface* Agent) const;
+	bool IsAgentValid(const UCrowdFollowingComponent* AgentComponent) const;
+	bool IsAgentValid(const ICrowdAgentInterface* Agent) const;
 
 	/** returns number of nearby agents */
-	int32 GetNumNearbyAgents(const class ICrowdAgentInterface* Agent) const;
+	int32 GetNumNearbyAgents(const ICrowdAgentInterface* Agent) const;
 
 	/** returns a list of locations of nearby agents */
 	int32 GetNearbyAgentLocations(const ICrowdAgentInterface* Agent, TArray<FVector>& OutLocations) const;
@@ -218,15 +225,15 @@ class AIMODULE_API UCrowdManager : public UObject
 	/** notify called when detour navmesh is changed */
 	void OnNavMeshUpdate();
 
-	class UWorld* GetWorld() const;
+	UWorld* GetWorld() const;
 
-	static UCrowdManager* GetCurrent(class UObject* WorldContextObject);
-	static UCrowdManager* GetCurrent(class UWorld* World);
+	static UCrowdManager* GetCurrent(UObject* WorldContextObject);
+	static UCrowdManager* GetCurrent(UWorld* World);
 
 protected:
 
 	UPROPERTY(transient)
-	class ANavigationData* MyNavData;
+	ANavigationData* MyNavData;
 
 	/** obstacle avoidance params */
 	UPROPERTY(EditAnywhere, Category=Config)
@@ -268,11 +275,11 @@ protected:
 
 #if WITH_RECAST
 	/** crowd manager */
-	class dtCrowd* DetourCrowd;
+	dtCrowd* DetourCrowd;
 
 	/** debug data */
-	struct dtCrowdAgentDebugInfo* DetourAgentDebug;
-	class dtObstacleAvoidanceDebugData* DetourAvoidanceDebug;
+	dtCrowdAgentDebugInfo* DetourAgentDebug;
+	dtObstacleAvoidanceDebugData* DetourAvoidanceDebug;
 #endif
 
 #if WITH_EDITOR
@@ -289,12 +296,12 @@ protected:
 	virtual void PostProximityUpdate();
 
 #if WITH_RECAST
-	void AddAgent(const class ICrowdAgentInterface* Agent, FCrowdAgentData& AgentData) const;
-	void RemoveAgent(const class ICrowdAgentInterface* Agent, FCrowdAgentData* AgentData) const;
-	void GetAgentParams(const class ICrowdAgentInterface* Agent, dtCrowdAgentParams& AgentParams) const;
+	void AddAgent(const ICrowdAgentInterface* Agent, FCrowdAgentData& AgentData) const;
+	void RemoveAgent(const ICrowdAgentInterface* Agent, FCrowdAgentData* AgentData) const;
+	void GetAgentParams(const ICrowdAgentInterface* Agent, dtCrowdAgentParams& AgentParams) const;
 
 	/** prepare agent for next step of simulation */
-	void PrepareAgentStep(const class ICrowdAgentInterface* Agent, FCrowdAgentData& AgentData, float DeltaTime) const;
+	void PrepareAgentStep(const ICrowdAgentInterface* Agent, FCrowdAgentData& AgentData, float DeltaTime) const;
 
 	/** pass new velocity to movement components */
 	void ApplyVelocity(UCrowdFollowingComponent* AgentComponent, int32 AgentIndex) const;
@@ -303,16 +310,16 @@ protected:
 	void UpdateAgentPaths();
 
 	/** switch debugger to object selected in PIE */
-	void UpdateSelectedDebug(const class ICrowdAgentInterface* Agent, int32 AgentIndex) const;
+	void UpdateSelectedDebug(const ICrowdAgentInterface* Agent, int32 AgentIndex) const;
 
 	void CreateCrowdManager();
 	void DestroyCrowdManager();
 
-	void DrawDebugCorners(const struct dtCrowdAgent* CrowdAgent) const;
-	void DrawDebugCollisionSegments(const struct dtCrowdAgent* CrowdAgent) const;
-	void DrawDebugPath(const struct dtCrowdAgent* CrowdAgent) const;
-	void DrawDebugVelocityObstacles(const struct dtCrowdAgent* CrowdAgent) const;
-	void DrawDebugPathOptimization(const struct dtCrowdAgent* CrowdAgent) const;
-	void DrawDebugNeighbors(const struct dtCrowdAgent* CrowdAgent) const;
+	void DrawDebugCorners(const dtCrowdAgent* CrowdAgent) const;
+	void DrawDebugCollisionSegments(const dtCrowdAgent* CrowdAgent) const;
+	void DrawDebugPath(const dtCrowdAgent* CrowdAgent) const;
+	void DrawDebugVelocityObstacles(const dtCrowdAgent* CrowdAgent) const;
+	void DrawDebugPathOptimization(const dtCrowdAgent* CrowdAgent) const;
+	void DrawDebugNeighbors(const dtCrowdAgent* CrowdAgent) const;
 #endif
 };
diff --git a/Engine/Source/Runtime/AIModule/Private/BehaviorTree/Decorators/BTDecorator_CheckGameplayTagsOnActor.cpp b/Engine/Source/Runtime/AIModule/Private/BehaviorTree/Decorators/BTDecorator_CheckGameplayTagsOnActor.cpp
--- a/Engine/Source/Runtime/AIModule/Private/BehaviorTree/Decorators/BTDecorator_CheckGameplayTagsOnActor.cpp
+++ b/Engine/Source/Runtime/AIModule/Private/BehaviorTree/Decorators/BTDecorator_CheckGameplayTagsOnActor.cpp
@@ -21,7 +21,7 @@ UBTDecorator_CheckGameplayTagsOnActor::UBTDecorator_CheckGameplayTagsOnActor(con
 	bAllowAbortChildNodes = false;
 }
 
-bool UBTDecorator_CheckGameplayTagsOnActor::CalculateRawConditionValue(class UBehaviorTreeComponent* OwnerComp, uint8* NodeMemory) const
+bool UBTDecorator_CheckGameplayTagsOnActor::CalculateRawConditionValue(UBehaviorTreeComponent* OwnerComp, uint8* NodeMemory) const
 {
 	check(OwnerComp != NULL);
 	const UBlackboardComponent* BlackboardComp = OwnerComp->GetBlackboardComponent();
@@ -52,7 +52,7 @@ bool UBTDecorator_CheckGameplayTagsOnActor::CalculateRawConditionValue(class UBe
 	}
 }
 
-void UBTDecorator_CheckGameplayTagsOnActor::DescribeRuntimeValues(const class UBehaviorTreeComponent* OwnerComp, uint8* NodeMemory, EBTDescriptionVerbosity::Type Verbosity, TArray<FString>& Values) const
+void UBTDecorator_CheckGameplayTagsOnActor::DescribeRuntimeValues(const UBehaviorTreeComponent* OwnerComp, uint8* NodeMemory, EBTDescriptionVerbosity::Type Verbosity, TArray<FString>& Values) const
 {
 	Super::DescribeRuntimeValues(OwnerComp, NodeMemory, Verbosity, Values);
 }
diff --git a/Engine/Source/Runtime/AIModule/Private/Navigation/CrowdManager.cpp b/Engine/Source/Runtime/AIModule/Private/Navigation/CrowdManager.cpp
--- a/Engine/Source/Runtime/AIModule/Private/Navigation/CrowdManager.cpp
+++ b/Engine/Source/Runtime/AIModule/Private/Navigation/CrowdManager.cpp
@@ -313,13 +313,13 @@ void UCrowdManager::UnregisterAgent(const ICrowdAgentInterface* Agent)
 	ActiveAgents.Remove(Agent);
 }
 
-bool UCrowdManager::IsAgentValid(const class UCrowdFollowingComponent* AgentComponent) const
+bool UCrowdManager::IsAgentValid(const UCrowdFollowingComponent* AgentComponent) const
 {
 	const FCrowdAgentData* AgentData = ActiveAgents.Find(AgentComponent);
 	return AgentData && AgentData->IsValid();
 }
 
-bool UCrowdManager::IsAgentValid(const class ICrowdAgentInterface* Agent) const
+bool UCrowdManager::IsAgentValid(const ICrowdAgentInterface* Agent) const
 {
 	const FCrowdAgentData* AgentData = ActiveAgents.Find(Agent);
 	return AgentData && AgentData->IsValid();
diff --git a/Engine/Source/Runtime/Engine/Classes/AI/Navigation/AvoidanceManager.h b/Engine/Source/Runtime/Engine/Classes/AI/Navigation/AvoidanceManager.h
--- a/Engine/Source/Runtime/Engine/Classes/AI/Navigation/AvoidanceManager.h
+++ b/Engine/Source/Runtime/Engine/Classes/AI/Navigation/AvoidanceManager.h
@@ -51,7 +51,7 @@ struct FNavAvoidanceData
 	FNavAvoidanceData(UAvoidanceManager* Manager, IRVOAvoidanceInterface* AvoidanceComp);
 
 	/** Init function for internal use to guard against data changes not being reflected in blueprint-accessible creation functions */
-	void Init(class UAvoidanceManager* Avoidance, const FVector& InCenter, float InRadius, float InHeight,
+	void Init(UAvoidanceManager* Avoidance, const FVector& InCenter, float InRadius, float InHeight,
 		const FVector& InVelocity, float InWeight = 0.5f,
 		int32 InGroupMask = 1, int32 InGroupsToAvoid = 0xffffffff, int32 InGroupsToIgnore = 0,
 		float InTestRadius2D = 500.0f);
diff --git a/Engine/Source/Runtime/Engine/Classes/AI/Navigation/NavAreas/NavArea.h b/Engine/Source/Runtime/Engine/Classes/AI/Navigation/NavAreas/NavArea.h
--- a/Engine/Source/Runtime/Engine/Classes/AI/Navigation/NavAreas/NavArea.h
+++ b/Engine/Source/Runtime/Engine/Classes/AI/Navigation/NavAreas/NavArea.h
@@ -74,7 +74,7 @@ class ENGINE_API UNavArea : public UObject
 	FORCEINLINE bool IsSupportingAgent(int32 AgentIndex) const { return (AgentIndex >= 0 && AgentIndex < 16) ? !!(SupportedAgentsBits & (1 << AgentIndex)) : false; }
 
 	/** Retrieved color declared for AreaDefinitionClass */
-	static FColor GetColor(class UClass* AreaDefinitionClass);
+	static FColor GetColor(UClass* AreaDefinitionClass);
 
 #if WITH_EDITOR
 	/** setup bSupportsAgentX properties */
diff --git a/Engine/Source/Runtime/Engine/Classes/AI/Navigation/NavAreas/NavAreaMeta.h b/Engine/Source/Runtime/Engine/Classes/AI/Navigation/NavAreas/NavAreaMeta.h
--- a/Engine/Source/Runtime/Engine/Classes/AI/Navigation/NavAreas/NavAreaMeta.h
+++ b/Engine/Source/Runtime/Engine/Classes/AI/Navigation/NavAreas/NavAreaMeta.h
@@ -1,9 +1,12 @@
 // Copyright 1998-2014 Epic Games, Inc. All Rights Reserved.
 
 #pragma once
+
 #include "AI/Navigation/NavAreas/NavArea.h"
 #include "NavAreaMeta.generated.h"
 
+class AActor;
+
 /** Class containing definition of a navigation area */
 UCLASS(Abstract)
 class ENGINE_API UNavAreaMeta : public UNavArea
@@ -14,9 +17,9 @@ class ENGINE_API UNavAreaMeta : public UNavArea
 	 *	Picks an navigation area class that should be used for Actor when
 	 *	queried by NavAgent
 	 */
-	virtual TSubclassOf<UNavArea> PickAreaClass(const class AActor* Actor, const FNavAgentProperties& NavAgent) const;
+	virtual TSubclassOf<UNavArea> PickAreaClass(const AActor* Actor, const FNavAgentProperties& NavAgent) const;
 
-	FORCEINLINE static TSubclassOf<UNavArea> PickAreaClass(TSubclassOf<UNavArea> AreaClass, const class AActor* Actor, const FNavAgentProperties& NavAgent)
+	FORCEINLINE static TSubclassOf<UNavArea> PickAreaClass(TSubclassOf<UNavArea> AreaClass, const AActor* Actor, const FNavAgentProperties& NavAgent)
 	{
 		if (AreaClass->IsChildOf(UNavAreaMeta::StaticClass()))
 		{
diff --git a/Engine/Source/Runtime/Engine/Classes/AI/Navigation/NavAreas/NavAreaMeta_SwitchByAgent.h b/Engine/Source/Runtime/Engine/Classes/AI/Navigation/NavAreas/NavAreaMeta_SwitchByAgent.h
--- a/Engine/Source/Runtime/Engine/Classes/AI/Navigation/NavAreas/NavAreaMeta_SwitchByAgent.h
+++ b/Engine/Source/Runtime/Engine/Classes/AI/Navigation/NavAreas/NavAreaMeta_SwitchByAgent.h
@@ -1,6 +1,7 @@
 // Copyright 1998-2014 Epic Games, Inc. All Rights Reserved.
 
 #pragma once
+
 #include "AI/Navigation/NavAreas/NavAreaMeta.h"
 #include "NavAreaMeta_SwitchByAgent.generated.h"
 
@@ -58,7 +59,7 @@ class ENGINE_API UNavAreaMeta_SwitchByAgent : public UNavAreaMeta
 	UPROPERTY(EditAnywhere, Category=AgentTypes)
 	TSubclassOf<UNavArea> Agent15Area;
 
-	virtual TSubclassOf<UNavArea> PickAreaClass(const class AActor* Actor, const FNavAgentProperties& NavAgent) const override;
+	virtual TSubclassOf<UNavArea> PickAreaClass(const AActor* Actor, const FNavAgentProperties& NavAgent) const override;
 
 #if WITH_EDITOR
 	/** setup AgentXArea properties */
diff --git a/Engine/Source/Runtime/Engine/Classes/AI/Navigation/NavCollision.h b/Engine/Source/Runtime/Engine/Classes/AI/Navigation/NavCollision.h
--- a/Engine/Source/Runtime/Engine/Classes/AI/Navigation/NavCollision.h
+++ b/Engine/Source/Runtime/Engine/Classes/AI/Navigation/NavCollision.h
@@ -5,6 +5,9 @@
 #include "NavigationTypes.h"
 #include "NavCollision.generated.h"
 
+class FPrimitiveDrawInterface;
+struct FCompositeNavModifier;
+
 USTRUCT()
 struct FNavCollisionCylinder
 {
@@ -89,10 +92,10 @@ class ENGINE_API UNavCollision : public UObject
 	void Setup(class UBodySetup* BodySetup);
 
 	/** show cylinder and box collisions */
-	void DrawSimpleGeom(class FPrimitiveDrawInterface* PDI, const FTransform& Transform, const FColor Color);
+	void DrawSimpleGeom(FPrimitiveDrawInterface* PDI, const FTransform& Transform, const FColor Color);
 
 	/** Get data for dynamic obstacle */
-	void GetNavigationModifier(struct FCompositeNavModifier& Modifier, const FTransform& LocalToWorld) const;
+	void GetNavigationModifier(FCompositeNavModifier& Modifier, const FTransform& LocalToWorld) const;
 
 	/** Read collisions data */
 	bool GatherCollision();
diff --git a/Engine/Source/Runtime/Engine/Classes/AI/Navigation/NavFilters/NavigationQueryFilter.h b/Engine/Source/Runtime/Engine/Classes/AI/Navigation/NavFilters/NavigationQueryFilter.h
--- a/Engine/Source/Runtime/Engine/Classes/AI/Navigation/NavFilters/NavigationQueryFilter.h
+++ b/Engine/Source/Runtime/Engine/Classes/AI/Navigation/NavFilters/NavigationQueryFilter.h
@@ -3,6 +3,9 @@
 #pragma once
 #include "NavigationQueryFilter.generated.h"
 
+class UNavArea;
+struct FNavigationQueryFilter;
+
 USTRUCT()
 struct ENGINE_API FNavigationFilterArea
 {
@@ -10,7 +13,7 @@ struct ENGINE_API FNavigationFilterArea
 
 	/** navigation area class */
 	UPROPERTY(EditAnywhere, Category=Area)
-	TSubclassOf<class UNavArea> AreaClass;
+	TSubclassOf<UNavArea> AreaClass;
 
 	/** override for travel cost */
 	UPROPERTY(EditAnywhere, Category=Area, meta=(EditCondition="bOverrideTravelCost",ClampMin=1))
@@ -212,13 +215,13 @@ class ENGINE_API UNavigationQueryFilter : public UObject
 	FNavigationFilterFlags ExcludeFlags;
 
 	/** get filter for given navigation data and initialize on first access */
-	TSharedPtr<const struct FNavigationQueryFilter> GetQueryFilter(const class ANavigationData* NavData) const;
+	TSharedPtr<const FNavigationQueryFilter> GetQueryFilter(const class ANavigationData* NavData) const;
 	
 	/** helper functions for accessing filter */
-	static TSharedPtr<const struct FNavigationQueryFilter> GetQueryFilter(const class ANavigationData* NavData, UClass* FilterClass);
+	static TSharedPtr<const FNavigationQueryFilter> GetQueryFilter(const class ANavigationData* NavData, UClass* FilterClass);
 
 	template<class T>
-	static TSharedPtr<const struct FNavigationQueryFilter> GetQueryFilter(const class ANavigationData* NavData, UClass* FilterClass = T::StaticClass())
+	static TSharedPtr<const FNavigationQueryFilter> GetQueryFilter(const class ANavigationData* NavData, UClass* FilterClass = T::StaticClass())
 	{
 		return GetQueryFilter(NavData, FilterClass);
 	}
@@ -230,13 +233,13 @@ class ENGINE_API UNavigationQueryFilter : public UObject
 protected:
 
 	/** helper functions for adding area overrides */
-	void AddTravelCostOverride(TSubclassOf<class UNavArea> AreaClass, float TravelCost);
-	void AddEnteringCostOverride(TSubclassOf<class UNavArea> AreaClass, float EnteringCost);
-	void AddExcludedArea(TSubclassOf<class UNavArea> AreaClass);
+	void AddTravelCostOverride(TSubclassOf<UNavArea> AreaClass, float TravelCost);
+	void AddEnteringCostOverride(TSubclassOf<UNavArea> AreaClass, float EnteringCost);
+	void AddExcludedArea(TSubclassOf<UNavArea> AreaClass);
 
 	/** find index of area data */
-	int32 FindAreaOverride(TSubclassOf<class UNavArea> AreaClass) const;
+	int32 FindAreaOverride(TSubclassOf<UNavArea> AreaClass) const;
 
 	/** setup filter for given navigation data, use to create custom filters */
-	virtual void InitializeFilter(const class ANavigationData* NavData, struct FNavigationQueryFilter* Filter) const;
+	virtual void InitializeFilter(const class ANavigationData* NavData, FNavigationQueryFilter* Filter) const;
 };
diff --git a/Engine/Source/Runtime/Engine/Classes/AI/Navigation/NavLinkCustomComponent.h b/Engine/Source/Runtime/Engine/Classes/AI/Navigation/NavLinkCustomComponent.h
--- a/Engine/Source/Runtime/Engine/Classes/AI/Navigation/NavLinkCustomComponent.h
+++ b/Engine/Source/Runtime/Engine/Classes/AI/Navigation/NavLinkCustomComponent.h
@@ -7,6 +7,11 @@
 #include "Engine/EngineTypes.h"
 #include "NavLinkCustomComponent.generated.h"
 
+class UNavArea;
+class UPathFollowingComponent;
+class UNavigationComponent;
+struct FNavigationRelevantData;
+
 /**
  *  Encapsulates NavLinkCustomInterface interface, can be used with Actors not relevant for navigation
  *  
@@ -21,20 +26,20 @@ class ENGINE_API UNavLinkCustomComponent : public UNavRelevantComponent, public
 {
 	GENERATED_UCLASS_BODY()
 
-	DECLARE_DELEGATE_ThreeParams(FOnMoveReachedLink, UNavLinkCustomComponent* /*ThisComp*/, class UPathFollowingComponent* /*PathComp*/, const FVector& /*DestPoint*/);
-	DECLARE_DELEGATE_TwoParams(FBroadcastFilter, UNavLinkCustomComponent* /*ThisComp*/, TArray<class UNavigationComponent*>& /*NotifyList*/);
+	DECLARE_DELEGATE_ThreeParams(FOnMoveReachedLink, UNavLinkCustomComponent* /*ThisComp*/, UPathFollowingComponent* /*PathComp*/, const FVector& /*DestPoint*/);
+	DECLARE_DELEGATE_TwoParams(FBroadcastFilter, UNavLinkCustomComponent* /*ThisComp*/, TArray<UNavigationComponent*>& /*NotifyList*/);
 
 	// BEGIN INavLinkCustomInterface
 	virtual void GetLinkData(FVector& LeftPt, FVector& RightPt, ENavLinkDirection::Type& Direction) const override;
 	virtual TSubclassOf<UNavArea> GetLinkAreaClass() const override;
 	virtual uint32 GetLinkId() const override;
 	virtual bool IsLinkPathfindingAllowed(const UObject* Querier) const override;
-	virtual bool OnLinkMoveStarted(class UPathFollowingComponent* PathComp, const FVector& DestPoint) override;
-	virtual void OnLinkMoveFinished(class UPathFollowingComponent* PathComp) override;
+	virtual bool OnLinkMoveStarted(UPathFollowingComponent* PathComp, const FVector& DestPoint) override;
+	virtual void OnLinkMoveFinished(UPathFollowingComponent* PathComp) override;
 	// END INavLinkCustomInterface
 
 	// BEGIN UNavRelevantComponent Interface
-	virtual void GetNavigationData(struct FNavigationRelevantData& Data) const;
+	virtual void GetNavigationData(FNavigationRelevantData& Data) const;
 	virtual void CalcBounds();
 	// END UNavRelevantComponent Interface
 
@@ -55,16 +60,16 @@ class ENGINE_API UNavLinkCustomComponent : public UNavRelevantComponent, public
 	virtual FNavigationLink GetLinkModifier() const;
 
 	/** set area class to use when link is enabled */
-	void SetEnabledArea(TSubclassOf<class UNavArea> AreaClass);
+	void SetEnabledArea(TSubclassOf<UNavArea> AreaClass);
 	TSubclassOf<UNavArea> GetEnabledArea() const { return EnabledAreaClass; }
 
 	/** set area class to use when link is disabled */
-	void SetDisabledArea(TSubclassOf<class UNavArea> AreaClass);
+	void SetDisabledArea(TSubclassOf<UNavArea> AreaClass);
 	TSubclassOf<UNavArea> GetDisabledArea() const { return DisabledAreaClass; }
 
 	/** add box obstacle during generation of navigation data
 	  * this can be used to create empty area under doors */
-	void AddNavigationObstacle(TSubclassOf<class UNavArea> AreaClass, const FVector& BoxExtent, const FVector& BoxOffset = FVector::ZeroVector);
+	void AddNavigationObstacle(TSubclassOf<UNavArea> AreaClass, const FVector& BoxExtent, const FVector& BoxOffset = FVector::ZeroVector);
 
 	/** removes simple obstacle */
 	void ClearNavigationObstacle();
@@ -127,11 +132,11 @@ protected:
 
 	/** area class to use when link is enabled */
 	UPROPERTY(EditAnywhere, Category=SmartLink)
-	TSubclassOf<class UNavArea> EnabledAreaClass;
+	TSubclassOf<UNavArea> EnabledAreaClass;
 
 	/** area class to use when link is disabled */
 	UPROPERTY(EditAnywhere, Category=SmartLink)
-	TSubclassOf<class UNavArea> DisabledAreaClass;
+	TSubclassOf<UNavArea> DisabledAreaClass;
 
 	/** start point, relative to owner */
 	UPROPERTY(EditAnywhere, Category=SmartLink)
@@ -171,7 +176,7 @@ protected:
 
 	/** area class for simple box obstacle */
 	UPROPERTY(EditAnywhere, Category=Obstacle)
-	TSubclassOf<class UNavArea> ObstacleAreaClass;
+	TSubclassOf<UNavArea> ObstacleAreaClass;
 
 	/** radius of state change broadcast */
 	UPROPERTY(EditAnywhere, Category=Broadcast)
@@ -189,7 +194,7 @@ protected:
 	FBroadcastFilter OnBroadcastFilter;
 
 	/** list of agents moving though this link */
-	TArray<TWeakObjectPtr<class UPathFollowingComponent> > MovingAgents;
+	TArray<TWeakObjectPtr<UPathFollowingComponent> > MovingAgents;
 
 	/** delegate to call when link is reached */
 	FOnMoveReachedLink OnMoveReachedLink;
diff --git a/Engine/Source/Runtime/Engine/Classes/AI/Navigation/NavLinkProxy.h b/Engine/Source/Runtime/Engine/Classes/AI/Navigation/NavLinkProxy.h
--- a/Engine/Source/Runtime/Engine/Classes/AI/Navigation/NavLinkProxy.h
+++ b/Engine/Source/Runtime/Engine/Classes/AI/Navigation/NavLinkProxy.h
@@ -7,6 +7,10 @@
 #include "GameFramework/Actor.h"
 #include "NavLinkProxy.generated.h"
 
+class UNavLinkDefinition;
+class UPathFollowingComponent;
+struct FNavigationRelevantData;
+
 DECLARE_DYNAMIC_MULTICAST_DELEGATE_TwoParams( FSmartLinkReachedSignature, class AActor*, MovingActor, const FVector&, DestinationPoint );
 
 UCLASS(Blueprintable)
@@ -48,13 +52,13 @@ public:
 #endif // WITH_EDITORONLY_DATA
 
 	// BEGIN INavRelevantInterface
-	virtual void GetNavigationData(struct FNavigationRelevantData& Data) const;
+	virtual void GetNavigationData(FNavigationRelevantData& Data) const;
 	virtual FBox GetNavigationBounds() const;
 	virtual bool IsNavigationRelevant() const;
 	// END INavRelevantInterface
 
 	// BEGIN INavLinkHostInterface
-	virtual bool GetNavigationLinksClasses(TArray<TSubclassOf<class UNavLinkDefinition> >& OutClasses) const override;
+	virtual bool GetNavigationLinksClasses(TArray<TSubclassOf<UNavLinkDefinition> >& OutClasses) const override;
 	virtual bool GetNavigationLinksArray(TArray<FNavigationLink>& OutLink, TArray<FNavigationSegmentLink>& OutSegments) const override;
 	// END INavLinkHostInterface
 
@@ -92,7 +96,7 @@ protected:
 	UPROPERTY(BlueprintAssignable)
 	FSmartLinkReachedSignature OnSmartLinkReached;
 
-	void NotifySmartLinkReached(UNavLinkCustomComponent* LinkComp, class UPathFollowingComponent* PathComp, const FVector& DestPoint);
+	void NotifySmartLinkReached(UNavLinkCustomComponent* LinkComp, UPathFollowingComponent* PathComp, const FVector& DestPoint);
 
 public:
 	/** Returns SmartLinkComp subobject **/
diff --git a/Engine/Source/Runtime/Engine/Classes/AI/Navigation/NavLinkRenderingComponent.h b/Engine/Source/Runtime/Engine/Classes/AI/Navigation/NavLinkRenderingComponent.h
--- a/Engine/Source/Runtime/Engine/Classes/AI/Navigation/NavLinkRenderingComponent.h
+++ b/Engine/Source/Runtime/Engine/Classes/AI/Navigation/NavLinkRenderingComponent.h
@@ -1,8 +1,11 @@
 // Copyright 1998-2014 Epic Games, Inc. All Rights Reserved.
 
 #pragma once
+
 #include "NavLinkRenderingComponent.generated.h"
 
+class UNavLinkDefinition;
+
 UCLASS(hidecategories=Object, editinlinenew)
 class ENGINE_API UNavLinkRenderingComponent : public UPrimitiveComponent
 {
diff --git a/Engine/Source/Runtime/Engine/Classes/AI/Navigation/NavModifierComponent.h b/Engine/Source/Runtime/Engine/Classes/AI/Navigation/NavModifierComponent.h
--- a/Engine/Source/Runtime/Engine/Classes/AI/Navigation/NavModifierComponent.h
+++ b/Engine/Source/Runtime/Engine/Classes/AI/Navigation/NavModifierComponent.h
@@ -17,7 +17,7 @@ class UNavModifierComponent : public UNavRelevantComponent
 	TSubclassOf<UNavArea> AreaClass;
 
 	virtual void CalcBounds() override;
-	virtual void GetNavigationData(struct FNavigationRelevantData& Data) const override;
+	virtual void GetNavigationData(FNavigationRelevantData& Data) const override;
 
 protected:
 	FVector ObstacleExtent;
diff --git a/Engine/Source/Runtime/Engine/Classes/AI/Navigation/NavModifierVolume.h b/Engine/Source/Runtime/Engine/Classes/AI/Navigation/NavModifierVolume.h
--- a/Engine/Source/Runtime/Engine/Classes/AI/Navigation/NavModifierVolume.h
+++ b/Engine/Source/Runtime/Engine/Classes/AI/Navigation/NavModifierVolume.h
@@ -16,7 +16,7 @@ class ANavModifierVolume : public AVolume, public INavRelevantInterface
 	UPROPERTY(EditAnywhere, BlueprintReadWrite, Category=Default)
 	TSubclassOf<class UNavArea> AreaClass;
 
-	virtual void GetNavigationData(struct FNavigationRelevantData& Data) const override;
+	virtual void GetNavigationData(FNavigationRelevantData& Data) const override;
 	virtual FBox GetNavigationBounds() const;
 
 #if WITH_EDITOR
diff --git a/Engine/Source/Runtime/Engine/Classes/AI/Navigation/NavRelevantInterface.h b/Engine/Source/Runtime/Engine/Classes/AI/Navigation/NavRelevantInterface.h
--- a/Engine/Source/Runtime/Engine/Classes/AI/Navigation/NavRelevantInterface.h
+++ b/Engine/Source/Runtime/Engine/Classes/AI/Navigation/NavRelevantInterface.h
@@ -4,6 +4,8 @@
 #include "Interface.h"
 #include "NavRelevantInterface.generated.h"
 
+struct FNavigationRelevantData;
+
 UINTERFACE(MinimalAPI, meta=(CannotImplementInterfaceInBlueprint))
 class UNavRelevantInterface : public UInterface
 {
@@ -15,7 +17,7 @@ class INavRelevantInterface
 	GENERATED_IINTERFACE_BODY()
 
 	/** Prepare navigation modifiers */
-	virtual void GetNavigationData(struct FNavigationRelevantData& Data) const {}
+	virtual void GetNavigationData(FNavigationRelevantData& Data) const {}
 
 	/** Get bounds for navigation octree */
 	virtual FBox GetNavigationBounds() const { return FBox(0); }
diff --git a/Engine/Source/Runtime/Engine/Classes/AI/Navigation/NavigationData.h b/Engine/Source/Runtime/Engine/Classes/AI/Navigation/NavigationData.h
--- a/Engine/Source/Runtime/Engine/Classes/AI/Navigation/NavigationData.h
+++ b/Engine/Source/Runtime/Engine/Classes/AI/Navigation/NavigationData.h
@@ -7,7 +7,9 @@
 #include "UniquePtr.h"
 #include "NavigationData.generated.h"
 
+class UNavigationQueryFilter;
 class FNavDataGenerator; 
+class INavLinkCustomInterface;
 
 USTRUCT()
 struct FSupportedAreaData
@@ -689,20 +691,20 @@ public:
 	// Custom navigation links
 	//----------------------------------------------------------------------//
 
-	virtual void UpdateCustomLink(const class INavLinkCustomInterface* CustomLink);
+	virtual void UpdateCustomLink(const INavLinkCustomInterface* CustomLink);
 
 	//----------------------------------------------------------------------//
 	// Filters
 	//----------------------------------------------------------------------//
 
 	/** get cached query filter */
-	TSharedPtr<const FNavigationQueryFilter> GetQueryFilter(TSubclassOf<class UNavigationQueryFilter> FilterClass) const;
+	TSharedPtr<const FNavigationQueryFilter> GetQueryFilter(TSubclassOf<UNavigationQueryFilter> FilterClass) const;
 
 	/** store cached query filter */
-	void StoreQueryFilter(TSubclassOf<class UNavigationQueryFilter> FilterClass, TSharedPtr<const FNavigationQueryFilter> NavFilter);
+	void StoreQueryFilter(TSubclassOf<UNavigationQueryFilter> FilterClass, TSharedPtr<const FNavigationQueryFilter> NavFilter);
 
 	/** removes cached query filter */
-	void RemoveQueryFilter(TSubclassOf<class UNavigationQueryFilter> FilterClass);
+	void RemoveQueryFilter(TSubclassOf<UNavigationQueryFilter> FilterClass);
 
 	//----------------------------------------------------------------------//
 	// all the rest                                                                
diff --git a/Engine/Source/Runtime/Engine/Classes/AI/Navigation/NavigationTestingActor.h b/Engine/Source/Runtime/Engine/Classes/AI/Navigation/NavigationTestingActor.h
--- a/Engine/Source/Runtime/Engine/Classes/AI/Navigation/NavigationTestingActor.h
+++ b/Engine/Source/Runtime/Engine/Classes/AI/Navigation/NavigationTestingActor.h
@@ -5,9 +5,11 @@
 #include "NavPathObserverInterface.h"
 #include "NavigationTestingActor.generated.h"
 
+class ANavigationTestingActor;
+
 struct FNavTestTickHelper : FTickableGameObject
 {
-	TWeakObjectPtr<class ANavigationTestingActor> Owner;
+	TWeakObjectPtr<ANavigationTestingActor> Owner;
 
 	FNavTestTickHelper() : Owner(NULL) {}
 	virtual void Tick(float DeltaTime);
@@ -162,8 +164,8 @@ public:
 
 	void UpdateNavData();
 	void UpdatePathfinding();
-	void GatherDetailedData(class ANavigationTestingActor* Goal);
-	void SearchPathTo(class ANavigationTestingActor* Goal);
+	void GatherDetailedData(ANavigationTestingActor* Goal);
+	void SearchPathTo(ANavigationTestingActor* Goal);
 
 	/*	Called when given path becomes invalid (via @see PathObserverDelegate)
 	 *	NOTE: InvalidatedPath doesn't have to be instance's current Path
diff --git a/Engine/Source/Runtime/Engine/Classes/AI/Navigation/RecastNavMesh.h b/Engine/Source/Runtime/Engine/Classes/AI/Navigation/RecastNavMesh.h
--- a/Engine/Source/Runtime/Engine/Classes/AI/Navigation/RecastNavMesh.h
+++ b/Engine/Source/Runtime/Engine/Classes/AI/Navigation/RecastNavMesh.h
@@ -39,6 +39,7 @@ class UNavArea;
 class UNavigationSystem;
 class URecastNavMeshDataChunk;
 struct FAreaNavModifier;
+const FNavigationQueryFilter;
 
 UENUM()
 namespace ERecastPartitioning
diff --git a/Engine/Source/Runtime/Engine/Private/AI/Navigation/AvoidanceManager.cpp b/Engine/Source/Runtime/Engine/Private/AI/Navigation/AvoidanceManager.cpp
--- a/Engine/Source/Runtime/Engine/Private/AI/Navigation/AvoidanceManager.cpp
+++ b/Engine/Source/Runtime/Engine/Private/AI/Navigation/AvoidanceManager.cpp
@@ -22,7 +22,7 @@ FNavAvoidanceData::FNavAvoidanceData(UAvoidanceManager* Manager, IRVOAvoidanceIn
 		);
 }
 
-void FNavAvoidanceData::Init(class UAvoidanceManager* Avoidance, const FVector& InCenter, float InRadius, float InHeight,
+void FNavAvoidanceData::Init(UAvoidanceManager* Avoidance, const FVector& InCenter, float InRadius, float InHeight,
 							 const FVector& InVelocity, float InWeight,
 							 int32 InGroupMask, int32 InGroupsToAvoid, int32 InGroupsToIgnore,
 							 float InTestRadius2D)
diff --git a/Engine/Source/Runtime/Engine/Private/AI/Navigation/NavAreas/NavArea.cpp b/Engine/Source/Runtime/Engine/Private/AI/Navigation/NavAreas/NavArea.cpp
--- a/Engine/Source/Runtime/Engine/Private/AI/Navigation/NavAreas/NavArea.cpp
+++ b/Engine/Source/Runtime/Engine/Private/AI/Navigation/NavAreas/NavArea.cpp
@@ -33,7 +33,7 @@ void UNavArea::FinishDestroy()
 	Super::FinishDestroy();
 }
 
-FColor UNavArea::GetColor(class UClass* AreaDefinitionClass)
+FColor UNavArea::GetColor(UClass* AreaDefinitionClass)
 {
 	return AreaDefinitionClass ? AreaDefinitionClass->GetDefaultObject<UNavArea>()->DrawColor : FColor::Black;
 }
diff --git a/Engine/Source/Runtime/Engine/Private/AI/Navigation/NavAreas/NavAreaMeta.cpp b/Engine/Source/Runtime/Engine/Private/AI/Navigation/NavAreas/NavAreaMeta.cpp
--- a/Engine/Source/Runtime/Engine/Private/AI/Navigation/NavAreas/NavAreaMeta.cpp
+++ b/Engine/Source/Runtime/Engine/Private/AI/Navigation/NavAreas/NavAreaMeta.cpp
@@ -7,7 +7,7 @@ UNavAreaMeta::UNavAreaMeta(const FObjectInitializer& ObjectInitializer) : Super(
 {
 }
 
-TSubclassOf<UNavArea> UNavAreaMeta::PickAreaClass(const class AActor* Actor, const FNavAgentProperties& NavAgent) const 
+TSubclassOf<UNavArea> UNavAreaMeta::PickAreaClass(const AActor* Actor, const FNavAgentProperties& NavAgent) const 
 { 
 	UE_LOG(LogNavigation, Warning, TEXT("UNavAreaMeta::PickAreaClass called. UNavAreaMeta is an abstract class and should never get used directly!"));
 
diff --git a/Engine/Source/Runtime/Engine/Private/AI/Navigation/NavAreas/NavAreaMeta_SwitchByAgent.cpp b/Engine/Source/Runtime/Engine/Private/AI/Navigation/NavAreas/NavAreaMeta_SwitchByAgent.cpp
--- a/Engine/Source/Runtime/Engine/Private/AI/Navigation/NavAreas/NavAreaMeta_SwitchByAgent.cpp
+++ b/Engine/Source/Runtime/Engine/Private/AI/Navigation/NavAreas/NavAreaMeta_SwitchByAgent.cpp
@@ -24,7 +24,7 @@ UNavAreaMeta_SwitchByAgent::UNavAreaMeta_SwitchByAgent(const FObjectInitializer&
 	Agent15Area = UNavArea_Default::StaticClass();
 }
 
-TSubclassOf<UNavArea> UNavAreaMeta_SwitchByAgent::PickAreaClass(const class AActor* Actor, const FNavAgentProperties& NavAgent) const
+TSubclassOf<UNavArea> UNavAreaMeta_SwitchByAgent::PickAreaClass(const AActor* Actor, const FNavAgentProperties& NavAgent) const
 {
 	const int8 AgentIndex = GetNavAgentIndex(NavAgent);
 	TSubclassOf<UNavArea> UseAreaClass = NULL;
diff --git a/Engine/Source/Runtime/Engine/Private/AI/Navigation/NavCollision.cpp b/Engine/Source/Runtime/Engine/Private/AI/Navigation/NavCollision.cpp
--- a/Engine/Source/Runtime/Engine/Private/AI/Navigation/NavCollision.cpp
+++ b/Engine/Source/Runtime/Engine/Private/AI/Navigation/NavCollision.cpp
@@ -48,8 +48,7 @@ public:
 class FDerivedDataNavCollisionCooker : public FDerivedDataPluginInterface
 {
 private:
-
-	class UNavCollision* NavCollisionInstance;
+	UNavCollision* NavCollisionInstance;
 	UObject* CollisionDataProvider;
 	FName Format;
 	FGuid DataGuid;
@@ -207,7 +206,7 @@ void UNavCollision::ClearCollision()
 	bHasConvexGeometry = false;
 }
 
-void UNavCollision::GetNavigationModifier(struct FCompositeNavModifier& Modifier, const FTransform& LocalToWorld) const
+void UNavCollision::GetNavigationModifier(FCompositeNavModifier& Modifier, const FTransform& LocalToWorld) const
 {
 	const TSubclassOf<UNavArea> UseAreaClass = AreaClass ? AreaClass : UNavigationSystem::GetDefaultObstacleArea();
 
@@ -297,7 +296,7 @@ void DrawBoxHelper(FPrimitiveDrawInterface* PDI, const FMatrix& ElemTM, const FV
 	}
 }
 
-void UNavCollision::DrawSimpleGeom(class FPrimitiveDrawInterface* PDI, const FTransform& Transform, const FColor Color)
+void UNavCollision::DrawSimpleGeom(FPrimitiveDrawInterface* PDI, const FTransform& Transform, const FColor Color)
 {
 	const FMatrix ParentTM = Transform.ToMatrixWithScale();
 	for (int32 i = 0; i < CylinderCollision.Num(); i++)
diff --git a/Engine/Source/Runtime/Engine/Private/AI/Navigation/NavFilters/NavigationQueryFilter.cpp b/Engine/Source/Runtime/Engine/Private/AI/Navigation/NavFilters/NavigationQueryFilter.cpp
--- a/Engine/Source/Runtime/Engine/Private/AI/Navigation/NavFilters/NavigationQueryFilter.cpp
+++ b/Engine/Source/Runtime/Engine/Private/AI/Navigation/NavFilters/NavigationQueryFilter.cpp
@@ -116,7 +116,7 @@ UNavigationQueryFilter::UNavigationQueryFilter(const FObjectInitializer& ObjectI
 	ExcludeFlags.Packed = 0;
 }
 
-TSharedPtr<const FNavigationQueryFilter> UNavigationQueryFilter::GetQueryFilter(const class ANavigationData* NavData) const
+TSharedPtr<const FNavigationQueryFilter> UNavigationQueryFilter::GetQueryFilter(const ANavigationData* NavData) const
 {
 	TSharedPtr<const FNavigationQueryFilter> SharedFilter = NavData->GetQueryFilter(GetClass());
 	if (!SharedFilter.IsValid())
@@ -172,13 +172,13 @@ void UNavigationQueryFilter::InitializeFilter(const ANavigationData* NavData, FN
 	Filter->SetExcludeFlags(ExcludeFlags.Packed);
 }
 
-TSharedPtr<const struct FNavigationQueryFilter> UNavigationQueryFilter::GetQueryFilter(const class ANavigationData* NavData, UClass* FilterClass)
+TSharedPtr<const FNavigationQueryFilter> UNavigationQueryFilter::GetQueryFilter(const ANavigationData* NavData, UClass* FilterClass)
 {
 	UNavigationQueryFilter* DefFilterOb = FilterClass ? FilterClass->GetDefaultObject<UNavigationQueryFilter>() : NULL;
 	return NavData && DefFilterOb ? DefFilterOb->GetQueryFilter(NavData) : NULL;
 }
 
-void UNavigationQueryFilter::AddTravelCostOverride(TSubclassOf<class UNavArea> AreaClass, float TravelCost)
+void UNavigationQueryFilter::AddTravelCostOverride(TSubclassOf<UNavArea> AreaClass, float TravelCost)
 {
 	int32 Idx = FindAreaOverride(AreaClass);
 	if (Idx == INDEX_NONE)
@@ -193,7 +193,7 @@ void UNavigationQueryFilter::AddTravelCostOverride(TSubclassOf<class UNavArea> A
 	Areas[Idx].TravelCostOverride = TravelCost;
 }
 
-void UNavigationQueryFilter::AddEnteringCostOverride(TSubclassOf<class UNavArea> AreaClass, float EnteringCost)
+void UNavigationQueryFilter::AddEnteringCostOverride(TSubclassOf<UNavArea> AreaClass, float EnteringCost)
 {
 	int32 Idx = FindAreaOverride(AreaClass);
 	if (Idx == INDEX_NONE)
@@ -208,7 +208,7 @@ void UNavigationQueryFilter::AddEnteringCostOverride(TSubclassOf<class UNavArea>
 	Areas[Idx].EnteringCostOverride = EnteringCost;
 }
 
-void UNavigationQueryFilter::AddExcludedArea(TSubclassOf<class UNavArea> AreaClass)
+void UNavigationQueryFilter::AddExcludedArea(TSubclassOf<UNavArea> AreaClass)
 {
 	int32 Idx = FindAreaOverride(AreaClass);
 	if (Idx == INDEX_NONE)
@@ -222,7 +222,7 @@ void UNavigationQueryFilter::AddExcludedArea(TSubclassOf<class UNavArea> AreaCla
 	Areas[Idx].bIsExcluded = true;
 }
 
-int32 UNavigationQueryFilter::FindAreaOverride(TSubclassOf<class UNavArea> AreaClass) const
+int32 UNavigationQueryFilter::FindAreaOverride(TSubclassOf<UNavArea> AreaClass) const
 {
 	for (int32 i = 0; i < Areas.Num(); i++)
 	{
@@ -236,7 +236,7 @@ int32 UNavigationQueryFilter::FindAreaOverride(TSubclassOf<class UNavArea> AreaC
 }
 
 #if WITH_EDITOR
-void UNavigationQueryFilter::PostEditChangeProperty(struct FPropertyChangedEvent& PropertyChangedEvent)
+void UNavigationQueryFilter::PostEditChangeProperty(FPropertyChangedEvent& PropertyChangedEvent)
 {
 	Super::PostEditChangeProperty(PropertyChangedEvent);
 
diff --git a/Engine/Source/Runtime/Engine/Private/AI/Navigation/NavLinkCustomComponent.cpp b/Engine/Source/Runtime/Engine/Private/AI/Navigation/NavLinkCustomComponent.cpp
--- a/Engine/Source/Runtime/Engine/Private/AI/Navigation/NavLinkCustomComponent.cpp
+++ b/Engine/Source/Runtime/Engine/Private/AI/Navigation/NavLinkCustomComponent.cpp
@@ -68,7 +68,7 @@ bool UNavLinkCustomComponent::IsLinkPathfindingAllowed(const UObject* Querier) c
 	return true;
 }
 
-bool UNavLinkCustomComponent::OnLinkMoveStarted(class UPathFollowingComponent* PathComp, const FVector& DestPoint)
+bool UNavLinkCustomComponent::OnLinkMoveStarted(UPathFollowingComponent* PathComp, const FVector& DestPoint)
 {
 	TWeakObjectPtr<UPathFollowingComponent> WeakPathComp = PathComp;
 	MovingAgents.Add(WeakPathComp);
@@ -82,13 +82,13 @@ bool UNavLinkCustomComponent::OnLinkMoveStarted(class UPathFollowingComponent* P
 	return false;
 }
 
-void UNavLinkCustomComponent::OnLinkMoveFinished(class UPathFollowingComponent* PathComp)
+void UNavLinkCustomComponent::OnLinkMoveFinished(UPathFollowingComponent* PathComp)
 {
 	TWeakObjectPtr<UPathFollowingComponent> WeakPathComp;
 	MovingAgents.Remove(WeakPathComp);
 }
 
-void UNavLinkCustomComponent::GetNavigationData(struct FNavigationRelevantData& Data) const
+void UNavLinkCustomComponent::GetNavigationData(FNavigationRelevantData& Data) const
 {
 	FNavigationLink LinkMod = GetLinkModifier();
 	Data.Modifiers.Add(FSimpleLinkNavModifier(LinkMod, GetOwner()->GetTransform()));
@@ -158,7 +158,7 @@ FNavigationLink UNavLinkCustomComponent::GetLinkModifier() const
 	return INavLinkCustomInterface::GetModifier(this);
 }
 
-void UNavLinkCustomComponent::SetEnabledArea(TSubclassOf<class UNavArea> AreaClass)
+void UNavLinkCustomComponent::SetEnabledArea(TSubclassOf<UNavArea> AreaClass)
 {
 	EnabledAreaClass = AreaClass;
 	if (IsNavigationRelevant() && bLinkEnabled)
@@ -168,7 +168,7 @@ void UNavLinkCustomComponent::SetEnabledArea(TSubclassOf<class UNavArea> AreaCla
 	}
 }
 
-void UNavLinkCustomComponent::SetDisabledArea(TSubclassOf<class UNavArea> AreaClass)
+void UNavLinkCustomComponent::SetDisabledArea(TSubclassOf<UNavArea> AreaClass)
 {
 	DisabledAreaClass = AreaClass;
 	if (IsNavigationRelevant() && !bLinkEnabled)
@@ -178,7 +178,7 @@ void UNavLinkCustomComponent::SetDisabledArea(TSubclassOf<class UNavArea> AreaCl
 	}
 }
 
-void UNavLinkCustomComponent::AddNavigationObstacle(TSubclassOf<class UNavArea> AreaClass, const FVector& BoxExtent, const FVector& BoxOffset)
+void UNavLinkCustomComponent::AddNavigationObstacle(TSubclassOf<UNavArea> AreaClass, const FVector& BoxExtent, const FVector& BoxOffset)
 {
 	ObstacleOffset = BoxOffset;
 	ObstacleExtent = BoxExtent;
diff --git a/Engine/Source/Runtime/Engine/Private/AI/Navigation/NavLinkProxy.cpp b/Engine/Source/Runtime/Engine/Private/AI/Navigation/NavLinkProxy.cpp
--- a/Engine/Source/Runtime/Engine/Private/AI/Navigation/NavLinkProxy.cpp
+++ b/Engine/Source/Runtime/Engine/Private/AI/Navigation/NavLinkProxy.cpp
@@ -74,7 +74,7 @@ void ANavLinkProxy::PostEditChangeProperty(FPropertyChangedEvent& PropertyChange
 }
 #endif // WITH_EDITOR
 
-void ANavLinkProxy::GetNavigationData(struct FNavigationRelevantData& Data) const
+void ANavLinkProxy::GetNavigationData(FNavigationRelevantData& Data) const
 {
 	NavigationHelper::ProcessNavLinkAndAppend(&Data.Modifiers, this, PointLinks);
 	NavigationHelper::ProcessNavLinkSegmentAndAppend(&Data.Modifiers, this, SegmentLinks);
@@ -90,7 +90,7 @@ bool ANavLinkProxy::IsNavigationRelevant() const
 	return (PointLinks.Num() > 0) || (SegmentLinks.Num() > 0);
 }
 
-bool ANavLinkProxy::GetNavigationLinksClasses(TArray<TSubclassOf<class UNavLinkDefinition> >& OutClasses) const
+bool ANavLinkProxy::GetNavigationLinksClasses(TArray<TSubclassOf<UNavLinkDefinition> >& OutClasses) const
 {
 	return false;
 }
@@ -136,7 +136,7 @@ FBox ANavLinkProxy::GetComponentsBoundingBox(bool bNonColliding) const
 	return LinksBB;
 }
 
-void ANavLinkProxy::NotifySmartLinkReached(UNavLinkCustomComponent* LinkComp, class UPathFollowingComponent* PathComp, const FVector& DestPoint)
+void ANavLinkProxy::NotifySmartLinkReached(UNavLinkCustomComponent* LinkComp, UPathFollowingComponent* PathComp, const FVector& DestPoint)
 {
 	AActor* PathOwner = PathComp->GetOwner();
 	AController* ControllerOwner = Cast<AController>(PathOwner);
diff --git a/Engine/Source/Runtime/Engine/Private/AI/Navigation/NavLinkRenderingComponent.cpp b/Engine/Source/Runtime/Engine/Private/AI/Navigation/NavLinkRenderingComponent.cpp
--- a/Engine/Source/Runtime/Engine/Private/AI/Navigation/NavLinkRenderingComponent.cpp
+++ b/Engine/Source/Runtime/Engine/Private/AI/Navigation/NavLinkRenderingComponent.cpp
@@ -103,7 +103,7 @@ FNavLinkRenderingProxy::FNavLinkRenderingProxy(const UPrimitiveComponent* InComp
 	if (LinkOwnerActor != NULL && LinkOwnerHost != NULL)
 	{
 		const FTransform LocalToWorld = LinkOwnerActor->ActorToWorld();
-		TArray<TSubclassOf<class UNavLinkDefinition> > NavLinkClasses;
+		TArray<TSubclassOf<UNavLinkDefinition> > NavLinkClasses;
 		LinkOwnerHost->GetNavigationLinksClasses(NavLinkClasses);
 
 		for (int32 NavLinkClassIdx = 0; NavLinkClassIdx < NavLinkClasses.Num(); ++NavLinkClassIdx)
diff --git a/Engine/Source/Runtime/Engine/Private/AI/Navigation/NavModifierComponent.cpp b/Engine/Source/Runtime/Engine/Private/AI/Navigation/NavModifierComponent.cpp
--- a/Engine/Source/Runtime/Engine/Private/AI/Navigation/NavModifierComponent.cpp
+++ b/Engine/Source/Runtime/Engine/Private/AI/Navigation/NavModifierComponent.cpp
@@ -22,7 +22,7 @@ void UNavModifierComponent::CalcBounds()
 	}
 }
 
-void UNavModifierComponent::GetNavigationData(struct FNavigationRelevantData& Data) const
+void UNavModifierComponent::GetNavigationData(FNavigationRelevantData& Data) const
 {
 	Data.Modifiers.Add(FAreaNavModifier(Bounds, FTransform::Identity, AreaClass));
 }
diff --git a/Engine/Source/Runtime/Engine/Private/AI/Navigation/NavModifierVolume.cpp b/Engine/Source/Runtime/Engine/Private/AI/Navigation/NavModifierVolume.cpp
--- a/Engine/Source/Runtime/Engine/Private/AI/Navigation/NavModifierVolume.cpp
+++ b/Engine/Source/Runtime/Engine/Private/AI/Navigation/NavModifierVolume.cpp
@@ -12,7 +12,7 @@ ANavModifierVolume::ANavModifierVolume(const FObjectInitializer& ObjectInitializ
 {
 }
 
-void ANavModifierVolume::GetNavigationData(struct FNavigationRelevantData& Data) const
+void ANavModifierVolume::GetNavigationData(FNavigationRelevantData& Data) const
 {
 	if (Brush && AreaClass && AreaClass != UNavigationSystem::GetDefaultWalkableArea())
 	{
diff --git a/Engine/Source/Runtime/Engine/Private/AI/Navigation/NavigationData.cpp b/Engine/Source/Runtime/Engine/Private/AI/Navigation/NavigationData.cpp
--- a/Engine/Source/Runtime/Engine/Private/AI/Navigation/NavigationData.cpp
+++ b/Engine/Source/Runtime/Engine/Private/AI/Navigation/NavigationData.cpp
@@ -8,7 +8,7 @@
 //----------------------------------------------------------------------//
 // FPathFindingQuery
 //----------------------------------------------------------------------//
-FPathFindingQuery::FPathFindingQuery(const UObject* InOwner, const class ANavigationData* InNavData, const FVector& Start, const FVector& End, TSharedPtr<const FNavigationQueryFilter> SourceQueryFilter, FNavPathSharedPtr InPathInstanceToFill)
+FPathFindingQuery::FPathFindingQuery(const UObject* InOwner, const ANavigationData* InNavData, const FVector& Start, const FVector& End, TSharedPtr<const FNavigationQueryFilter> SourceQueryFilter, FNavPathSharedPtr InPathInstanceToFill)
 : NavData(InNavData)
 , Owner(InOwner)
 , StartLocation(Start)
@@ -597,22 +597,22 @@ void ANavigationData::SetNavRenderingEnabled(bool bEnable)
 	}
 }
 
-void ANavigationData::UpdateCustomLink(const class INavLinkCustomInterface* CustomLink)
+void ANavigationData::UpdateCustomLink(const INavLinkCustomInterface* CustomLink)
 {
 	// no implementation for abstract class
 }
 
-TSharedPtr<const FNavigationQueryFilter> ANavigationData::GetQueryFilter(TSubclassOf<class UNavigationQueryFilter> FilterClass) const
+TSharedPtr<const FNavigationQueryFilter> ANavigationData::GetQueryFilter(TSubclassOf<UNavigationQueryFilter> FilterClass) const
 {
 	return QueryFilters.FindRef(FilterClass);
 }
 
-void ANavigationData::StoreQueryFilter(TSubclassOf<class UNavigationQueryFilter> FilterClass, TSharedPtr<const FNavigationQueryFilter> NavFilter)
+void ANavigationData::StoreQueryFilter(TSubclassOf<UNavigationQueryFilter> FilterClass, TSharedPtr<const FNavigationQueryFilter> NavFilter)
 {
 	QueryFilters.Add(FilterClass, NavFilter);
 }
 
-void ANavigationData::RemoveQueryFilter(TSubclassOf<class UNavigationQueryFilter> FilterClass)
+void ANavigationData::RemoveQueryFilter(TSubclassOf<UNavigationQueryFilter> FilterClass)
 {
 	QueryFilters.Remove(FilterClass);
 }
diff --git a/Engine/Source/Runtime/Engine/Private/AI/Navigation/NavigationModifier.cpp b/Engine/Source/Runtime/Engine/Private/AI/Navigation/NavigationModifier.cpp
--- a/Engine/Source/Runtime/Engine/Private/AI/Navigation/NavigationModifier.cpp
+++ b/Engine/Source/Runtime/Engine/Private/AI/Navigation/NavigationModifier.cpp
@@ -59,7 +59,7 @@ const TArray<FNavigationSegmentLink>& UNavLinkDefinition::GetSegmentLinksDefinit
 }
 
 #if WITH_EDITOR
-void UNavLinkDefinition::PostEditChangeProperty(struct FPropertyChangedEvent& PropertyChangedEvent)
+void UNavLinkDefinition::PostEditChangeProperty(FPropertyChangedEvent& PropertyChangedEvent)
 {
 	// In case relevant data has changed, clear the flag that says we've determined whether there's a meta area class
 	// so it will be recalculated the next time it's needed.
@@ -133,7 +133,7 @@ bool UNavLinkDefinition::HasAdjustableLinks() const
 // FAreaNavModifier
 //----------------------------------------------------------------------//
 
-FAreaNavModifier::FAreaNavModifier(float Radius, float Height, const FTransform& LocalToWorld, const TSubclassOf<class UNavArea> InAreaClass)
+FAreaNavModifier::FAreaNavModifier(float Radius, float Height, const FTransform& LocalToWorld, const TSubclassOf<UNavArea> InAreaClass)
 {
 	Init(InAreaClass);
 
@@ -146,31 +146,31 @@ FAreaNavModifier::FAreaNavModifier(float Radius, float Height, const FTransform&
 	Bounds = FBox::BuildAABB(LocalToWorld.GetLocation(), FVector(Radius, Radius, Height));
 }
 
-FAreaNavModifier::FAreaNavModifier(const FVector& Extent, const FTransform& LocalToWorld, const TSubclassOf<class UNavArea> InAreaClass)
+FAreaNavModifier::FAreaNavModifier(const FVector& Extent, const FTransform& LocalToWorld, const TSubclassOf<UNavArea> InAreaClass)
 {
 	Init(InAreaClass);
 	SetBox(FBox::BuildAABB(FVector::ZeroVector, Extent), LocalToWorld);
 }
 
-FAreaNavModifier::FAreaNavModifier(const FBox& Box, const FTransform& LocalToWorld, const TSubclassOf<class UNavArea> InAreaClass)
+FAreaNavModifier::FAreaNavModifier(const FBox& Box, const FTransform& LocalToWorld, const TSubclassOf<UNavArea> InAreaClass)
 {
 	Init(InAreaClass);
 	SetBox(Box, LocalToWorld);
 }
 
-FAreaNavModifier::FAreaNavModifier(const TArray<FVector>& InPoints, ENavigationCoordSystem::Type CoordType, const FTransform& LocalToWorld, const TSubclassOf<class UNavArea> InAreaClass)
+FAreaNavModifier::FAreaNavModifier(const TArray<FVector>& InPoints, ENavigationCoordSystem::Type CoordType, const FTransform& LocalToWorld, const TSubclassOf<UNavArea> InAreaClass)
 {
 	Init(InAreaClass);
 	SetConvex(InPoints, 0, InPoints.Num(), CoordType, LocalToWorld);
 }
 
-FAreaNavModifier::FAreaNavModifier(const TArray<FVector>& InPoints, const int32 FirstIndex, const int32 LastIndex, ENavigationCoordSystem::Type CoordType, const FTransform& LocalToWorld, const TSubclassOf<class UNavArea> InAreaClass)
+FAreaNavModifier::FAreaNavModifier(const TArray<FVector>& InPoints, const int32 FirstIndex, const int32 LastIndex, ENavigationCoordSystem::Type CoordType, const FTransform& LocalToWorld, const TSubclassOf<UNavArea> InAreaClass)
 {
 	Init(InAreaClass);
 	SetConvex(InPoints, FirstIndex, LastIndex, CoordType, LocalToWorld);
 }
 
-FAreaNavModifier::FAreaNavModifier(const class UBrushComponent* BrushComponent, const TSubclassOf<class UNavArea> InAreaClass)
+FAreaNavModifier::FAreaNavModifier(const UBrushComponent* BrushComponent, const TSubclassOf<UNavArea> InAreaClass)
 {
 	check(BrushComponent != NULL);
 
@@ -215,7 +215,7 @@ void FAreaNavModifier::GetConvex(FConvexNavAreaData& Data) const
 	Data.MaxZ = LastPoint.Y;
 }
 
-void FAreaNavModifier::Init(const TSubclassOf<class UNavArea> InAreaClass)
+void FAreaNavModifier::Init(const TSubclassOf<UNavArea> InAreaClass)
 {
 	Cost = 0.0f;
 	FixedCost = 0.0f;
@@ -223,14 +223,14 @@ void FAreaNavModifier::Init(const TSubclassOf<class UNavArea> InAreaClass)
 	SetAreaClass(InAreaClass);
 }
 
-void FAreaNavModifier::SetAreaClass(const TSubclassOf<class UNavArea> InAreaClass)
+void FAreaNavModifier::SetAreaClass(const TSubclassOf<UNavArea> InAreaClass)
 {
 	AreaClass = InAreaClass;
 	bHasMetaAreas = (ReplaceAreaClass && ReplaceAreaClass->IsChildOf(UNavAreaMeta::StaticClass())) ||
 		(AreaClass && AreaClass->IsChildOf(UNavAreaMeta::StaticClass()));
 }
 
-void FAreaNavModifier::SetAreaClassToReplace(const TSubclassOf<class UNavArea> InAreaClass)
+void FAreaNavModifier::SetAreaClassToReplace(const TSubclassOf<UNavArea> InAreaClass)
 {
 	ReplaceAreaClass = InAreaClass;
 	bHasMetaAreas = (ReplaceAreaClass && ReplaceAreaClass->IsChildOf(UNavAreaMeta::StaticClass())) ||
@@ -346,7 +346,7 @@ void FAreaNavModifier::SetConvex(const TArray<FVector>& InPoints, const int32 Fi
 //----------------------------------------------------------------------//
 // FCustomLinkNavModifier
 //----------------------------------------------------------------------//
-void FCustomLinkNavModifier::Set(TSubclassOf<class UNavLinkDefinition> InPresetLinkClass, const FTransform& InLocalToWorld)
+void FCustomLinkNavModifier::Set(TSubclassOf<UNavLinkDefinition> InPresetLinkClass, const FTransform& InLocalToWorld)
 {
 	LinkDefinitionClass = InPresetLinkClass;
 	LocalToWorld = InLocalToWorld;
diff --git a/Engine/Source/Runtime/Engine/Private/AI/Navigation/NavigationSystem.cpp b/Engine/Source/Runtime/Engine/Private/AI/Navigation/NavigationSystem.cpp
--- a/Engine/Source/Runtime/Engine/Private/AI/Navigation/NavigationSystem.cpp
+++ b/Engine/Source/Runtime/Engine/Private/AI/Navigation/NavigationSystem.cpp
@@ -1841,7 +1841,7 @@ int32 GetDirtyFlagHelper(int32 UpdateFlags, int32 DefaultValue)
 		DefaultValue;
 }
 
-FSetElementId UNavigationSystem::RegisterNavOctreeElement(UObject* ElementOwner, class INavRelevantInterface* ElementInterface, int32 UpdateFlags)
+FSetElementId UNavigationSystem::RegisterNavOctreeElement(UObject* ElementOwner, INavRelevantInterface* ElementInterface, int32 UpdateFlags)
 {
 	FSetElementId SetId;
 
@@ -1983,7 +1983,7 @@ bool UNavigationSystem::GetNavOctreeElementData(UObject* NodeOwner, int32& Dirty
 	return false;
 }
 
-void UNavigationSystem::UnregisterNavOctreeElement(UObject* ElementOwner, class INavRelevantInterface* ElementInterface, int32 UpdateFlags)
+void UNavigationSystem::UnregisterNavOctreeElement(UObject* ElementOwner, INavRelevantInterface* ElementInterface, int32 UpdateFlags)
 {
 #if WITH_EDITOR
 	if (IsNavigationUnregisterLocked())
@@ -2044,7 +2044,7 @@ void UNavigationSystem::RemoveNavOctreeElementId(const FOctreeElementId& Element
 	}
 }
 
-void UNavigationSystem::UpdateNavOctree(class AActor* Actor)
+void UNavigationSystem::UpdateNavOctree(AActor* Actor)
 {
 	SCOPE_CYCLE_COUNTER(STAT_DebugNavOctree);
 
@@ -2059,7 +2059,7 @@ void UNavigationSystem::UpdateNavOctree(class AActor* Actor)
 	}
 }
 
-void UNavigationSystem::UpdateNavOctree(class UActorComponent* Comp)
+void UNavigationSystem::UpdateNavOctree(UActorComponent* Comp)
 {
 	SCOPE_CYCLE_COUNTER(STAT_DebugNavOctree);
 
@@ -2092,7 +2092,7 @@ void UNavigationSystem::UpdateNavOctree(class UActorComponent* Comp)
 	}
 }
 
-void UNavigationSystem::UpdateNavOctreeAll(class AActor* Actor)
+void UNavigationSystem::UpdateNavOctreeAll(AActor* Actor)
 {
 	if (Actor)
 	{
@@ -2108,7 +2108,7 @@ void UNavigationSystem::UpdateNavOctreeAll(class AActor* Actor)
 	}
 }
 
-void UNavigationSystem::UpdateNavOctreeBounds(class AActor* Actor)
+void UNavigationSystem::UpdateNavOctreeBounds(AActor* Actor)
 {
 	TArray<UActorComponent*> Components;
 	Actor->GetComponents(Components);
@@ -2123,7 +2123,7 @@ void UNavigationSystem::UpdateNavOctreeBounds(class AActor* Actor)
 	}
 }
 
-void UNavigationSystem::ClearNavOctreeAll(class AActor* Actor)
+void UNavigationSystem::ClearNavOctreeAll(AActor* Actor)
 {
 	if (Actor)
 	{
@@ -2139,7 +2139,7 @@ void UNavigationSystem::ClearNavOctreeAll(class AActor* Actor)
 	}
 }
 
-void UNavigationSystem::UpdateNavOctreeElement(UObject* ElementOwner, class INavRelevantInterface* ElementInterface, int32 UpdateFlags)
+void UNavigationSystem::UpdateNavOctreeElement(UObject* ElementOwner, INavRelevantInterface* ElementInterface, int32 UpdateFlags)
 {
 	INC_DWORD_STAT(STAT_Navigation_UpdateNavOctree);
 
@@ -2244,7 +2244,7 @@ void UNavigationSystem::OnComponentUnregistered(UActorComponent* Comp)
 	}
 }
 
-void UNavigationSystem::OnActorRegistered(class AActor* Actor)
+void UNavigationSystem::OnActorRegistered(AActor* Actor)
 {
 	SCOPE_CYCLE_COUNTER(STAT_DebugNavOctree);
 	INavRelevantInterface* NavInterface = Cast<INavRelevantInterface>(Actor);
@@ -2258,7 +2258,7 @@ void UNavigationSystem::OnActorRegistered(class AActor* Actor)
 	}
 }
 
-void UNavigationSystem::OnActorUnregistered(class AActor* Actor)
+void UNavigationSystem::OnActorUnregistered(AActor* Actor)
 {
 	SCOPE_CYCLE_COUNTER(STAT_DebugNavOctree);
 	INavRelevantInterface* NavInterface = Cast<INavRelevantInterface>(Actor);
diff --git a/Engine/Source/Runtime/Engine/Private/AI/Navigation/NavigationSystemHelpers.cpp b/Engine/Source/Runtime/Engine/Private/AI/Navigation/NavigationSystemHelpers.cpp
--- a/Engine/Source/Runtime/Engine/Private/AI/Navigation/NavigationSystemHelpers.cpp
+++ b/Engine/Source/Runtime/Engine/Private/AI/Navigation/NavigationSystemHelpers.cpp
@@ -12,7 +12,7 @@
 
 namespace NavigationHelper
 {
-	void GatherCollision(class UBodySetup* RigidBody, TNavStatArray<FVector>& OutVertexBuffer, TNavStatArray<int32>& OutIndexBuffer, const FTransform& LocalToWorld)
+	void GatherCollision(UBodySetup* RigidBody, TNavStatArray<FVector>& OutVertexBuffer, TNavStatArray<int32>& OutIndexBuffer, const FTransform& LocalToWorld)
 	{
 		if (RigidBody == NULL)
 		{
@@ -23,7 +23,7 @@ namespace NavigationHelper
 #endif // WITH_RECAST
 	}
 
-	void GatherCollision(class UBodySetup* RigidBody, class UNavCollision* NavCollision)
+	void GatherCollision(UBodySetup* RigidBody, UNavCollision* NavCollision)
 	{
 		if (RigidBody == NULL || NavCollision == NULL)
 		{
@@ -62,7 +62,7 @@ namespace NavigationHelper
 		return FallDownHeight;
 	}
 
-	void DefaultNavLinkProcessorImpl(struct FCompositeNavModifier* OUT CompositeModifier, const class AActor* Actor, const TArray<FNavigationLink>& IN NavLinks)
+	void DefaultNavLinkProcessorImpl(FCompositeNavModifier* OUT CompositeModifier, const AActor* Actor, const TArray<FNavigationLink>& IN NavLinks)
 	{
 		const FTransform LocalToWorld = Actor->ActorToWorld();
 		FSimpleLinkNavModifier SimpleLink(NavLinks, LocalToWorld);
@@ -97,7 +97,7 @@ namespace NavigationHelper
 		CompositeModifier->Add(SimpleLink);
 	}
 
-	void DefaultNavLinkSegmentProcessorImpl(struct FCompositeNavModifier* OUT CompositeModifier, const class AActor* Actor, const TArray<FNavigationSegmentLink>& IN NavLinks)
+	void DefaultNavLinkSegmentProcessorImpl(FCompositeNavModifier* OUT CompositeModifier, const AActor* Actor, const TArray<FNavigationSegmentLink>& IN NavLinks)
 	{
 		const FTransform LocalToWorld = Actor->ActorToWorld();
 		FSimpleLinkNavModifier SimpleLink(NavLinks, LocalToWorld);
@@ -147,7 +147,7 @@ namespace NavigationHelper
 	FNavLinkProcessorDelegate NavLinkProcessor = FNavLinkProcessorDelegate::CreateStatic(DefaultNavLinkProcessorImpl);
 	FNavLinkSegmentProcessorDelegate NavLinkSegmentProcessor = FNavLinkSegmentProcessorDelegate::CreateStatic(DefaultNavLinkSegmentProcessorImpl);
 
-	void ProcessNavLinkAndAppend(struct FCompositeNavModifier* OUT CompositeModifier, const class AActor* Actor, const TArray<FNavigationLink>& IN NavLinks)
+	void ProcessNavLinkAndAppend(FCompositeNavModifier* OUT CompositeModifier, const AActor* Actor, const TArray<FNavigationLink>& IN NavLinks)
 	{
 		SCOPE_CYCLE_COUNTER(STAT_Navigation_AdjustingNavLinks);
 
@@ -158,7 +158,7 @@ namespace NavigationHelper
 		}
 	}
 
-	void ProcessNavLinkSegmentAndAppend(struct FCompositeNavModifier* OUT CompositeModifier, const class AActor* Actor, const TArray<FNavigationSegmentLink>& IN NavLinks)
+	void ProcessNavLinkSegmentAndAppend(FCompositeNavModifier* OUT CompositeModifier, const AActor* Actor, const TArray<FNavigationSegmentLink>& IN NavLinks)
 	{
 		SCOPE_CYCLE_COUNTER(STAT_Navigation_AdjustingNavLinks);
 
diff --git a/Engine/Source/Runtime/Engine/Private/AI/Navigation/NavigationTestingActor.cpp b/Engine/Source/Runtime/Engine/Private/AI/Navigation/NavigationTestingActor.cpp
--- a/Engine/Source/Runtime/Engine/Private/AI/Navigation/NavigationTestingActor.cpp
+++ b/Engine/Source/Runtime/Engine/Private/AI/Navigation/NavigationTestingActor.cpp
@@ -316,7 +316,7 @@ void ANavigationTestingActor::UpdatePathfinding()
 	}
 }
 
-void ANavigationTestingActor::SearchPathTo(class ANavigationTestingActor* Goal)
+void ANavigationTestingActor::SearchPathTo(ANavigationTestingActor* Goal)
 {
 #if WITH_EDITORONLY_DATA
 	if (EdRenderComp)
diff --git a/Engine/Source/Runtime/Engine/Private/AI/Navigation/PImplRecastNavMesh.cpp b/Engine/Source/Runtime/Engine/Private/AI/Navigation/PImplRecastNavMesh.cpp
--- a/Engine/Source/Runtime/Engine/Private/AI/Navigation/PImplRecastNavMesh.cpp
+++ b/Engine/Source/Runtime/Engine/Private/AI/Navigation/PImplRecastNavMesh.cpp
@@ -1685,7 +1685,7 @@ void FPImplRecastNavMesh::GetEdgesForPathCorridor(const TArray<NavNodeRef>* Path
 	GetEdgesForPathCorridorImpl(PathCorridor, PathCorridorEdges, NavQuery);
 }
 
-bool FPImplRecastNavMesh::FilterPolys(TArray<NavNodeRef>& PolyRefs, const class FRecastQueryFilter* Filter, const UObject* Owner) const
+bool FPImplRecastNavMesh::FilterPolys(TArray<NavNodeRef>& PolyRefs, const FRecastQueryFilter* Filter, const UObject* Owner) const
 {
 	if (Filter == NULL || DetourNavMesh == NULL)
 	{
diff --git a/Engine/Source/Runtime/Engine/Private/AI/Navigation/RecastNavMesh.cpp b/Engine/Source/Runtime/Engine/Private/AI/Navigation/RecastNavMesh.cpp
--- a/Engine/Source/Runtime/Engine/Private/AI/Navigation/RecastNavMesh.cpp
+++ b/Engine/Source/Runtime/Engine/Private/AI/Navigation/RecastNavMesh.cpp
@@ -993,7 +993,7 @@ NavNodeRef ARecastNavMesh::FindNearestPoly(FVector const& Loc, FVector const& Ex
 	return PolyRef;
 }
 
-void ARecastNavMesh::UpdateCustomLink(const class INavLinkCustomInterface* CustomLink)
+void ARecastNavMesh::UpdateCustomLink(const INavLinkCustomInterface* CustomLink)
 {
 	TSubclassOf<UNavArea> AreaClass = CustomLink->GetLinkAreaClass();
 	const int32 UserId = CustomLink->GetLinkId();
@@ -1007,7 +1007,7 @@ void ARecastNavMesh::UpdateCustomLink(const class INavLinkCustomInterface* Custo
 	}
 }
 
-void ARecastNavMesh::UpdateNavigationLinkArea(int32 UserId, TSubclassOf<class UNavArea> AreaClass) const
+void ARecastNavMesh::UpdateNavigationLinkArea(int32 UserId, TSubclassOf<UNavArea> AreaClass) const
 {
 	int32 AreaId = GetAreaID(AreaClass);
 	if (AreaId >= 0 && RecastNavMeshImpl)
@@ -1018,7 +1018,7 @@ void ARecastNavMesh::UpdateNavigationLinkArea(int32 UserId, TSubclassOf<class UN
 	}
 }
 
-void ARecastNavMesh::UpdateSegmentLinkArea(int32 UserId, TSubclassOf<class UNavArea> AreaClass) const
+void ARecastNavMesh::UpdateSegmentLinkArea(int32 UserId, TSubclassOf<UNavArea> AreaClass) const
 {
 	int32 AreaId = GetAreaID(AreaClass);
 	if (AreaId >= 0 && RecastNavMeshImpl)
@@ -1345,7 +1345,7 @@ void ARecastNavMesh::SetDefaultForbiddenFlags(uint16 ForbiddenAreaFlags)
 	FPImplRecastNavMesh::SetFilterForbiddenFlags((FRecastQueryFilter*)DefaultQueryFilter->GetImplementation(), ForbiddenAreaFlags);
 }
 
-bool ARecastNavMesh::FilterPolys(TArray<NavNodeRef>& PolyRefs, const class FRecastQueryFilter* Filter, const UObject* QueryOwner) const
+bool ARecastNavMesh::FilterPolys(TArray<NavNodeRef>& PolyRefs, const FRecastQueryFilter* Filter, const UObject* QueryOwner) const
 {
 	bool bSuccess = false;
 	if (RecastNavMeshImpl)
@@ -1676,7 +1676,7 @@ void ARecastNavMesh::UpdateNavVersion()
 
 #if WITH_EDITOR
 
-void ARecastNavMesh::PostEditChangeProperty( struct FPropertyChangedEvent& PropertyChangedEvent)
+void ARecastNavMesh::PostEditChangeProperty(FPropertyChangedEvent& PropertyChangedEvent)
 {
 	static const FName NAME_Generation = FName(TEXT("Generation"));
 	static const FName NAME_Display = FName(TEXT("Display"));
@@ -1756,7 +1756,7 @@ bool ARecastNavMesh::IsVoxelCacheEnabled()
 	return DefOb && DefOb->bUseVoxelCache;
 }
 
-const class FRecastQueryFilter* ARecastNavMesh::GetNamedFilter(ERecastNamedFilter::Type FilterType)
+const FRecastQueryFilter* ARecastNavMesh::GetNamedFilter(ERecastNamedFilter::Type FilterType)
 {
 	check(FilterType < ERecastNamedFilter::NamedFiltersCount); 
 	return NamedFilters[FilterType];
diff --git a/Engine/Source/Runtime/Engine/Private/AI/Navigation/RecastNavMeshGenerator.cpp b/Engine/Source/Runtime/Engine/Private/AI/Navigation/RecastNavMeshGenerator.cpp
--- a/Engine/Source/Runtime/Engine/Private/AI/Navigation/RecastNavMeshGenerator.cpp
+++ b/Engine/Source/Runtime/Engine/Private/AI/Navigation/RecastNavMeshGenerator.cpp
@@ -52,7 +52,7 @@ FORCEINLINE bool DoesBoxContainBox(const FBox& BigBox, const FBox& SmallBox)
 	return DoesBoxContainOrOverlapVector(BigBox, SmallBox.Min) && DoesBoxContainOrOverlapVector(BigBox, SmallBox.Max);
 }
 
-int32 GetTilesCountHelper(const class dtNavMesh* DetourMesh)
+int32 GetTilesCountHelper(const dtNavMesh* DetourMesh)
 {
 	int32 NumTiles = 0;
 	if (DetourMesh)
@@ -191,9 +191,9 @@ struct FRecastGeometryExport : public FNavigableGeometryExport
 	FWalkableSlopeOverride SlopeOverride;
 
 #if WITH_PHYSX
-	virtual void ExportPxTriMesh16Bit(class physx::PxTriangleMesh const * const TriMesh, const FTransform& LocalToWorld) override;
-	virtual void ExportPxTriMesh32Bit(class physx::PxTriangleMesh const * const TriMesh, const FTransform& LocalToWorld) override;
-	virtual void ExportPxConvexMesh(class physx::PxConvexMesh const * const ConvexMesh, const FTransform& LocalToWorld) override;
+	virtual void ExportPxTriMesh16Bit(physx::PxTriangleMesh const * const TriMesh, const FTransform& LocalToWorld) override;
+	virtual void ExportPxTriMesh32Bit(physx::PxTriangleMesh const * const TriMesh, const FTransform& LocalToWorld) override;
+	virtual void ExportPxConvexMesh(physx::PxConvexMesh const * const ConvexMesh, const FTransform& LocalToWorld) override;
 	virtual void ExportPxHeightField(physx::PxHeightField const * const HeightField, const FTransform& LocalToWorld) override;
 #endif // WITH_PHYSX
 	virtual void ExportCustomMesh(const FVector* InVertices, int32 NumVerts, const int32* InIndices, int32 NumIndices, const FTransform& LocalToWorld) override;
@@ -2709,7 +2709,7 @@ static int32 CaclulateMaxTilesCount(const TNavStatArray<FBox>& NavigableAreas, f
 // FRecastNavMeshGenerator
 //----------------------------------------------------------------------//
 
-FRecastNavMeshGenerator::FRecastNavMeshGenerator(class ARecastNavMesh* InDestNavMesh)
+FRecastNavMeshGenerator::FRecastNavMeshGenerator(ARecastNavMesh* InDestNavMesh)
 	: NumActiveTiles(0),
 	  MaxTileGeneratorTasks(1),
 	  AvgLayersPerTile(8.0f),
diff --git a/Engine/Source/Runtime/Engine/Public/AI/Navigation/PImplRecastNavMesh.h b/Engine/Source/Runtime/Engine/Public/AI/Navigation/PImplRecastNavMesh.h
--- a/Engine/Source/Runtime/Engine/Public/AI/Navigation/PImplRecastNavMesh.h
+++ b/Engine/Source/Runtime/Engine/Public/AI/Navigation/PImplRecastNavMesh.h
@@ -145,7 +145,7 @@ public:
 	bool FindStraightPath(const FVector& StartLoc, const FVector& EndLoc, const TArray<NavNodeRef>& PathCorridor, TArray<FNavPathPoint>& PathPoints, TArray<uint32>* CustomLinks = NULL) const;
 
 	/** Filters nav polys in PolyRefs with Filter */
-	bool FilterPolys(TArray<NavNodeRef>& PolyRefs, const class FRecastQueryFilter* Filter, const UObject* Owner) const;
+	bool FilterPolys(TArray<NavNodeRef>& PolyRefs, const FRecastQueryFilter* Filter, const UObject* Owner) const;
 
 	/** Get all polys from tile */
 	bool GetPolysInTile(int32 TileIndex, TArray<FNavPoly>& Polys) const;
diff --git a/Engine/Source/Runtime/Engine/Public/AI/NavigationModifier.h b/Engine/Source/Runtime/Engine/Public/AI/NavigationModifier.h
--- a/Engine/Source/Runtime/Engine/Public/AI/NavigationModifier.h
+++ b/Engine/Source/Runtime/Engine/Public/AI/NavigationModifier.h
@@ -3,6 +3,10 @@
 #pragma once 
 #include "AI/Navigation/NavLinkDefinition.h"
 
+class UNavArea;
+class UNavLinkDefinition;
+class UBrushComponent;
+
 struct ENGINE_API FNavigationModifier
 {
 	FNavigationModifier() : bHasMetaAreas(false) {}
@@ -63,21 +67,21 @@ struct ENGINE_API FAreaNavModifier : public FNavigationModifier
 	float FixedCost;
 
 	FAreaNavModifier() : Cost(0.0f), FixedCost(0.0f), AreaClass(NULL), ReplaceAreaClass(NULL), ShapeType(ENavigationShapeType::Unknown) {}
-	FAreaNavModifier(float Radius, float Height, const FTransform& LocalToWorld, const TSubclassOf<class UNavArea> AreaClass);
-	FAreaNavModifier(const FVector& Extent, const FTransform& LocalToWorld, const TSubclassOf<class UNavArea> AreaClass);
-	FAreaNavModifier(const FBox& Box, const FTransform& LocalToWorld, const TSubclassOf<class UNavArea> AreaClass);
-	FAreaNavModifier(const TArray<FVector>& Points, ENavigationCoordSystem::Type CoordType, const FTransform& LocalToWorld, const TSubclassOf<class UNavArea> AreaClass);
-	FAreaNavModifier(const TArray<FVector>& Points, const int32 FirstIndex, const int32 LastIndex, ENavigationCoordSystem::Type CoordType, const FTransform& LocalToWorld, const TSubclassOf<class UNavArea> AreaClass);
-	FAreaNavModifier(const class UBrushComponent* BrushComponent, const TSubclassOf<class UNavArea> AreaClass);
+	FAreaNavModifier(float Radius, float Height, const FTransform& LocalToWorld, const TSubclassOf<UNavArea> AreaClass);
+	FAreaNavModifier(const FVector& Extent, const FTransform& LocalToWorld, const TSubclassOf<UNavArea> AreaClass);
+	FAreaNavModifier(const FBox& Box, const FTransform& LocalToWorld, const TSubclassOf<UNavArea> AreaClass);
+	FAreaNavModifier(const TArray<FVector>& Points, ENavigationCoordSystem::Type CoordType, const FTransform& LocalToWorld, const TSubclassOf<UNavArea> AreaClass);
+	FAreaNavModifier(const TArray<FVector>& Points, const int32 FirstIndex, const int32 LastIndex, ENavigationCoordSystem::Type CoordType, const FTransform& LocalToWorld, const TSubclassOf<UNavArea> AreaClass);
+	FAreaNavModifier(const UBrushComponent* BrushComponent, const TSubclassOf<UNavArea> AreaClass);
 
 	FORCEINLINE const FBox& GetBounds() const { return Bounds; }
 	FORCEINLINE ENavigationShapeType::Type GetShapeType() const { return ShapeType; }
 	FORCEINLINE bool ShouldIncludeAgentHeight() const { return bIncludeAgentHeight; }
 	FORCEINLINE void SetIncludeAgentHeight(bool bInclude) { bIncludeAgentHeight = bInclude; }
-	FORCEINLINE const TSubclassOf<class UNavArea> GetAreaClass() const { return AreaClass; }
-	FORCEINLINE const TSubclassOf<class UNavArea> GetAreaClassToReplace() const { return ReplaceAreaClass; }
-	void SetAreaClass(const TSubclassOf<class UNavArea> AreaClass);
-	void SetAreaClassToReplace(const TSubclassOf<class UNavArea> AreaClass);
+	FORCEINLINE const TSubclassOf<UNavArea> GetAreaClass() const { return AreaClass; }
+	FORCEINLINE const TSubclassOf<UNavArea> GetAreaClassToReplace() const { return ReplaceAreaClass; }
+	void SetAreaClass(const TSubclassOf<UNavArea> AreaClass);
+	void SetAreaClassToReplace(const TSubclassOf<UNavArea> AreaClass);
 
 	void GetCylinder(FCylinderNavAreaData& Data) const;
 	void GetBox(FBoxNavAreaData& Data) const;
@@ -85,8 +89,8 @@ struct ENGINE_API FAreaNavModifier : public FNavigationModifier
 
 protected:
 	/** this should take a value of a game specific navigation modifier	*/
-	TSubclassOf<class UNavArea> AreaClass;
-	TSubclassOf<class UNavArea> ReplaceAreaClass;
+	TSubclassOf<UNavArea> AreaClass;
+	TSubclassOf<UNavArea> ReplaceAreaClass;
 	FBox Bounds;
 
 	TArray<FVector> Points;
@@ -95,7 +99,7 @@ protected:
 	/** if set, area shape will be extended by agent's height to cover area underneath like regular colliding geometry */
 	uint8 bIncludeAgentHeight : 1;
 
-	void Init(const TSubclassOf<class UNavArea> AreaClass);
+	void Init(const TSubclassOf<UNavArea> AreaClass);
 	void SetConvex(const TArray<FVector>& Points, const int32 FirstIndex, const int32 LastIndex, ENavigationCoordSystem::Type CoordType, const FTransform& LocalToWorld);
 	void SetBox(const FBox& Box, const FTransform& LocalToWorld);
 };
@@ -185,11 +189,11 @@ struct ENGINE_API FCustomLinkNavModifier : public FNavigationModifier
 	FTransform LocalToWorld;
 
 	FCustomLinkNavModifier() : LinkDefinitionClass(NULL) {}
-	void Set(TSubclassOf<class UNavLinkDefinition> LinkDefinitionClass, const FTransform& LocalToWorld);
-	FORCEINLINE const TSubclassOf<class UNavLinkDefinition> GetNavLinkClass() const { return LinkDefinitionClass; }
+	void Set(TSubclassOf<UNavLinkDefinition> LinkDefinitionClass, const FTransform& LocalToWorld);
+	FORCEINLINE const TSubclassOf<UNavLinkDefinition> GetNavLinkClass() const { return LinkDefinitionClass; }
 
 protected:
-	TSubclassOf<class UNavLinkDefinition> LinkDefinitionClass;
+	TSubclassOf<UNavLinkDefinition> LinkDefinitionClass;
 };
 
 struct ENGINE_API FCompositeNavModifier : public FNavigationModifier
diff --git a/Engine/Source/Runtime/Engine/Public/AI/NavigationSystemHelpers.h b/Engine/Source/Runtime/Engine/Public/AI/NavigationSystemHelpers.h
--- a/Engine/Source/Runtime/Engine/Public/AI/NavigationSystemHelpers.h
+++ b/Engine/Source/Runtime/Engine/Public/AI/NavigationSystemHelpers.h
@@ -6,6 +6,11 @@
 #include "AI/Navigation/NavLinkDefinition.h"
 #include "AI/Navigation/NavigationTypes.h"
 
+class UBodySetup;
+class UNavCollision;
+class AActor;
+struct FCompositeNavModifier;
+
 #if WITH_PHYSX
 namespace physx
 {
@@ -33,11 +38,11 @@ struct FNavigableGeometryExport
 
 namespace NavigationHelper
 {
-	void GatherCollision(class UBodySetup* RigidBody, TNavStatArray<FVector>& OutVertexBuffer, TNavStatArray<int32>& OutIndexBuffer, const FTransform& ComponentToWorld = FTransform::Identity);
-	void GatherCollision(class UBodySetup* RigidBody, class UNavCollision* NavCollision);
+	void GatherCollision(UBodySetup* RigidBody, TNavStatArray<FVector>& OutVertexBuffer, TNavStatArray<int32>& OutIndexBuffer, const FTransform& ComponentToWorld = FTransform::Identity);
+	void GatherCollision(UBodySetup* RigidBody, UNavCollision* NavCollision);
 
-	DECLARE_DELEGATE_ThreeParams(FNavLinkProcessorDelegate, struct FCompositeNavModifier*, const class AActor*, const TArray<FNavigationLink>&);
-	DECLARE_DELEGATE_ThreeParams(FNavLinkSegmentProcessorDelegate, struct FCompositeNavModifier*, const class AActor*, const TArray<FNavigationSegmentLink>&);
+	DECLARE_DELEGATE_ThreeParams(FNavLinkProcessorDelegate, FCompositeNavModifier*, const AActor*, const TArray<FNavigationLink>&);
+	DECLARE_DELEGATE_ThreeParams(FNavLinkSegmentProcessorDelegate, FCompositeNavModifier*, const AActor*, const TArray<FNavigationSegmentLink>&);
 
 	/** Set new implementation of nav link processor, a function that will be
 	 *	be used to process/transform links before adding them to CompositeModifier.
@@ -47,12 +52,12 @@ namespace NavigationHelper
 	ENGINE_API void SetNavLinkSegmentProcessorDelegate(const FNavLinkSegmentProcessorDelegate& NewDelegate);
 
 	/** called to do any necessary processing on NavLinks and put results in CompositeModifier */
-	ENGINE_API void ProcessNavLinkAndAppend(struct FCompositeNavModifier* OUT CompositeModifier, const class AActor* Actor, const TArray<FNavigationLink>& IN NavLinks);
+	ENGINE_API void ProcessNavLinkAndAppend(FCompositeNavModifier* OUT CompositeModifier, const AActor* Actor, const TArray<FNavigationLink>& IN NavLinks);
 
 	/** called to do any necessary processing on NavLinks and put results in CompositeModifier */
-	ENGINE_API void ProcessNavLinkSegmentAndAppend(struct FCompositeNavModifier* OUT CompositeModifier, const class AActor* Actor, const TArray<FNavigationSegmentLink>& IN NavLinks);
+	ENGINE_API void ProcessNavLinkSegmentAndAppend(FCompositeNavModifier* OUT CompositeModifier, const AActor* Actor, const TArray<FNavigationSegmentLink>& IN NavLinks);
 
-	ENGINE_API void DefaultNavLinkProcessorImpl(struct FCompositeNavModifier* OUT CompositeModifier, const class AActor* Actor, const TArray<FNavigationLink>& IN NavLinks);
+	ENGINE_API void DefaultNavLinkProcessorImpl(FCompositeNavModifier* OUT CompositeModifier, const AActor* Actor, const TArray<FNavigationLink>& IN NavLinks);
 
-	ENGINE_API void DefaultNavLinkSegmentProcessorImpl(struct FCompositeNavModifier* OUT CompositeModifier, const class AActor* Actor, const TArray<FNavigationSegmentLink>& IN NavLinks);
+	ENGINE_API void DefaultNavLinkSegmentProcessorImpl(FCompositeNavModifier* OUT CompositeModifier, const AActor* Actor, const TArray<FNavigationSegmentLink>& IN NavLinks);
 }
