Fixed AISystem trying to load classes with empty FStringClassReferences #UE4

- addressed UE-5620

diff --git a/Engine/Source/Runtime/AIModule/Private/AISystem.cpp b/Engine/Source/Runtime/AIModule/Private/AISystem.cpp
--- a/Engine/Source/Runtime/AIModule/Private/AISystem.cpp
+++ b/Engine/Source/Runtime/AIModule/Private/AISystem.cpp
@@ -36,13 +36,13 @@ void UAISystem::PostInitProperties()
 		UWorld* WorldOuter = Cast<UWorld>(GetOuter());
 		UObject* ManagersOuter = WorldOuter != NULL ? (UObject*)WorldOuter : (UObject*)this;
 		
-		TSubclassOf<UAIHotSpotManager> HotSpotManagerClass = LoadClass<UAIHotSpotManager>(NULL, *HotSpotManagerClassName.ToString(), NULL, LOAD_None, NULL);
+		TSubclassOf<UAIHotSpotManager> HotSpotManagerClass = HotSpotManagerClassName.IsValid() ? LoadClass<UAIHotSpotManager>(NULL, *HotSpotManagerClassName.ToString(), NULL, LOAD_None, NULL) : nullptr;
 		if (HotSpotManagerClass)
 		{
 			HotSpotManager = NewObject<UAIHotSpotManager>(ManagersOuter, HotSpotManagerClass);
 		}
 
-		TSubclassOf<UAIPerceptionSystem> PerceptionSystemClass = LoadClass<UAIPerceptionSystem>(NULL, *PerceptionSystemClassName.ToString(), NULL, LOAD_None, NULL);
+		TSubclassOf<UAIPerceptionSystem> PerceptionSystemClass = PerceptionSystemClassName.IsValid() ? LoadClass<UAIPerceptionSystem>(NULL, *PerceptionSystemClassName.ToString(), NULL, LOAD_None, NULL) : nullptr;
 		if (PerceptionSystemClass)
 		{
 			PerceptionSystem = NewObject<UAIPerceptionSystem>(ManagersOuter, PerceptionSystemClass);
