More in-function type forward declarations removed from AI headers #UE4

diff --git a/Engine/Source/Runtime/AIModule/Classes/AIController.h b/Engine/Source/Runtime/AIModule/Classes/AIController.h
--- a/Engine/Source/Runtime/AIModule/Classes/AIController.h
+++ b/Engine/Source/Runtime/AIModule/Classes/AIController.h
@@ -19,6 +19,13 @@ class UAIPerceptionComponent;
 struct FBasedPosition;
 class UPawnAction;
 class UPawnActionsComponent;
+class UNavigationQueryFilter;
+class UBehaviorTree;
+class UBlackboardData;
+class UCanvas;
+#if ENABLE_VISUAL_LOG
+struct FVisualLogEntry;
+#endif // ENABLE_VISUAL_LOG
 
 DECLARE_DYNAMIC_MULTICAST_DELEGATE_TwoParams(FAIMoveCompletedSignature, FAIRequestID, RequestID, EPathFollowingResult::Type, Result);
 
@@ -51,7 +58,7 @@ struct FFocusKnowledge
 	{
 		Priorities.Reserve(6);
 	}
-	TArray<struct FFocusItem> Priorities;
+	TArray<FFocusItem> Priorities;
 };
 
 //=============================================================================
@@ -128,7 +135,7 @@ public:
 	 *	@note AcceptanceRadius has default value or -1 due to Header Parser not being able to recognize UPathFollowingComponent::DefaultAcceptanceRadius
 	 */
 	UFUNCTION(BlueprintCallable, Category="AI|Navigation")
-	EPathFollowingRequestResult::Type MoveToActor(AActor* Goal, float AcceptanceRadius = -1, bool bStopOnOverlap = true, bool bUsePathfinding = true, bool bCanStrafe = true, TSubclassOf<class UNavigationQueryFilter> FilterClass = NULL);
+	EPathFollowingRequestResult::Type MoveToActor(AActor* Goal, float AcceptanceRadius = -1, bool bStopOnOverlap = true, bool bUsePathfinding = true, bool bCanStrafe = true, TSubclassOf<UNavigationQueryFilter> FilterClass = NULL);
 	// @todo: above should be: 
 	// EPathFollowingRequestResult::Type MoveToActor(AActor* Goal, float AcceptanceRadius = /*UPathFollowingComponent::DefaultAcceptanceRadius==*/-1, bool bStopOnOverlap = true, bool bUsePathfinding = true, bool bCanStrafe = true);
 	// but parser doesn't like this (it's a bug, when fixed this will be changed)
@@ -141,7 +148,7 @@ public:
 	 *	@note AcceptanceRadius has default value or -1 due to Header Parser not being able to recognize UPathFollowingComponent::DefaultAcceptanceRadius
 	 */
 	UFUNCTION(BlueprintCallable, Category="AI|Navigation")
-	EPathFollowingRequestResult::Type MoveToLocation(const FVector& Dest, float AcceptanceRadius = -1, bool bStopOnOverlap = true, bool bUsePathfinding = true, bool bProjectDestinationToNavigation = false, bool bCanStrafe = true, TSubclassOf<class UNavigationQueryFilter> FilterClass = NULL);
+	EPathFollowingRequestResult::Type MoveToLocation(const FVector& Dest, float AcceptanceRadius = -1, bool bStopOnOverlap = true, bool bUsePathfinding = true, bool bProjectDestinationToNavigation = false, bool bCanStrafe = true, TSubclassOf<UNavigationQueryFilter> FilterClass = NULL);
 	// @todo: above should be: 
 	// EPathFollowingRequestResult::Type MoveToLocation(const FVector& Dest, float AcceptanceRadius = /*UPathFollowingComponent::DefaultAcceptanceRadius==*/-1, bool bStopOnOverlap = true, bool bUsePathfinding = true, bool bCanStrafe = true);
 	// but parser doesn't like this (it's a bug, when fixed this will be changed)
@@ -201,11 +208,11 @@ public:
 
 	/** Starts executing behavior tree. */
 	UFUNCTION(BlueprintCallable, Category="AI")
-	virtual bool RunBehaviorTree(class UBehaviorTree* BTAsset);
+	virtual bool RunBehaviorTree(UBehaviorTree* BTAsset);
 
 	/** makes AI use specified BB asset */
 	UFUNCTION(BlueprintCallable, Category = "AI")
-	virtual bool UseBlackboard(class UBlackboardData* BlackboardAsset);
+	virtual bool UseBlackboard(UBlackboardData* BlackboardAsset);
 
 	/** Retrieve the final position that controller should be looking at. */
 	UFUNCTION(BlueprintCallable, Category="AI")
@@ -254,12 +261,12 @@ public:
 	// End AActor Interface
 
 	// Begin AController Interface
-	virtual void Possess(class APawn* InPawn) override;
+	virtual void Possess(APawn* InPawn) override;
 	virtual void UnPossess() override;
-	virtual void DisplayDebug(class UCanvas* Canvas, const FDebugDisplayInfo& DebugDisplay, float& YL, float& YPos) override;
+	virtual void DisplayDebug(UCanvas* Canvas, const FDebugDisplayInfo& DebugDisplay, float& YL, float& YPos) override;
 
 #if ENABLE_VISUAL_LOG
-	virtual void GrabDebugSnapshot(struct FVisualLogEntry* Snapshot) const override;
+	virtual void GrabDebugSnapshot(FVisualLogEntry* Snapshot) const override;
 #endif
 
 	virtual void Reset() override;
@@ -294,7 +301,7 @@ public:
 	//----------------------------------------------------------------------//
 	// IAIPerceptionListenerInterface
 	//----------------------------------------------------------------------//
-	virtual class UAIPerceptionComponent* GetPerceptionComponent() override { return PerceptionComponent; }
+	virtual UAIPerceptionComponent* GetPerceptionComponent() override { return PerceptionComponent; }
 
 	//----------------------------------------------------------------------//
 	// Actions
diff --git a/Engine/Source/Runtime/AIModule/Classes/EnvironmentQuery/Contexts/EnvQueryContext_BlueprintBase.h b/Engine/Source/Runtime/AIModule/Classes/EnvironmentQuery/Contexts/EnvQueryContext_BlueprintBase.h
--- a/Engine/Source/Runtime/AIModule/Classes/EnvironmentQuery/Contexts/EnvQueryContext_BlueprintBase.h
+++ b/Engine/Source/Runtime/AIModule/Classes/EnvironmentQuery/Contexts/EnvQueryContext_BlueprintBase.h
@@ -5,6 +5,10 @@
 #include "EnvironmentQuery/EnvQueryContext.h"
 #include "EnvQueryContext_BlueprintBase.generated.h"
 
+class UWorld;
+struct FEnvQueryInstance;
+struct FEnvQueryContextData;
+
 UCLASS(MinimalAPI, Blueprintable)
 class UEnvQueryContext_BlueprintBase : public UEnvQueryContext
 {
@@ -22,9 +26,9 @@ class UEnvQueryContext_BlueprintBase : public UEnvQueryContext
 	ECallMode CallMode;
 
 	// We need to implement GetWorld() so that blueprint functions which use a hidden WorldContextObject* will work properly.
-	virtual class UWorld* GetWorld() const;
+	virtual UWorld* GetWorld() const;
 
-	virtual void ProvideContext(struct FEnvQueryInstance& QueryInstance, struct FEnvQueryContextData& ContextData) const override;
+	virtual void ProvideContext(FEnvQueryInstance& QueryInstance, FEnvQueryContextData& ContextData) const override;
 
 	UFUNCTION(BlueprintImplementableEvent)
 	virtual void ProvideSingleActor(AActor* QuerierActor, AActor*& ResultingActor) const;
diff --git a/Engine/Source/Runtime/AIModule/Classes/EnvironmentQuery/Contexts/EnvQueryContext_Querier.h b/Engine/Source/Runtime/AIModule/Classes/EnvironmentQuery/Contexts/EnvQueryContext_Querier.h
--- a/Engine/Source/Runtime/AIModule/Classes/EnvironmentQuery/Contexts/EnvQueryContext_Querier.h
+++ b/Engine/Source/Runtime/AIModule/Classes/EnvironmentQuery/Contexts/EnvQueryContext_Querier.h
@@ -4,10 +4,13 @@
 #include "EnvironmentQuery/EnvQueryContext.h"
 #include "EnvQueryContext_Querier.generated.h"
 
+struct FEnvQueryInstance;
+struct FEnvQueryContextData;
+
 UCLASS(MinimalAPI)
 class UEnvQueryContext_Querier : public UEnvQueryContext
 {
 	GENERATED_UCLASS_BODY()
  
-	virtual void ProvideContext(struct FEnvQueryInstance& QueryInstance, struct FEnvQueryContextData& ContextData) const override;
+	virtual void ProvideContext(FEnvQueryInstance& QueryInstance, FEnvQueryContextData& ContextData) const override;
 };
diff --git a/Engine/Source/Runtime/AIModule/Classes/EnvironmentQuery/EQSRenderingComponent.h b/Engine/Source/Runtime/AIModule/Classes/EnvironmentQuery/EQSRenderingComponent.h
--- a/Engine/Source/Runtime/AIModule/Classes/EnvironmentQuery/EQSRenderingComponent.h
+++ b/Engine/Source/Runtime/AIModule/Classes/EnvironmentQuery/EQSRenderingComponent.h
@@ -19,7 +19,7 @@ public:
 
 #if  USE_EQS_DEBUGGER 
 	static void CollectEQSData(const UPrimitiveComponent* InComponent, const class IEQSQueryResultSourceInterface* QueryDataSource, TArray<FSphere>& Spheres, TArray<FText3d>& Texts, TArray<EQSDebug::FDebugHelper>& DebugItems);
-	static void CollectEQSData(const struct FEnvQueryResult* ResultItems, const struct FEnvQueryInstance* QueryInstance, TArray<FSphere>& Spheres, TArray<FText3d>& Texts, bool ShouldDrawFailedItems, TArray<EQSDebug::FDebugHelper>& DebugItems);
+	static void CollectEQSData(const FEnvQueryResult* ResultItems, const struct FEnvQueryInstance* QueryInstance, TArray<FSphere>& Spheres, TArray<FText3d>& Texts, bool ShouldDrawFailedItems, TArray<EQSDebug::FDebugHelper>& DebugItems);
 #endif
 private:
 	FEnvQueryResult QueryResult;	
diff --git a/Engine/Source/Runtime/AIModule/Classes/EnvironmentQuery/EQSTestingPawn.h b/Engine/Source/Runtime/AIModule/Classes/EnvironmentQuery/EQSTestingPawn.h
--- a/Engine/Source/Runtime/AIModule/Classes/EnvironmentQuery/EQSTestingPawn.h
+++ b/Engine/Source/Runtime/AIModule/Classes/EnvironmentQuery/EQSTestingPawn.h
@@ -7,20 +7,23 @@
 #include "EnvironmentQuery/EnvQueryTypes.h"
 #include "EQSTestingPawn.generated.h"
 
+class UEnvQuery;
+class UEQSRenderingComponent;
+
 /** this class is abstract even though it's perfectly functional on its own.
  *	The reason is to stop it from showing as valid player pawn type when configuring 
  *	project's game mode. */
-UCLASS(abstract, hidecategories=(Advanced, Attachment, Collision, Mesh, Animation, Clothing, Physics, Rendering, Lighting, Activation, CharacterMovement, AgentPhysics, Avoidance, MovementComponent, Velocity, Shape, Camera, Input, Layers, SkeletalMesh, Optimization, Pawn))
+UCLASS(abstract, hidecategories=(Advanced, Attachment, Collision, Mesh, Animation, Clothing, Physics, Rendering, Lighting, Activation, CharacterMovement, AgentPhysics, Avoidance, MovementComponent, Velocity, Shape, Camera, Input, Layers, SkeletalMesh, Optimization, Pawn, Replication, Actor))
 class AIMODULE_API AEQSTestingPawn : public ACharacter, public IEQSQueryResultSourceInterface
 {
 	GENERATED_UCLASS_BODY()
 	
 	UPROPERTY(Category=EQS, EditAnywhere)
-	class UEnvQuery* QueryTemplate;
+	UEnvQuery* QueryTemplate;
 
 	/** optional parameters for query */
 	UPROPERTY(Category=EQS, EditAnywhere)
-	TArray<struct FEnvNamedValue> QueryParams;
+	TArray<FEnvNamedValue> QueryParams;
 
 	UPROPERTY(Category=EQS, EditAnywhere)
 	float TimeLimitPerStep;
@@ -48,12 +51,12 @@ class AIMODULE_API AEQSTestingPawn : public ACharacter, public IEQSQueryResultSo
 private_subobject:
 	/** Editor Preview */
 	DEPRECATED_FORGAME(4.6, "EdRenderComp should not be accessed directly, please use GetEdRenderComp() function instead. EdRenderComp will soon be private and your code will not compile.")
-	UPROPERTY()
-	class UEQSRenderingComponent* EdRenderComp;
+	UPROPERTY(Transient)
+	UEQSRenderingComponent* EdRenderComp;
 #endif // WITH_EDITORONLY_DATA
 
 protected:
-	TSharedPtr<struct FEnvQueryInstance> QueryInstance;
+	TSharedPtr<FEnvQueryInstance> QueryInstance;
 
 	TArray<FEnvQueryInstance> StepResults;
 
@@ -68,8 +71,8 @@ public:
 #endif // WITH_EDITOR
 
 	// IEQSQueryResultSourceInterface start
-	virtual const struct FEnvQueryResult* GetQueryResult() const override;
-	virtual const struct FEnvQueryInstance* GetQueryInstance() const  override;
+	virtual const FEnvQueryResult* GetQueryResult() const override;
+	virtual const FEnvQueryInstance* GetQueryInstance() const  override;
 
 	virtual bool GetShouldDebugDrawLabels() const override { return bDrawLabels; }
 	virtual bool GetShouldDrawFailedItems() const override{ return bDrawFailedItems; }
@@ -88,6 +91,6 @@ protected:
 public:
 #if WITH_EDITORONLY_DATA
 	/** Returns EdRenderComp subobject **/
-	class UEQSRenderingComponent* GetEdRenderComp();
+	UEQSRenderingComponent* GetEdRenderComp();
 #endif
 };
diff --git a/Engine/Source/Runtime/AIModule/Classes/EnvironmentQuery/EnvQuery.h b/Engine/Source/Runtime/AIModule/Classes/EnvironmentQuery/EnvQuery.h
--- a/Engine/Source/Runtime/AIModule/Classes/EnvironmentQuery/EnvQuery.h
+++ b/Engine/Source/Runtime/AIModule/Classes/EnvironmentQuery/EnvQuery.h
@@ -5,6 +5,11 @@
 #include "EnvQueryTypes.h"
 #include "EnvQuery.generated.h"
 
+class UEnvQueryOption;
+#if WITH_EDITORONLY_DATA
+class UEdGraph;
+#endif // WITH_EDITORONLY_DATA
+
 UCLASS()
 class AIMODULE_API UEnvQuery : public UObject
 {
@@ -13,12 +18,12 @@ class AIMODULE_API UEnvQuery : public UObject
 #if WITH_EDITORONLY_DATA
 	/** Graph for query */
 	UPROPERTY()
-	class UEdGraph*	EdGraph;
+	UEdGraph*	EdGraph;
 #endif
 
 	UPROPERTY()
-	TArray<class UEnvQueryOption*> Options;
+	TArray<UEnvQueryOption*> Options;
 
 	/** Gather all required named params */
-	void CollectQueryParams(TArray<struct FEnvNamedValue>& NamedValues) const;
+	void CollectQueryParams(TArray<FEnvNamedValue>& NamedValues) const;
 };
diff --git a/Engine/Source/Runtime/AIModule/Classes/EnvironmentQuery/EnvQueryContext.h b/Engine/Source/Runtime/AIModule/Classes/EnvironmentQuery/EnvQueryContext.h
--- a/Engine/Source/Runtime/AIModule/Classes/EnvironmentQuery/EnvQueryContext.h
+++ b/Engine/Source/Runtime/AIModule/Classes/EnvironmentQuery/EnvQueryContext.h
@@ -3,10 +3,13 @@
 #pragma once
 #include "EnvQueryContext.generated.h"
 
+struct FEnvQueryContextData;
+struct FEnvQueryInstance;
+
 UCLASS(Abstract)
 class AIMODULE_API UEnvQueryContext : public UObject
 {
 	GENERATED_UCLASS_BODY()
 
-	virtual void ProvideContext(struct FEnvQueryInstance& QueryInstance, struct FEnvQueryContextData& ContextData) const;
+	virtual void ProvideContext(FEnvQueryInstance& QueryInstance, FEnvQueryContextData& ContextData) const;
 };
diff --git a/Engine/Source/Runtime/AIModule/Classes/EnvironmentQuery/EnvQueryDebugHelpers.h b/Engine/Source/Runtime/AIModule/Classes/EnvironmentQuery/EnvQueryDebugHelpers.h
--- a/Engine/Source/Runtime/AIModule/Classes/EnvironmentQuery/EnvQueryDebugHelpers.h
+++ b/Engine/Source/Runtime/AIModule/Classes/EnvironmentQuery/EnvQueryDebugHelpers.h
@@ -6,8 +6,12 @@
 #include "VisualLoggerExtension.h"
 #include "EnvQueryDebugHelpers.generated.h"
 
+struct FEnvQueryInstance;
 
 #if  USE_EQS_DEBUGGER || ENABLE_VISUAL_LOG
+
+struct FLogCategoryBase;
+
 namespace EQSDebug
 {
 	struct FItemData
@@ -143,28 +147,27 @@ FArchive& operator<<(FArchive& Ar, EQSDebug::FQueryData& Data)
 #	define UE_VLOG_EQS(Query, CategoryName, Verbosity)
 #endif //ENABLE_VISUAL_LOG && USE_EQS_DEBUGGER
 
-UCLASS(Abstract, CustomConstructor)
+UCLASS(Abstract)
 class AIMODULE_API UEnvQueryDebugHelpers : public UObject
 {
-	GENERATED_UCLASS_BODY()
-
-	UEnvQueryDebugHelpers(const FObjectInitializer& ObjectInitializer) : Super(ObjectInitializer) {}
+	GENERATED_BODY()
+public:
 #if USE_EQS_DEBUGGER
-	static void QueryToDebugData(struct FEnvQueryInstance& Query, EQSDebug::FQueryData& EQSLocalData);
-	static void QueryToBlobArray(struct FEnvQueryInstance& Query, TArray<uint8>& BlobArray, bool bUseCompression = false);
+	static void QueryToDebugData(FEnvQueryInstance& Query, EQSDebug::FQueryData& EQSLocalData);
+	static void QueryToBlobArray(FEnvQueryInstance& Query, TArray<uint8>& BlobArray, bool bUseCompression = false);
 	static void BlobArrayToDebugData(const TArray<uint8>& BlobArray, EQSDebug::FQueryData& EQSLocalData, bool bUseCompression = false);
 #endif
 
 #if ENABLE_VISUAL_LOG && USE_EQS_DEBUGGER
-	static void LogQuery(struct FEnvQueryInstance& Query, const struct FLogCategoryBase& Category, ELogVerbosity::Type Verbosity);
+	static void LogQuery(FEnvQueryInstance& Query, const FLogCategoryBase& Category, ELogVerbosity::Type Verbosity);
 
 private:
-	static void LogQueryInternal(struct FEnvQueryInstance& Query, const struct FLogCategoryBase& Category, ELogVerbosity::Type Verbosity, float TimeSeconds, FVisualLogEntry *CurrentEntry);
+	static void LogQueryInternal(FEnvQueryInstance& Query, const FLogCategoryBase& Category, ELogVerbosity::Type Verbosity, float TimeSeconds, FVisualLogEntry *CurrentEntry);
 #endif
 };
 
 #if ENABLE_VISUAL_LOG && USE_EQS_DEBUGGER
-FORCEINLINE void UEnvQueryDebugHelpers::LogQuery(struct FEnvQueryInstance& Query, const struct FLogCategoryBase& Category, ELogVerbosity::Type Verbosity)
+FORCEINLINE void UEnvQueryDebugHelpers::LogQuery(FEnvQueryInstance& Query, const FLogCategoryBase& Category, ELogVerbosity::Type Verbosity)
 {
 	UWorld *World = NULL;
 	FVisualLogEntry *CurrentEntry = NULL;
diff --git a/Engine/Source/Runtime/AIModule/Classes/EnvironmentQuery/EnvQueryGenerator.h b/Engine/Source/Runtime/AIModule/Classes/EnvironmentQuery/EnvQueryGenerator.h
--- a/Engine/Source/Runtime/AIModule/Classes/EnvironmentQuery/EnvQueryGenerator.h
+++ b/Engine/Source/Runtime/AIModule/Classes/EnvironmentQuery/EnvQueryGenerator.h
@@ -4,6 +4,8 @@
 #include "EnvironmentQuery/EnvQueryTypes.h"
 #include "EnvQueryGenerator.generated.h"
 
+class UEnvQueryItemType;
+
 UCLASS(Abstract)
 class AIMODULE_API UEnvQueryGenerator : public UObject
 {
@@ -14,7 +16,7 @@ class AIMODULE_API UEnvQueryGenerator : public UObject
 
 	/** type of generated items */
 	UPROPERTY()
-	TSubclassOf<class UEnvQueryItemType> ItemType;
+	TSubclassOf<UEnvQueryItemType> ItemType;
 
 	virtual void GenerateItems(FEnvQueryInstance& QueryInstance) const { checkNoEntry(); }
 
diff --git a/Engine/Source/Runtime/AIModule/Classes/EnvironmentQuery/EnvQueryManager.h b/Engine/Source/Runtime/AIModule/Classes/EnvironmentQuery/EnvQueryManager.h
--- a/Engine/Source/Runtime/AIModule/Classes/EnvironmentQuery/EnvQueryManager.h
+++ b/Engine/Source/Runtime/AIModule/Classes/EnvironmentQuery/EnvQueryManager.h
@@ -113,7 +113,7 @@ class AIMODULE_API UEnvQueryManager : public UObject, public FTickableGameObject
 
 	// We need to implement GetWorld() so that any EQS-related blueprints (such as blueprint contexts) can implement
 	// GetWorld() and so provide access to blueprint nodes using hidden WorldContextObject parameters.
-	virtual class UWorld* GetWorld() const;
+	virtual UWorld* GetWorld() const;
 
 	/** [FTickableGameObject] tick function */
 	virtual void Tick(float DeltaTime) override;
diff --git a/Engine/Source/Runtime/AIModule/Classes/EnvironmentQuery/EnvQueryTypes.h b/Engine/Source/Runtime/AIModule/Classes/EnvironmentQuery/EnvQueryTypes.h
--- a/Engine/Source/Runtime/AIModule/Classes/EnvironmentQuery/EnvQueryTypes.h
+++ b/Engine/Source/Runtime/AIModule/Classes/EnvironmentQuery/EnvQueryTypes.h
@@ -5,11 +5,16 @@
 #include "EnvironmentQuery/Items/EnvQueryItemType.h"
 #include "EnvQueryTypes.generated.h"
 
+class ARecastNavMesh;
+class UNavigationQueryFilter;
 class UEnvQueryTest;
 class UEnvQueryGenerator;
 class UEnvQueryItemType_VectorBase;
 class UEnvQueryItemType_ActorBase;
+class UEnvQueryContext;
 struct FEnvQueryInstance;
+struct FEnvQueryOptionInstance;
+struct FEnvQueryItemDetails;
 
 AIMODULE_API DECLARE_LOG_CATEGORY_EXTERN(LogEQS, Warning, All);
 
@@ -28,8 +33,6 @@ DECLARE_DWORD_COUNTER_STAT_EXTERN(TEXT("Num Instances"),STAT_AI_EQS_NumInstances
 DECLARE_DWORD_COUNTER_STAT_EXTERN(TEXT("Num Items"),STAT_AI_EQS_NumItems,STATGROUP_AI_EQS, );
 DECLARE_MEMORY_STAT_EXTERN(TEXT("Instance memory"),STAT_AI_EQS_InstanceMemory,STATGROUP_AI_EQS, AIMODULE_API);
 
-class ARecastNavMesh;
-
 UENUM()
 namespace EEnvTestPurpose
 {
@@ -274,9 +277,6 @@ struct AIMODULE_API FEnvNamedValue
 
 	UPROPERTY(EditAnywhere, BlueprintReadWrite, Category = Param)
 	float Value;
-
-	UPROPERTY(EditAnywhere, BlueprintReadWrite, Category = Param)
-	class UAISenseConfig* SenseConfig;
 };
 
 USTRUCT()
@@ -286,15 +286,15 @@ struct AIMODULE_API FEnvDirection
 
 	/** line A: start context */
 	UPROPERTY(EditDefaultsOnly, Category=Direction)
-	TSubclassOf<class UEnvQueryContext> LineFrom;
+	TSubclassOf<UEnvQueryContext> LineFrom;
 
 	/** line A: finish context */
 	UPROPERTY(EditDefaultsOnly, Category=Direction)
-	TSubclassOf<class UEnvQueryContext> LineTo;
+	TSubclassOf<UEnvQueryContext> LineTo;
 
 	/** line A: direction context */
 	UPROPERTY(EditDefaultsOnly, Category=Direction)
-	TSubclassOf<class UEnvQueryContext> Rotation;
+	TSubclassOf<UEnvQueryContext> Rotation;
 
 	/** defines direction of second line used by test */
 	UPROPERTY(EditDefaultsOnly, Category=Direction, meta=(DisplayName="Mode"))
@@ -322,7 +322,7 @@ struct AIMODULE_API FEnvTraceData
 
 	/** navigation filter for tracing */
 	UPROPERTY(EditDefaultsOnly, Category=Trace)
-	TSubclassOf<class UNavigationQueryFilter> NavigationFilter;
+	TSubclassOf<UNavigationQueryFilter> NavigationFilter;
 
 	/** search height: below point */
 	UPROPERTY(EditDefaultsOnly, Category=Trace, meta=(UIMin=0, ClampMin=0))
@@ -468,7 +468,7 @@ public:
 //////////////////////////////////////////////////////////////////////////
 // Runtime processing structures
 
-DECLARE_DELEGATE_OneParam(FQueryFinishedSignature, TSharedPtr<struct FEnvQueryResult>);
+DECLARE_DELEGATE_OneParam(FQueryFinishedSignature, TSharedPtr<FEnvQueryResult>);
 
 struct AIMODULE_API FEnvQuerySpatialData
 {
@@ -593,7 +593,7 @@ struct AIMODULE_API FEnvQueryInstance : public FEnvQueryResult
 	TMap<UClass*, FEnvQueryContextData> ContextCache;
 
 	/** list of options */
-	TArray<struct FEnvQueryOptionInstance> Options;
+	TArray<FEnvQueryOptionInstance> Options;
 
 	/** currently processed test (-1 = generator) */
 	int32 CurrentTest;
@@ -603,7 +603,7 @@ struct AIMODULE_API FEnvQueryInstance : public FEnvQueryResult
 	int32 CurrentTestStartingItem;
 
 	/** list of item details */
-	TArray<struct FEnvQueryItemDetails> ItemDetails;
+	TArray<FEnvQueryItemDetails> ItemDetails;
 
 	/** number of valid items on list */
 	int32 NumValidItems;
@@ -966,22 +966,22 @@ public:
 namespace FEQSHelpers
 {
 #if WITH_RECAST
-	const ARecastNavMesh* FindNavMeshForQuery(struct FEnvQueryInstance& QueryInstance);
+	const ARecastNavMesh* FindNavMeshForQuery(FEnvQueryInstance& QueryInstance);
 #endif // WITH_RECAST
 }
 
-UCLASS(Abstract, CustomConstructor)
+UCLASS(Abstract)
 class AIMODULE_API UEnvQueryTypes : public UObject
 {
-	GENERATED_UCLASS_BODY()
+	GENERATED_BODY()
 
+public:
 	/** special test value assigned to items skipped by condition check */
 	static float SkippedItemValue;
 
-	UEnvQueryTypes(const FObjectInitializer& ObjectInitializer) : Super(ObjectInitializer) {}
 	static FText GetShortTypeName(const UObject* Ob);
 
-	static FText DescribeContext(TSubclassOf<class UEnvQueryContext> ContextClass);
+	static FText DescribeContext(TSubclassOf<UEnvQueryContext> ContextClass);
 	static FString DescribeIntParam(const FEnvIntParam& Param);
 	static FString DescribeFloatParam(const FEnvFloatParam& Param);
 	static FString DescribeBoolParam(const FEnvBoolParam& Param);
diff --git a/Engine/Source/Runtime/AIModule/Classes/EnvironmentQuery/Generators/EnvQueryGenerator_ActorsOfClass.h b/Engine/Source/Runtime/AIModule/Classes/EnvironmentQuery/Generators/EnvQueryGenerator_ActorsOfClass.h
--- a/Engine/Source/Runtime/AIModule/Classes/EnvironmentQuery/Generators/EnvQueryGenerator_ActorsOfClass.h
+++ b/Engine/Source/Runtime/AIModule/Classes/EnvironmentQuery/Generators/EnvQueryGenerator_ActorsOfClass.h
@@ -5,6 +5,7 @@
 #include "EnvQueryGenerator_ActorsOfClass.generated.h"
 
 class AActor;
+class UEnvQueryContext;
 
 UCLASS()	
 class UEnvQueryGenerator_ActorsOfClass : public UEnvQueryGenerator
@@ -20,7 +21,7 @@ class UEnvQueryGenerator_ActorsOfClass : public UEnvQueryGenerator
 
 	/** context */
 	UPROPERTY(EditAnywhere, Category=Generator)
-	TSubclassOf<class UEnvQueryContext> SearchCenter;
+	TSubclassOf<UEnvQueryContext> SearchCenter;
 
 	virtual void GenerateItems(FEnvQueryInstance& QueryInstance) const override;
 
diff --git a/Engine/Source/Runtime/AIModule/Classes/EnvironmentQuery/Generators/EnvQueryGenerator_OnCircle.h b/Engine/Source/Runtime/AIModule/Classes/EnvironmentQuery/Generators/EnvQueryGenerator_OnCircle.h
--- a/Engine/Source/Runtime/AIModule/Classes/EnvironmentQuery/Generators/EnvQueryGenerator_OnCircle.h
+++ b/Engine/Source/Runtime/AIModule/Classes/EnvironmentQuery/Generators/EnvQueryGenerator_OnCircle.h
@@ -4,6 +4,8 @@
 #include "EnvironmentQuery/Generators/EnvQueryGenerator_ProjectedPoints.h"
 #include "EnvQueryGenerator_OnCircle.generated.h"
 
+struct FEnvQueryInstance;
+
 class FVectorAndDataContainer
 {
 public:
@@ -73,7 +75,7 @@ class AIMODULE_API UEnvQueryGenerator_OnCircle : public UEnvQueryGenerator_Proje
 #endif // WITH_EDITOR
 
 protected:
-	FVector CalcDirection(struct FEnvQueryInstance& QueryInstance) const;
+	FVector CalcDirection(FEnvQueryInstance& QueryInstance) const;
 
 	// For derived classes that need to store additional data associated with the FVector, override these two functions:
 	virtual void PrepareGeneratorContext(FEnvQueryInstance& QueryInstance,
diff --git a/Engine/Source/Runtime/AIModule/Classes/EnvironmentQuery/Generators/EnvQueryGenerator_PathingGrid.h b/Engine/Source/Runtime/AIModule/Classes/EnvironmentQuery/Generators/EnvQueryGenerator_PathingGrid.h
--- a/Engine/Source/Runtime/AIModule/Classes/EnvironmentQuery/Generators/EnvQueryGenerator_PathingGrid.h
+++ b/Engine/Source/Runtime/AIModule/Classes/EnvironmentQuery/Generators/EnvQueryGenerator_PathingGrid.h
@@ -4,6 +4,13 @@
 #include "EnvironmentQuery/EnvQueryGenerator.h"
 #include "EnvQueryGenerator_PathingGrid.generated.h"
 
+class UEnvQueryContext;
+class UNavigationQueryFilter;
+struct FNavLocation;
+#if WITH_RECAST
+class ARecastNavMesh;
+#endif // WITH_RECAST
+
 /**
  *  Navigation grid, generates points on navmesh
  *  with paths to/from context no further than given limit
@@ -28,7 +35,7 @@ class UEnvQueryGenerator_PathingGrid : public UEnvQueryGenerator
 
 	/** context */
 	UPROPERTY(EditAnywhere, Category=Generator)
-	TSubclassOf<class UEnvQueryContext> GenerateAround;
+	TSubclassOf<UEnvQueryContext> GenerateAround;
 
 	/** context */
 	UPROPERTY(EditAnywhere, Category=Generator)
@@ -36,7 +43,7 @@ class UEnvQueryGenerator_PathingGrid : public UEnvQueryGenerator
 
 	/** navigation filter for tracing */
 	UPROPERTY(EditDefaultsOnly, Category=Generator)
-	TSubclassOf<class UNavigationQueryFilter> NavigationFilter;
+	TSubclassOf<UNavigationQueryFilter> NavigationFilter;
 
 	virtual void GenerateItems(FEnvQueryInstance& QueryInstance) const override;
 
@@ -45,11 +52,10 @@ class UEnvQueryGenerator_PathingGrid : public UEnvQueryGenerator
 
 protected:
 #if WITH_RECAST
-
 	/** find all navmesh node refs in pathing distance */
-	void FindNodeRefsInPathDistance(const class ARecastNavMesh* NavMesh, const FVector& ContextLocation, float MaxPathDistance, bool bPathFromContext, TArray<NavNodeRef>& NodeRefs, FBox& NodeRefsBounds) const;
+	void FindNodeRefsInPathDistance(const ARecastNavMesh* NavMesh, const FVector& ContextLocation, float MaxPathDistance, bool bPathFromContext, TArray<NavNodeRef>& NodeRefs, FBox& NodeRefsBounds) const;
 
 	/** check if nav location is in allowed set */
-	bool IsNavLocationInPathDistance(const class ARecastNavMesh* NavMesh, const struct FNavLocation& NavLocation, const TArray<NavNodeRef>& NodeRefs) const;
+	bool IsNavLocationInPathDistance(const ARecastNavMesh* NavMesh, const FNavLocation& NavLocation, const TArray<NavNodeRef>& NodeRefs) const;
 #endif
 };
diff --git a/Engine/Source/Runtime/AIModule/Classes/EnvironmentQuery/Generators/EnvQueryGenerator_SimpleGrid.h b/Engine/Source/Runtime/AIModule/Classes/EnvironmentQuery/Generators/EnvQueryGenerator_SimpleGrid.h
--- a/Engine/Source/Runtime/AIModule/Classes/EnvironmentQuery/Generators/EnvQueryGenerator_SimpleGrid.h
+++ b/Engine/Source/Runtime/AIModule/Classes/EnvironmentQuery/Generators/EnvQueryGenerator_SimpleGrid.h
@@ -4,6 +4,8 @@
 #include "EnvironmentQuery/Generators/EnvQueryGenerator_ProjectedPoints.h"
 #include "EnvQueryGenerator_SimpleGrid.generated.h"
 
+class UEnvQueryContext;
+
 /**
  *  Simple grid, generates points in 2D square around context
  */
@@ -23,7 +25,7 @@ class UEnvQueryGenerator_SimpleGrid : public UEnvQueryGenerator_ProjectedPoints
 
 	/** context */
 	UPROPERTY(EditDefaultsOnly, Category=Generator)
-	TSubclassOf<class UEnvQueryContext> GenerateAround;
+	TSubclassOf<UEnvQueryContext> GenerateAround;
 
 	virtual void GenerateItems(FEnvQueryInstance& QueryInstance) const override;
 
diff --git a/Engine/Source/Runtime/AIModule/Classes/EnvironmentQuery/Items/EnvQueryItemType.h b/Engine/Source/Runtime/AIModule/Classes/EnvironmentQuery/Items/EnvQueryItemType.h
--- a/Engine/Source/Runtime/AIModule/Classes/EnvironmentQuery/Items/EnvQueryItemType.h
+++ b/Engine/Source/Runtime/AIModule/Classes/EnvironmentQuery/Items/EnvQueryItemType.h
@@ -3,6 +3,9 @@
 #pragma once
 #include "EnvQueryItemType.generated.h"
 
+class UBlackboardComponent;
+struct FBlackboardKeySelector;
+
 UCLASS(Abstract)
 class AIMODULE_API UEnvQueryItemType : public UObject
 {
@@ -12,10 +15,10 @@ class AIMODULE_API UEnvQueryItemType : public UObject
 	FORCEINLINE uint16 GetValueSize() const { return ValueSize; }
 
 	/** add filters for blackboard key selector */
-	virtual void AddBlackboardFilters(struct FBlackboardKeySelector& KeySelector, UObject* FilterOwner) const;
+	virtual void AddBlackboardFilters(FBlackboardKeySelector& KeySelector, UObject* FilterOwner) const;
 
 	/** store value in blackboard entry */
-	virtual bool StoreInBlackboard(struct FBlackboardKeySelector& KeySelector, class UBlackboardComponent* Blackboard, const uint8* RawData) const;
+	virtual bool StoreInBlackboard(FBlackboardKeySelector& KeySelector, UBlackboardComponent* Blackboard, const uint8* RawData) const;
 
 	/** unregister from known types in EnvQueryManager */
 	virtual void FinishDestroy() override;
diff --git a/Engine/Source/Runtime/AIModule/Classes/EnvironmentQuery/Items/EnvQueryItemType_Actor.h b/Engine/Source/Runtime/AIModule/Classes/EnvironmentQuery/Items/EnvQueryItemType_Actor.h
--- a/Engine/Source/Runtime/AIModule/Classes/EnvironmentQuery/Items/EnvQueryItemType_Actor.h
+++ b/Engine/Source/Runtime/AIModule/Classes/EnvironmentQuery/Items/EnvQueryItemType_Actor.h
@@ -5,6 +5,8 @@
 #include "EnvironmentQuery/Items/EnvQueryItemType_ActorBase.h"
 #include "EnvQueryItemType_Actor.generated.h"
 
+struct FEnvQueryContextData;
+
 UCLASS()
 class AIMODULE_API UEnvQueryItemType_Actor : public UEnvQueryItemType_ActorBase
 {
@@ -13,9 +15,9 @@ class AIMODULE_API UEnvQueryItemType_Actor : public UEnvQueryItemType_ActorBase
 	static AActor* GetValue(const uint8* RawData);
 	static void SetValue(uint8* RawData, const AActor* Value);
 
-	static void SetContextHelper(struct FEnvQueryContextData& ContextData, const AActor* SingleActor);
-	static void SetContextHelper(struct FEnvQueryContextData& ContextData, const TArray<const AActor*>& MultipleActors);
-	static void SetContextHelper(struct FEnvQueryContextData& ContextData, const TArray<AActor*>& MultipleActors);
+	static void SetContextHelper(FEnvQueryContextData& ContextData, const AActor* SingleActor);
+	static void SetContextHelper(FEnvQueryContextData& ContextData, const TArray<const AActor*>& MultipleActors);
+	static void SetContextHelper(FEnvQueryContextData& ContextData, const TArray<AActor*>& MultipleActors);
 
 	virtual FVector GetItemLocation(const uint8* RawData) const override;
 	virtual FRotator GetItemRotation(const uint8* RawData) const override;
diff --git a/Engine/Source/Runtime/AIModule/Classes/EnvironmentQuery/Items/EnvQueryItemType_ActorBase.h b/Engine/Source/Runtime/AIModule/Classes/EnvironmentQuery/Items/EnvQueryItemType_ActorBase.h
--- a/Engine/Source/Runtime/AIModule/Classes/EnvironmentQuery/Items/EnvQueryItemType_ActorBase.h
+++ b/Engine/Source/Runtime/AIModule/Classes/EnvironmentQuery/Items/EnvQueryItemType_ActorBase.h
@@ -5,6 +5,8 @@
 #include "EnvironmentQuery/Items/EnvQueryItemType_VectorBase.h"
 #include "EnvQueryItemType_ActorBase.generated.h"
 
+struct FBlackboardKeySelector;
+class UBlackboardComponent;
 
 UCLASS(Abstract)
 class AIMODULE_API UEnvQueryItemType_ActorBase : public UEnvQueryItemType_VectorBase
@@ -13,7 +15,7 @@ class AIMODULE_API UEnvQueryItemType_ActorBase : public UEnvQueryItemType_Vector
 
 	virtual AActor* GetActor(const uint8* RawData) const;
 
-	virtual void AddBlackboardFilters(struct FBlackboardKeySelector& KeySelector, UObject* FilterOwner) const override;
-	virtual bool StoreInBlackboard(struct FBlackboardKeySelector& KeySelector, class UBlackboardComponent* Blackboard, const uint8* RawData) const override;
+	virtual void AddBlackboardFilters(FBlackboardKeySelector& KeySelector, UObject* FilterOwner) const override;
+	virtual bool StoreInBlackboard(FBlackboardKeySelector& KeySelector, UBlackboardComponent* Blackboard, const uint8* RawData) const override;
 	virtual FString GetDescription(const uint8* RawData) const override;
 };
diff --git a/Engine/Source/Runtime/AIModule/Classes/EnvironmentQuery/Items/EnvQueryItemType_Direction.h b/Engine/Source/Runtime/AIModule/Classes/EnvironmentQuery/Items/EnvQueryItemType_Direction.h
--- a/Engine/Source/Runtime/AIModule/Classes/EnvironmentQuery/Items/EnvQueryItemType_Direction.h
+++ b/Engine/Source/Runtime/AIModule/Classes/EnvironmentQuery/Items/EnvQueryItemType_Direction.h
@@ -4,6 +4,8 @@
 #include "EnvironmentQuery/Items/EnvQueryItemType_VectorBase.h"
 #include "EnvQueryItemType_Direction.generated.h"
 
+struct FEnvQueryContextData;
+
 UCLASS()
 class AIMODULE_API UEnvQueryItemType_Direction : public UEnvQueryItemType_VectorBase
 {
@@ -15,10 +17,10 @@ class AIMODULE_API UEnvQueryItemType_Direction : public UEnvQueryItemType_Vector
 	static FRotator GetValueRot(const uint8* RawData);
 	static void SetValueRot(uint8* RawData, const FRotator& Value);
 
-	static void SetContextHelper(struct FEnvQueryContextData& ContextData, const FVector& SingleDirection);
-	static void SetContextHelper(struct FEnvQueryContextData& ContextData, const FRotator& SingleRotation);
-	static void SetContextHelper(struct FEnvQueryContextData& ContextData, const TArray<FVector>& MultipleDirections);
-	static void SetContextHelper(struct FEnvQueryContextData& ContextData, const TArray<FRotator>& MultipleRotations);
+	static void SetContextHelper(FEnvQueryContextData& ContextData, const FVector& SingleDirection);
+	static void SetContextHelper(FEnvQueryContextData& ContextData, const FRotator& SingleRotation);
+	static void SetContextHelper(FEnvQueryContextData& ContextData, const TArray<FVector>& MultipleDirections);
+	static void SetContextHelper(FEnvQueryContextData& ContextData, const TArray<FRotator>& MultipleRotations);
 
 	virtual FRotator GetItemRotation(const uint8* RawData) const override;
 	virtual FString GetDescription(const uint8* RawData) const override;
diff --git a/Engine/Source/Runtime/AIModule/Classes/EnvironmentQuery/Items/EnvQueryItemType_Point.h b/Engine/Source/Runtime/AIModule/Classes/EnvironmentQuery/Items/EnvQueryItemType_Point.h
--- a/Engine/Source/Runtime/AIModule/Classes/EnvironmentQuery/Items/EnvQueryItemType_Point.h
+++ b/Engine/Source/Runtime/AIModule/Classes/EnvironmentQuery/Items/EnvQueryItemType_Point.h
@@ -4,6 +4,8 @@
 #include "EnvironmentQuery/Items/EnvQueryItemType_VectorBase.h"
 #include "EnvQueryItemType_Point.generated.h"
 
+struct FEnvQueryContextData;
+
 UCLASS()
 class AIMODULE_API UEnvQueryItemType_Point : public UEnvQueryItemType_VectorBase
 {
@@ -12,8 +14,8 @@ class AIMODULE_API UEnvQueryItemType_Point : public UEnvQueryItemType_VectorBase
 	static FVector GetValue(const uint8* RawData);
 	static void SetValue(uint8* RawData, const FVector& Value);
 
-	static void SetContextHelper(struct FEnvQueryContextData& ContextData, const FVector& SinglePoint);
-	static void SetContextHelper(struct FEnvQueryContextData& ContextData, const TArray<FVector>& MultiplePoints);
+	static void SetContextHelper(FEnvQueryContextData& ContextData, const FVector& SinglePoint);
+	static void SetContextHelper(FEnvQueryContextData& ContextData, const TArray<FVector>& MultiplePoints);
 
 	virtual FVector GetItemLocation(const uint8* RawData) const override;
 };
diff --git a/Engine/Source/Runtime/AIModule/Classes/EnvironmentQuery/Items/EnvQueryItemType_VectorBase.h b/Engine/Source/Runtime/AIModule/Classes/EnvironmentQuery/Items/EnvQueryItemType_VectorBase.h
--- a/Engine/Source/Runtime/AIModule/Classes/EnvironmentQuery/Items/EnvQueryItemType_VectorBase.h
+++ b/Engine/Source/Runtime/AIModule/Classes/EnvironmentQuery/Items/EnvQueryItemType_VectorBase.h
@@ -4,6 +4,9 @@
 #include "EnvironmentQuery/Items/EnvQueryItemType.h"
 #include "EnvQueryItemType_VectorBase.generated.h"
 
+struct FBlackboardKeySelector;
+class UBlackboardComponent;
+
 UCLASS(Abstract)
 class AIMODULE_API UEnvQueryItemType_VectorBase : public UEnvQueryItemType
 {
@@ -12,7 +15,7 @@ class AIMODULE_API UEnvQueryItemType_VectorBase : public UEnvQueryItemType
 	virtual FVector GetItemLocation(const uint8* RawData) const;
 	virtual FRotator GetItemRotation(const uint8* RawData) const;
 
-	virtual void AddBlackboardFilters(struct FBlackboardKeySelector& KeySelector, UObject* FilterOwner) const override;
-	virtual bool StoreInBlackboard(struct FBlackboardKeySelector& KeySelector, class UBlackboardComponent* Blackboard, const uint8* RawData) const override;
+	virtual void AddBlackboardFilters(FBlackboardKeySelector& KeySelector, UObject* FilterOwner) const override;
+	virtual bool StoreInBlackboard(FBlackboardKeySelector& KeySelector, UBlackboardComponent* Blackboard, const uint8* RawData) const override;
 	virtual FString GetDescription(const uint8* RawData) const override;
 };
diff --git a/Engine/Source/Runtime/AIModule/Classes/EnvironmentQuery/Tests/EnvQueryTest_Distance.h b/Engine/Source/Runtime/AIModule/Classes/EnvironmentQuery/Tests/EnvQueryTest_Distance.h
--- a/Engine/Source/Runtime/AIModule/Classes/EnvironmentQuery/Tests/EnvQueryTest_Distance.h
+++ b/Engine/Source/Runtime/AIModule/Classes/EnvironmentQuery/Tests/EnvQueryTest_Distance.h
@@ -4,6 +4,8 @@
 #include "EnvironmentQuery/EnvQueryTest.h"
 #include "EnvQueryTest_Distance.generated.h"
 
+class UEnvQueryContext;
+
 UENUM()
 namespace EEnvTestDistance
 {
@@ -26,7 +28,7 @@ class UEnvQueryTest_Distance : public UEnvQueryTest
 
 	/** context */
 	UPROPERTY(EditDefaultsOnly, Category=Distance)
-	TSubclassOf<class UEnvQueryContext> DistanceTo;
+	TSubclassOf<UEnvQueryContext> DistanceTo;
 
 	virtual void RunTest(FEnvQueryInstance& QueryInstance) const override;
 
diff --git a/Engine/Source/Runtime/AIModule/Classes/EnvironmentQuery/Tests/EnvQueryTest_Dot.h b/Engine/Source/Runtime/AIModule/Classes/EnvironmentQuery/Tests/EnvQueryTest_Dot.h
--- a/Engine/Source/Runtime/AIModule/Classes/EnvironmentQuery/Tests/EnvQueryTest_Dot.h
+++ b/Engine/Source/Runtime/AIModule/Classes/EnvironmentQuery/Tests/EnvQueryTest_Dot.h
@@ -4,6 +4,9 @@
 #include "EnvironmentQuery/EnvQueryTest.h"
 #include "EnvQueryTest_Dot.generated.h"
 
+class UEnvQueryContext;
+struct FEnvQueryInstance;
+
 UENUM(BlueprintType)
 enum class EEnvTestDot : uint8
 {
@@ -42,18 +45,18 @@ protected:
 	virtual FText GetDescriptionDetails() const override;
 
 	/** helper function: gather directions from context pairs */
-	void GatherLineDirections(TArray<FVector>& Directions, struct FEnvQueryInstance& QueryInstance, const FVector& ItemLocation,
-		TSubclassOf<class UEnvQueryContext> LineFrom, TSubclassOf<class UEnvQueryContext> LineTo) const;
+	void GatherLineDirections(TArray<FVector>& Directions, FEnvQueryInstance& QueryInstance, const FVector& ItemLocation,
+		TSubclassOf<UEnvQueryContext> LineFrom, TSubclassOf<UEnvQueryContext> LineTo) const;
 
 	/** helper function: gather directions from context */
-	void GatherLineDirections(TArray<FVector>& Directions, struct FEnvQueryInstance& QueryInstance, const FRotator& ItemRotation,
-		TSubclassOf<class UEnvQueryContext> LineDirection) const;
+	void GatherLineDirections(TArray<FVector>& Directions, FEnvQueryInstance& QueryInstance, const FRotator& ItemRotation,
+		TSubclassOf<UEnvQueryContext> LineDirection) const;
 
 	/** helper function: gather directions from proper contexts */
-	void GatherLineDirections(TArray<FVector>& Directions, struct FEnvQueryInstance& QueryInstance,
-		TSubclassOf<class UEnvQueryContext> LineFrom, TSubclassOf<class UEnvQueryContext> LineTo, TSubclassOf<class UEnvQueryContext> LineDirection, bool bUseDirectionContext,
+	void GatherLineDirections(TArray<FVector>& Directions, FEnvQueryInstance& QueryInstance,
+		TSubclassOf<UEnvQueryContext> LineFrom, TSubclassOf<UEnvQueryContext> LineTo, TSubclassOf<UEnvQueryContext> LineDirection, bool bUseDirectionContext,
 		const FVector& ItemLocation = FVector::ZeroVector, const FRotator& ItemRotation = FRotator::ZeroRotator) const;
 
 	/** helper function: check if contexts are updated per item */
-	bool RequiresPerItemUpdates(TSubclassOf<class UEnvQueryContext> LineFrom, TSubclassOf<class UEnvQueryContext> LineTo, TSubclassOf<class UEnvQueryContext> LineDirection, bool bUseDirectionContext) const;
+	bool RequiresPerItemUpdates(TSubclassOf<UEnvQueryContext> LineFrom, TSubclassOf<UEnvQueryContext> LineTo, TSubclassOf<UEnvQueryContext> LineDirection, bool bUseDirectionContext) const;
 };
diff --git a/Engine/Source/Runtime/AIModule/Classes/EnvironmentQuery/Tests/EnvQueryTest_Pathfinding.h b/Engine/Source/Runtime/AIModule/Classes/EnvironmentQuery/Tests/EnvQueryTest_Pathfinding.h
--- a/Engine/Source/Runtime/AIModule/Classes/EnvironmentQuery/Tests/EnvQueryTest_Pathfinding.h
+++ b/Engine/Source/Runtime/AIModule/Classes/EnvironmentQuery/Tests/EnvQueryTest_Pathfinding.h
@@ -4,6 +4,8 @@
 #include "EnvironmentQuery/EnvQueryTest.h"
 #include "EnvQueryTest_Pathfinding.generated.h"
 
+class UEnvQueryContext;
+
 UENUM()
 namespace EEnvTestPathfinding
 {
@@ -26,7 +28,7 @@ class UEnvQueryTest_Pathfinding : public UEnvQueryTest
 
 	/** context: other end of pathfinding test */
 	UPROPERTY(EditDefaultsOnly, Category=Pathfinding)
-	TSubclassOf<class UEnvQueryContext> Context;
+	TSubclassOf<UEnvQueryContext> Context;
 
 	/** pathfinding direction */
 	UPROPERTY(EditDefaultsOnly, Category=Pathfinding)
diff --git a/Engine/Source/Runtime/AIModule/Classes/EnvironmentQuery/Tests/EnvQueryTest_Trace.h b/Engine/Source/Runtime/AIModule/Classes/EnvironmentQuery/Tests/EnvQueryTest_Trace.h
--- a/Engine/Source/Runtime/AIModule/Classes/EnvironmentQuery/Tests/EnvQueryTest_Trace.h
+++ b/Engine/Source/Runtime/AIModule/Classes/EnvironmentQuery/Tests/EnvQueryTest_Trace.h
@@ -4,6 +4,8 @@
 #include "EnvironmentQuery/EnvQueryTest.h"
 #include "EnvQueryTest_Trace.generated.h"
 
+class UEnvQueryContext;
+
 UCLASS(MinimalAPI)
 class UEnvQueryTest_Trace : public UEnvQueryTest
 {
@@ -27,7 +29,7 @@ class UEnvQueryTest_Trace : public UEnvQueryTest
 
 	/** context: other end of trace test */
 	UPROPERTY(EditDefaultsOnly, Category=Trace)
-	TSubclassOf<class UEnvQueryContext> Context;
+	TSubclassOf<UEnvQueryContext> Context;
 
 	virtual void RunTest(FEnvQueryInstance& QueryInstance) const override;
 
diff --git a/Engine/Source/Runtime/AIModule/Classes/VisualLoggerExtension.h b/Engine/Source/Runtime/AIModule/Classes/VisualLoggerExtension.h
--- a/Engine/Source/Runtime/AIModule/Classes/VisualLoggerExtension.h
+++ b/Engine/Source/Runtime/AIModule/Classes/VisualLoggerExtension.h
@@ -5,23 +5,30 @@
 #endif
 #include "VisualLoggerExtension.generated.h"
 
+class UWorld;
+class AActor;
+
 namespace EVisLogTags
 {
 	const FString TAG_EQS = TEXT("EQS");
 }
 
 #if ENABLE_VISUAL_LOG
+struct FVisualLogDataBlock;
+struct FLogEntryItem;
+class UCanvas;
+
 class FVisualLoggerExtension : public FVisualLogExtensionInterface
 {
 public:
 	FVisualLoggerExtension();
-	virtual void OnTimestampChange(float Timestamp, class UWorld* InWorld, class AActor* HelperActor) override;
-	virtual void DrawData(class UWorld* InWorld, class UCanvas* Canvas, class AActor* HelperActor, const FName& TagName, const struct FVisualLogDataBlock& DataBlock, float Timestamp) override;
-	virtual void DisableDrawingForData(class UWorld* InWorld, class UCanvas* Canvas, class AActor* HelperActor, const FName& TagName, const FVisualLogDataBlock& DataBlock, float Timestamp) override;
-	virtual void LogEntryLineSelectionChanged(TSharedPtr<struct FLogEntryItem> SelectedItem, int64 UserData, FName TagName) override;
+	virtual void OnTimestampChange(float Timestamp, UWorld* InWorld, AActor* HelperActor) override;
+	virtual void DrawData(UWorld* InWorld, UCanvas* Canvas, AActor* HelperActor, const FName& TagName, const FVisualLogDataBlock& DataBlock, float Timestamp) override;
+	virtual void DisableDrawingForData(UWorld* InWorld, UCanvas* Canvas, AActor* HelperActor, const FName& TagName, const FVisualLogDataBlock& DataBlock, float Timestamp) override;
+	virtual void LogEntryLineSelectionChanged(TSharedPtr<FLogEntryItem> SelectedItem, int64 UserData, FName TagName) override;
 
 private:
-	void DisableEQSRendering(class AActor* HelperActor);
+	void DisableEQSRendering(AActor* HelperActor);
 
 protected:
 	uint32 CachedEQSId;
@@ -30,10 +37,8 @@ protected:
 };
 #endif //ENABLE_VISUAL_LOG
 
-UCLASS(Abstract, CustomConstructor)
+UCLASS(Abstract)
 class AIMODULE_API UVisualLoggerExtension : public UObject
 {
-	GENERATED_UCLASS_BODY()
-
-	UVisualLoggerExtension(const FObjectInitializer& ObjectInitializer) : Super(ObjectInitializer) {}
+	GENERATED_BODY()
 };
diff --git a/Engine/Source/Runtime/Engine/Public/AI/BlueprintNodeHelpers.h b/Engine/Source/Runtime/Engine/Public/AI/BlueprintNodeHelpers.h
--- a/Engine/Source/Runtime/Engine/Public/AI/BlueprintNodeHelpers.h
+++ b/Engine/Source/Runtime/Engine/Public/AI/BlueprintNodeHelpers.h
@@ -2,6 +2,8 @@
 
 #pragma once
 
+class UBTNode;
+
 namespace BlueprintNodeHelpers
 {
 	FString CollectPropertyDescription(const UObject* Ob, const UClass* StopAtClass, const TArray<UProperty*>& PropertyData);
@@ -16,7 +18,7 @@ namespace BlueprintNodeHelpers
 	void CopyPropertiesToContext(const TArray<UProperty*>& PropertyData, uint8* ObjectMemory, uint8* ContextMemory);
 	void CopyPropertiesFromContext(const TArray<UProperty*>& PropertyData, uint8* ObjectMemory, uint8* ContextMemory);
 
-	bool FindNodeOwner(AActor* OwningActor, class UBTNode* Node, UBehaviorTreeComponent*& OwningComp, int32& OwningInstanceIdx);
+	bool FindNodeOwner(AActor* OwningActor, UBTNode* Node, UBehaviorTreeComponent*& OwningComp, int32& OwningInstanceIdx);
 
 	void AbortLatentActions(UActorComponent* OwnerOb, const UObject* Ob);
 
diff --git a/Engine/Source/Runtime/Engine/Public/AI/Navigation/PImplRecastNavMesh.h b/Engine/Source/Runtime/Engine/Public/AI/Navigation/PImplRecastNavMesh.h
--- a/Engine/Source/Runtime/Engine/Public/AI/Navigation/PImplRecastNavMesh.h
+++ b/Engine/Source/Runtime/Engine/Public/AI/Navigation/PImplRecastNavMesh.h
@@ -15,6 +15,10 @@
 
 #define RECAST_VERY_SMALL_AGENT_RADIUS 0.0f
 
+class ARecastNavMesh;
+class dtNavMesh;
+struct dtMeshTile;
+
 class ENGINE_API FRecastQueryFilter : public INavigationQueryFilterInterface, public dtQueryFilter
 {
 public:
@@ -58,7 +62,7 @@ class ENGINE_API FPImplRecastNavMesh
 public:
 
 	/** Constructor */
-	FPImplRecastNavMesh(class ARecastNavMesh* Owner);
+	FPImplRecastNavMesh(ARecastNavMesh* Owner);
 
 	/** Dtor */
 	~FPImplRecastNavMesh();
@@ -182,15 +186,15 @@ public:
 	static void SetFilterForbiddenFlags(FRecastQueryFilter* Filter, uint16 ForbiddenFlags);
 
 public:
-	class dtNavMesh const* GetRecastMesh() const { return DetourNavMesh; };
-	class dtNavMesh* GetRecastMesh() { return DetourNavMesh; };
+	dtNavMesh const* GetRecastMesh() const { return DetourNavMesh; };
+	dtNavMesh* GetRecastMesh() { return DetourNavMesh; };
 	void ReleaseDetourNavMesh();
 
 	/** Assigns recast generated navmesh to this instance.
 	 *	@param bOwnData if true from now on this FPImplRecastNavMesh instance will be responsible for this piece 
 	 *		of memory
 	 */
-	void SetRecastMesh(class dtNavMesh* NavMesh);
+	void SetRecastMesh(dtNavMesh* NavMesh);
 
 	float GetTotalDataSize() const;
 
@@ -202,10 +206,10 @@ public:
 	 *	@note no check if segment is on poly is performed. */
 	float CalcSegmentCostOnPoly(NavNodeRef PolyID, const dtQueryFilter* Filter, const FVector& StartLoc, const FVector& EndLoc) const;
 
-	class ARecastNavMesh* NavMeshOwner;
+	ARecastNavMesh* NavMeshOwner;
 	
 	/** Recast's runtime navmesh data that we can query against */
-	class dtNavMesh* DetourNavMesh;
+	dtNavMesh* DetourNavMesh;
 
 	/** query used for searching data on game thread */
 	mutable dtNavMeshQuery SharedNavQuery;
@@ -227,7 +231,7 @@ public:
 		const FVector& RecastStart, FVector& RecastEnd,
 		dtQueryResult& PathResult) const;
 
-	void GetDebugPolyEdges(const struct dtMeshTile* Tile, bool bInternalEdges, bool bNavMeshEdges, TArray<FVector>& InternalEdgeVerts, TArray<FVector>& NavMeshEdgeVerts) const;
+	void GetDebugPolyEdges(const dtMeshTile* Tile, bool bInternalEdges, bool bNavMeshEdges, TArray<FVector>& InternalEdgeVerts, TArray<FVector>& NavMeshEdgeVerts) const;
 
 	/** workhorse function finding portal edges between corridor polys */
 	void GetEdgesForPathCorridorImpl(const TArray<NavNodeRef>* PathCorridor, TArray<FNavigationPortalEdge>* PathCorridorEdges, const dtNavMeshQuery& NavQuery) const;
