More in-function type forward declarations removed from AI headers, pass 3 #UE4

diff --git a/Engine/Source/Runtime/AIModule/Classes/Actions/PawnAction.h b/Engine/Source/Runtime/AIModule/Classes/Actions/PawnAction.h
--- a/Engine/Source/Runtime/AIModule/Classes/Actions/PawnAction.h
+++ b/Engine/Source/Runtime/AIModule/Classes/Actions/PawnAction.h
@@ -5,9 +5,13 @@
 #include "BrainComponent.h"
 #include "PawnAction.generated.h"
 
+class APawn;
+class AController;
 class UPawnAction;
 class UPawnActionsComponent;
 class UBrainComponent;
+struct FPawnActionStack;
+struct FAIMessage;
 
 UENUM()
 namespace EPawnSubActionTriggeringPolicy
@@ -42,7 +46,7 @@ class AIMODULE_API UPawnAction : public UObject
 	GENERATED_UCLASS_BODY()
 
 	friend UPawnActionsComponent;
-	friend struct FPawnActionStack;
+	friend FPawnActionStack;
 
 private:
 	/** Current child node executing on top of this Action */
@@ -159,8 +163,8 @@ public:
 	FORCEINLINE EPawnActionAbortState::Type GetAbortState() const { return AbortState; }
 	FORCEINLINE UPawnActionsComponent* GetOwnerComponent() const { return OwnerComponent; }
 	FORCEINLINE UObject* GetInstigator() const { return Instigator; }
-	class APawn* GetPawn();
-	class AController* GetController();
+	APawn* GetPawn();
+	AController* GetController();
 
 	template<class TActionClass>
 	static TActionClass* CreateActionInstance(UWorld& World)
@@ -174,7 +178,7 @@ public:
 	//----------------------------------------------------------------------//
 	void WaitForMessage(FName MessageType, FAIRequestID RequestID = FAIRequestID::AnyRequest);
 	// @note this function will change its signature once AI messaging is rewritten @todo
-	virtual void HandleAIMessage(UBrainComponent*, const struct FAIMessage&){};
+	virtual void HandleAIMessage(UBrainComponent*, const FAIMessage&){};
 
 	void SetActionObserver(const FPawnActionEventDelegate& ActionObserver) { this->ActionObserver = ActionObserver; }
 
diff --git a/Engine/Source/Runtime/AIModule/Classes/Actions/PawnAction_Move.h b/Engine/Source/Runtime/AIModule/Classes/Actions/PawnAction_Move.h
--- a/Engine/Source/Runtime/AIModule/Classes/Actions/PawnAction_Move.h
+++ b/Engine/Source/Runtime/AIModule/Classes/Actions/PawnAction_Move.h
@@ -4,7 +4,10 @@
 #include "TimerManager.h"
 #include "PawnAction_Move.generated.h"
 
+class AActor;
 class AAIController;
+class UNavigationQueryFilter;
+struct FAIMessage;
 
 namespace EPawnActionMoveMode
 {
@@ -21,7 +24,7 @@ class AIMODULE_API UPawnAction_Move : public UPawnAction
 	GENERATED_UCLASS_BODY()
 protected:
 	UPROPERTY(Category = PawnAction, EditAnywhere, BlueprintReadWrite)
-	class AActor* GoalActor;
+	AActor* GoalActor;
 
 	UPROPERTY(Category = PawnAction, EditAnywhere, BlueprintReadWrite)
 	FVector GoalLocation;
@@ -31,7 +34,7 @@ protected:
 
 	/** "None" will result in default filter being used */
 	UPROPERTY(Category = PawnAction, EditAnywhere, BlueprintReadWrite)
-	TSubclassOf<class UNavigationQueryFilter> FilterClass;
+	TSubclassOf<UNavigationQueryFilter> FilterClass;
 
 	UPROPERTY(Category = PawnAction, EditAnywhere, BlueprintReadWrite)
 	uint32 bAllowStrafe : 1;
@@ -59,13 +62,13 @@ protected:
 public:
 	virtual void BeginDestroy() override;
 
-	static UPawnAction_Move* CreateAction(UWorld& World, class AActor* GoalActor, EPawnActionMoveMode::Type Mode);
+	static UPawnAction_Move* CreateAction(UWorld& World, AActor* GoalActor, EPawnActionMoveMode::Type Mode);
 	static UPawnAction_Move* CreateAction(UWorld& World, const FVector& GoalLocation, EPawnActionMoveMode::Type Mode);
 
 	static bool CheckAlreadyAtGoal(AAIController& Controller, const FVector& TestLocation, float Radius);
 	static bool CheckAlreadyAtGoal(AAIController& Controller, const AActor& TestGoal, float Radius);
 
-	virtual void HandleAIMessage(UBrainComponent*, const struct FAIMessage&) override;
+	virtual void HandleAIMessage(UBrainComponent*, const FAIMessage&) override;
 
 	void SetPath(FNavPathSharedRef InPath);
 	void OnPathUpdated(FNavigationPath* UpdatedPath, ENavPathEvent::Type Event);
@@ -76,7 +79,7 @@ public:
 	void EnablePathUpdateOnMoveGoalLocationChange(bool bEnable) { bUpdatePathToGoal = bEnable; }
 	void EnableGoalLocationProjectionToNavigation(bool bEnable) { bProjectGoalToNavigation = bEnable; }
 	void EnableChildAbortionOnPathUpdate(bool bEnable) { bAbortChildActionOnPathChange = bEnable; }
-	void SetFilterClass(TSubclassOf<class UNavigationQueryFilter> NewFilterClass) { FilterClass = NewFilterClass; }
+	void SetFilterClass(TSubclassOf<UNavigationQueryFilter> NewFilterClass) { FilterClass = NewFilterClass; }
 
 protected:
 	/** currently followed path */
diff --git a/Engine/Source/Runtime/AIModule/Classes/Actions/PawnAction_Repeat.h b/Engine/Source/Runtime/AIModule/Classes/Actions/PawnAction_Repeat.h
--- a/Engine/Source/Runtime/AIModule/Classes/Actions/PawnAction_Repeat.h
+++ b/Engine/Source/Runtime/AIModule/Classes/Actions/PawnAction_Repeat.h
@@ -29,7 +29,7 @@ class AIMODULE_API UPawnAction_Repeat : public UPawnAction
 	EPawnSubActionTriggeringPolicy::Type SubActionTriggeringPolicy;
 
 	/** @param NumberOfRepeats number of times to repeat action. UPawnAction_Repeat::LoopForever loops forever */
-	static UPawnAction_Repeat* CreateAction(UWorld& World, class UPawnAction* ActionToRepeat, int32 NumberOfRepeats, EPawnSubActionTriggeringPolicy::Type InSubActionTriggeringPolicy = EPawnSubActionTriggeringPolicy::CopyBeforeTriggering);
+	static UPawnAction_Repeat* CreateAction(UWorld& World, UPawnAction* ActionToRepeat, int32 NumberOfRepeats, EPawnSubActionTriggeringPolicy::Type InSubActionTriggeringPolicy = EPawnSubActionTriggeringPolicy::CopyBeforeTriggering);
 
 protected:
 	virtual bool Start() override;
diff --git a/Engine/Source/Runtime/AIModule/Classes/BehaviorTree/BTAuxiliaryNode.h b/Engine/Source/Runtime/AIModule/Classes/BehaviorTree/BTAuxiliaryNode.h
--- a/Engine/Source/Runtime/AIModule/Classes/BehaviorTree/BTAuxiliaryNode.h
+++ b/Engine/Source/Runtime/AIModule/Classes/BehaviorTree/BTAuxiliaryNode.h
@@ -5,6 +5,8 @@
 #include "BehaviorTree/BTNode.h"
 #include "BTAuxiliaryNode.generated.h"
 
+class UBehaviorTreeComponent;
+
 struct FBTAuxiliaryMemory : public FBTInstancedNodeMemory
 {
 	float NextTickRemainingTime;
@@ -31,15 +33,15 @@ class AIMODULE_API UBTAuxiliaryNode : public UBTNode
 	GENERATED_UCLASS_BODY()
 
 	/** wrapper for node instancing: OnBecomeRelevant */
-	void WrappedOnBecomeRelevant(class UBehaviorTreeComponent* OwnerComp, uint8* NodeMemory) const;
+	void WrappedOnBecomeRelevant(UBehaviorTreeComponent* OwnerComp, uint8* NodeMemory) const;
 
 	/** wrapper for node instancing: OnCeaseRelevant */
-	void WrappedOnCeaseRelevant(class UBehaviorTreeComponent* OwnerComp, uint8* NodeMemory) const;
+	void WrappedOnCeaseRelevant(UBehaviorTreeComponent* OwnerComp, uint8* NodeMemory) const;
 
 	/** wrapper for node instancing: TickNode */
-	void WrappedTickNode(class UBehaviorTreeComponent* OwnerComp, uint8* NodeMemory, float DeltaSeconds) const;
+	void WrappedTickNode(UBehaviorTreeComponent* OwnerComp, uint8* NodeMemory, float DeltaSeconds) const;
 
-	virtual void DescribeRuntimeValues(const class UBehaviorTreeComponent* OwnerComp, uint8* NodeMemory, EBTDescriptionVerbosity::Type Verbosity, TArray<FString>& Values) const override;
+	virtual void DescribeRuntimeValues(const UBehaviorTreeComponent* OwnerComp, uint8* NodeMemory, EBTDescriptionVerbosity::Type Verbosity, TArray<FString>& Values) const override;
 	virtual uint16 GetSpecialMemorySize() const override;
 
 protected:
diff --git a/Engine/Source/Runtime/AIModule/Classes/BehaviorTree/BTCompositeNode.h b/Engine/Source/Runtime/AIModule/Classes/BehaviorTree/BTCompositeNode.h
--- a/Engine/Source/Runtime/AIModule/Classes/BehaviorTree/BTCompositeNode.h
+++ b/Engine/Source/Runtime/AIModule/Classes/BehaviorTree/BTCompositeNode.h
@@ -5,7 +5,16 @@
 #include "BehaviorTree/BTNode.h"
 #include "BTCompositeNode.generated.h"
 
-DECLARE_DELEGATE_RetVal_ThreeParams(int32, FGetNextChildDelegate, struct FBehaviorTreeSearchData& /*search data*/, int32 /*last child index*/, EBTNodeResult::Type /*last result*/);
+class UBTNode;
+class UBTTaskNode;
+class UBTCompositeNode;
+class UBTDecorator;
+class UBTService;
+class UBehaviorTreeComponent;
+struct FBehaviorTreeSearchData;
+struct FBTNodeIndex;
+
+DECLARE_DELEGATE_RetVal_ThreeParams(int32, FGetNextChildDelegate, FBehaviorTreeSearchData& /*search data*/, int32 /*last child index*/, EBTNodeResult::Type /*last result*/);
 
 struct FBTCompositeMemory
 {
@@ -59,14 +68,14 @@ struct FBTCompositeChild
 
 	/** child node */
 	UPROPERTY()
-	class UBTCompositeNode* ChildComposite;
+	UBTCompositeNode* ChildComposite;
 
 	UPROPERTY()
-	class UBTTaskNode* ChildTask;
+	UBTTaskNode* ChildTask;
 
 	/** execution decorators */
 	UPROPERTY()
-	TArray<class UBTDecorator*> Decorators;
+	TArray<UBTDecorator*> Decorators;
 
 	/** logic operations for decorators */
 	UPROPERTY()
@@ -84,7 +93,7 @@ class AIMODULE_API UBTCompositeNode : public UBTNode
 
 	/** service nodes */
 	UPROPERTY()
-	TArray<class UBTService*> Services;
+	TArray<UBTService*> Services;
 
 	/** delegate for finding next child to execute */
 	FGetNextChildDelegate OnNextChild;
@@ -93,38 +102,38 @@ class AIMODULE_API UBTCompositeNode : public UBTNode
 	void InitializeComposite(uint16 InLastExecutionIndex);
 
 	/** find next child branch to execute */
-	int32 FindChildToExecute(struct FBehaviorTreeSearchData& SearchData, EBTNodeResult::Type& LastResult) const;
+	int32 FindChildToExecute(FBehaviorTreeSearchData& SearchData, EBTNodeResult::Type& LastResult) const;
 
 	/** get index of child node (handle subtrees) */
-	int32 GetChildIndex(struct FBehaviorTreeSearchData& SearchData, const class UBTNode* ChildNode) const;
+	int32 GetChildIndex(FBehaviorTreeSearchData& SearchData, const UBTNode* ChildNode) const;
 	/** get index of child node */
-	int32 GetChildIndex(const class UBTNode* ChildNode) const;
+	int32 GetChildIndex(const UBTNode* ChildNode) const;
 
 	/** called before passing search to child node */
-	void OnChildActivation(struct FBehaviorTreeSearchData& SearchData, const class UBTNode* ChildNode) const;
-	void OnChildActivation(struct FBehaviorTreeSearchData& SearchData, int32 ChildIndex) const;
+	void OnChildActivation(FBehaviorTreeSearchData& SearchData, const UBTNode* ChildNode) const;
+	void OnChildActivation(FBehaviorTreeSearchData& SearchData, int32 ChildIndex) const;
 
 	/** called after child has finished search */
-	void OnChildDeactivation(struct FBehaviorTreeSearchData& SearchData, const class UBTNode* ChildNode, EBTNodeResult::Type& NodeResult) const;
-	void OnChildDeactivation(struct FBehaviorTreeSearchData& SearchData, int32 ChildIndex, EBTNodeResult::Type& NodeResult) const;
+	void OnChildDeactivation(FBehaviorTreeSearchData& SearchData, const UBTNode* ChildNode, EBTNodeResult::Type& NodeResult) const;
+	void OnChildDeactivation(FBehaviorTreeSearchData& SearchData, int32 ChildIndex, EBTNodeResult::Type& NodeResult) const;
 
 	/** called when start enters this node */
-	void OnNodeActivation(struct FBehaviorTreeSearchData& SearchData) const;
+	void OnNodeActivation(FBehaviorTreeSearchData& SearchData) const;
 
 	/** called when search leaves this node */
-	void OnNodeDeactivation(struct FBehaviorTreeSearchData& SearchData, EBTNodeResult::Type& NodeResult) const;
+	void OnNodeDeactivation(FBehaviorTreeSearchData& SearchData, EBTNodeResult::Type& NodeResult) const;
 
 	/** called when search needs to reactivate this node */
-	void OnNodeRestart(struct FBehaviorTreeSearchData& SearchData) const;
+	void OnNodeRestart(FBehaviorTreeSearchData& SearchData) const;
 
 	/** notify about task execution start */
-	void ConditionalNotifyChildExecution(class UBehaviorTreeComponent* OwnerComp, uint8* NodeMemory, const class UBTNode* ChildNode, EBTNodeResult::Type& NodeResult) const;
+	void ConditionalNotifyChildExecution(UBehaviorTreeComponent* OwnerComp, uint8* NodeMemory, const UBTNode* ChildNode, EBTNodeResult::Type& NodeResult) const;
 
 	/** size of instance memory */
 	virtual uint16 GetInstanceMemorySize() const override;
 
 	/** @return child node at given index */
-	class UBTNode* GetChildNode(int32 Index) const;
+	UBTNode* GetChildNode(int32 Index) const;
 
 	/** @return children count */
 	int32 GetChildrenNum() const;
@@ -139,10 +148,10 @@ class AIMODULE_API UBTCompositeNode : public UBTNode
 	void SetChildOverride(FBehaviorTreeSearchData& SearchData, int8 Index) const;
 
 	/** gathers description of all runtime parameters */
-	virtual void DescribeRuntimeValues(const class UBehaviorTreeComponent* OwnerComp, uint8* NodeMemory, EBTDescriptionVerbosity::Type Verbosity, TArray<FString>& Values) const override;
+	virtual void DescribeRuntimeValues(const UBehaviorTreeComponent* OwnerComp, uint8* NodeMemory, EBTDescriptionVerbosity::Type Verbosity, TArray<FString>& Values) const override;
 
 	/** check if child node can execute new subtree */
-	virtual bool CanPushSubtree(class UBehaviorTreeComponent* OwnerComp, uint8* NodeMemory, int32 ChildIdx) const;
+	virtual bool CanPushSubtree(UBehaviorTreeComponent* OwnerComp, uint8* NodeMemory, int32 ChildIdx) const;
 
 #if WITH_EDITOR
 	/** @return allowed flow abort modes for decorators */
@@ -151,10 +160,10 @@ class AIMODULE_API UBTCompositeNode : public UBTNode
 #endif // WITH_EDITOR
 
 	/** find branch containing specified node index */
-	int32 GetMatchingChildIndex(int32 ActiveInstanceIdx, struct FBTNodeIndex& NodeIdx) const;
+	int32 GetMatchingChildIndex(int32 ActiveInstanceIdx, FBTNodeIndex& NodeIdx) const;
 
 	/** is child execution allowed by decorators? */
-	bool DoDecoratorsAllowExecution(class UBehaviorTreeComponent* OwnerComp, int32 InstanceIdx, int32 ChildIdx) const;
+	bool DoDecoratorsAllowExecution(UBehaviorTreeComponent* OwnerComp, int32 InstanceIdx, int32 ChildIdx) const;
 
 protected:
 
@@ -171,35 +180,35 @@ protected:
 	uint8 bUseNodeDeactivationNotify : 1;
 
 	/** called just after child execution, allows to modify result */
-	virtual void NotifyChildExecution(class UBehaviorTreeComponent* OwnerComp, uint8* NodeMemory, int32 ChildIdx, EBTNodeResult::Type& NodeResult) const;
+	virtual void NotifyChildExecution(UBehaviorTreeComponent* OwnerComp, uint8* NodeMemory, int32 ChildIdx, EBTNodeResult::Type& NodeResult) const;
 
 	/** called when start enters this node */
-	virtual void NotifyNodeActivation(struct FBehaviorTreeSearchData& SearchData) const;
+	virtual void NotifyNodeActivation(FBehaviorTreeSearchData& SearchData) const;
 
 	/** called when start leaves this node */
-	virtual void NotifyNodeDeactivation(struct FBehaviorTreeSearchData& SearchData, EBTNodeResult::Type& NodeResult) const;
+	virtual void NotifyNodeDeactivation(FBehaviorTreeSearchData& SearchData, EBTNodeResult::Type& NodeResult) const;
 
 	/** runs through decorators on given child node and notify them about activation */
-	void NotifyDecoratorsOnActivation(struct FBehaviorTreeSearchData& SearchData, int32 ChildIdx) const;
+	void NotifyDecoratorsOnActivation(FBehaviorTreeSearchData& SearchData, int32 ChildIdx) const;
 
 	/** runs through decorators on given child node and notify them about deactivation */
-	void NotifyDecoratorsOnDeactivation(struct FBehaviorTreeSearchData& SearchData, int32 ChildIdx, EBTNodeResult::Type& NodeResult) const;
+	void NotifyDecoratorsOnDeactivation(FBehaviorTreeSearchData& SearchData, int32 ChildIdx, EBTNodeResult::Type& NodeResult) const;
 
 	/** runs through decorators on given child node and notify them about failed activation */
-	void NotifyDecoratorsOnFailedActivation(struct FBehaviorTreeSearchData& SearchData, int32 ChildIdx, EBTNodeResult::Type& NodeResult) const;
+	void NotifyDecoratorsOnFailedActivation(FBehaviorTreeSearchData& SearchData, int32 ChildIdx, EBTNodeResult::Type& NodeResult) const;
 
 	/** get next child to process and store it in CurrentChild */
-	int32 GetNextChild(struct FBehaviorTreeSearchData& SearchData, int32 LastChildIdx, EBTNodeResult::Type LastResult) const;
+	int32 GetNextChild(FBehaviorTreeSearchData& SearchData, int32 LastChildIdx, EBTNodeResult::Type LastResult) const;
 
 	/** store delayed execution request */
-	void RequestDelayedExecution(class UBehaviorTreeComponent* OwnerComp, EBTNodeResult::Type LastResult) const;
+	void RequestDelayedExecution(UBehaviorTreeComponent* OwnerComp, EBTNodeResult::Type LastResult) const;
 };
 
 
 //////////////////////////////////////////////////////////////////////////
 // Inlines
 
-FORCEINLINE class UBTNode* UBTCompositeNode::GetChildNode(int32 Index) const
+FORCEINLINE UBTNode* UBTCompositeNode::GetChildNode(int32 Index) const
 {
 	return Children.IsValidIndex(Index) ?
 		(Children[Index].ChildComposite ?
diff --git a/Engine/Source/Runtime/AIModule/Classes/BehaviorTree/BTDecorator.h b/Engine/Source/Runtime/AIModule/Classes/BehaviorTree/BTDecorator.h
--- a/Engine/Source/Runtime/AIModule/Classes/BehaviorTree/BTDecorator.h
+++ b/Engine/Source/Runtime/AIModule/Classes/BehaviorTree/BTDecorator.h
@@ -5,6 +5,11 @@
 #include "BehaviorTree/BTAuxiliaryNode.h"
 #include "BTDecorator.generated.h"
 
+class UBTNode;
+class UBehaviorTreeComponent;
+class FBehaviorDecoratorDetails; 
+struct FBehaviorTreeSearchData;
+
 /** 
  * Decorators are supporting nodes placed on parent-child connection, that receive notification about execution flow and can be ticked
  *
@@ -31,19 +36,19 @@ class AIMODULE_API UBTDecorator : public UBTAuxiliaryNode
 	void InitializeDecorator(uint8 InChildIndex);
 
 	/** wrapper for node instancing: CalculateRawConditionValue */
-	bool WrappedCanExecute(class UBehaviorTreeComponent* OwnerComp, uint8* NodeMemory) const;
+	bool WrappedCanExecute(UBehaviorTreeComponent* OwnerComp, uint8* NodeMemory) const;
 
 	/** wrapper for node instancing: OnNodeActivation  */
-	void WrappedOnNodeActivation(struct FBehaviorTreeSearchData& SearchData) const;
+	void WrappedOnNodeActivation(FBehaviorTreeSearchData& SearchData) const;
 	
 	/** wrapper for node instancing: OnNodeDeactivation */
-	void WrappedOnNodeDeactivation(struct FBehaviorTreeSearchData& SearchData, EBTNodeResult::Type NodeResult) const;
+	void WrappedOnNodeDeactivation(FBehaviorTreeSearchData& SearchData, EBTNodeResult::Type NodeResult) const;
 
 	/** wrapper for node instancing: OnNodeProcessed */
-	void WrappedOnNodeProcessed(struct FBehaviorTreeSearchData& SearchData, EBTNodeResult::Type& NodeResult) const;
+	void WrappedOnNodeProcessed(FBehaviorTreeSearchData& SearchData, EBTNodeResult::Type& NodeResult) const;
 
 	/** @return decorated child */
-	const class UBTNode* GetMyNode() const;
+	const UBTNode* GetMyNode() const;
 
 	/** @return index of child in parent's array */
 	uint8 GetChildIndex() const;
@@ -102,20 +107,20 @@ protected:
 
 	/** called when underlying node is activated
 	  * this function should be considered as const (don't modify state of object) if node is not instanced! */
-	virtual void OnNodeActivation(struct FBehaviorTreeSearchData& SearchData);
+	virtual void OnNodeActivation(FBehaviorTreeSearchData& SearchData);
 
 	/** called when underlying node has finished
 	 * this function should be considered as const (don't modify state of object) if node is not instanced! */
-	virtual void OnNodeDeactivation(struct FBehaviorTreeSearchData& SearchData, EBTNodeResult::Type NodeResult);
+	virtual void OnNodeDeactivation(FBehaviorTreeSearchData& SearchData, EBTNodeResult::Type NodeResult);
 
 	/** called when underlying node was processed (deactivated or failed to activate)
 	 * this function should be considered as const (don't modify state of object) if node is not instanced! */
-	virtual void OnNodeProcessed(struct FBehaviorTreeSearchData& SearchData, EBTNodeResult::Type& NodeResult);
+	virtual void OnNodeProcessed(FBehaviorTreeSearchData& SearchData, EBTNodeResult::Type& NodeResult);
 
 	/** calculates raw, core value of decorator's condition. Should not include calling IsInversed */
-	virtual bool CalculateRawConditionValue(class UBehaviorTreeComponent* OwnerComp, uint8* NodeMemory) const;
+	virtual bool CalculateRawConditionValue(UBehaviorTreeComponent* OwnerComp, uint8* NodeMemory) const;
 
-	friend class FBehaviorDecoratorDetails;
+	friend FBehaviorDecoratorDetails;
 };
 
 
diff --git a/Engine/Source/Runtime/AIModule/Classes/BehaviorTree/BTService.h b/Engine/Source/Runtime/AIModule/Classes/BehaviorTree/BTService.h
--- a/Engine/Source/Runtime/AIModule/Classes/BehaviorTree/BTService.h
+++ b/Engine/Source/Runtime/AIModule/Classes/BehaviorTree/BTService.h
@@ -37,7 +37,7 @@ class AIMODULE_API UBTService : public UBTAuxiliaryNode
 
 	virtual FString GetStaticDescription() const override;
 
-	void NotifyParentActivation(struct FBehaviorTreeSearchData& SearchData);
+	void NotifyParentActivation(FBehaviorTreeSearchData& SearchData);
 
 protected:
 
@@ -63,7 +63,7 @@ protected:
 	virtual void TickNode(UBehaviorTreeComponent* OwnerComp, uint8* NodeMemory, float DeltaSeconds) override;
 
 	/** called when search enters underlying branch */
-	virtual void OnSearchStart(struct FBehaviorTreeSearchData& SearchData);
+	virtual void OnSearchStart(FBehaviorTreeSearchData& SearchData);
 
 #if WITH_EDITOR
 	virtual FName GetNodeIconName() const override;
diff --git a/Engine/Source/Runtime/AIModule/Classes/BehaviorTree/BehaviorTreeManager.h b/Engine/Source/Runtime/AIModule/Classes/BehaviorTree/BehaviorTreeManager.h
--- a/Engine/Source/Runtime/AIModule/Classes/BehaviorTree/BehaviorTreeManager.h
+++ b/Engine/Source/Runtime/AIModule/Classes/BehaviorTree/BehaviorTreeManager.h
@@ -4,6 +4,9 @@
 #include "BehaviorTreeManager.generated.h"
 
 class UBehaviorTreeComponent;
+class UBTCompositeNode;
+class UBTDecorator;
+class UBehaviorTree;
 
 USTRUCT()
 struct FBehaviorTreeTemplateInfo
@@ -12,11 +15,11 @@ struct FBehaviorTreeTemplateInfo
 
 	/** behavior tree asset */
 	UPROPERTY()
-	class UBehaviorTree* Asset;
+	UBehaviorTree* Asset;
 
 	/** initialized template */
 	UPROPERTY(transient)
-	class UBTCompositeNode* Template;
+	UBTCompositeNode* Template;
 
 	/** size required for instance memory */
 	uint16 InstanceMemorySize;
@@ -32,13 +35,13 @@ class AIMODULE_API UBehaviorTreeManager : public UObject
 	int32 MaxDebuggerSteps;
 
 	/** get behavior tree template for given blueprint */
-	bool LoadTree(UBehaviorTree& Asset, class UBTCompositeNode*& Root, uint16& InstanceMemorySize);
+	bool LoadTree(UBehaviorTree& Asset, UBTCompositeNode*& Root, uint16& InstanceMemorySize);
 
 	/** get aligned memory size */
 	static int32 GetAlignedDataSize(int32 Size);
 
 	/** helper function for sorting and aligning node memory */
-	static void InitializeMemoryHelper(const TArray<class UBTDecorator*>& Nodes, TArray<uint16>& MemoryOffsets, int32& MemorySize);
+	static void InitializeMemoryHelper(const TArray<UBTDecorator*>& Nodes, TArray<uint16>& MemoryOffsets, int32& MemorySize);
 
 	/** cleanup hooks for map loading */
 	virtual void FinishDestroy() override;
diff --git a/Engine/Source/Runtime/AIModule/Classes/BehaviorTree/BehaviorTreeTypes.h b/Engine/Source/Runtime/AIModule/Classes/BehaviorTree/BehaviorTreeTypes.h
--- a/Engine/Source/Runtime/AIModule/Classes/BehaviorTree/BehaviorTreeTypes.h
+++ b/Engine/Source/Runtime/AIModule/Classes/BehaviorTree/BehaviorTreeTypes.h
@@ -4,6 +4,18 @@
 #include "BehaviorTreeTypes.generated.h"
 
 class UBlackboardData;
+class UBlackboardComponent;
+class UBehaviorTreeComponent;
+class UBTNode;
+class UBTTaskNode;
+class UBehaviorTree;
+class UBTCompositeNode;
+class UBTAuxiliaryNode;
+class UBlackboardKeyType;
+class FBlackboardDecoratorDetails;
+
+struct FBTNodeIndex;
+struct FBehaviorTreeSearchUpdate;
 
 // Visual logging helper
 #define BT_VLOG(Context, Verbosity, Format, ...) UE_VLOG(Context->OwnerComp.IsValid() ? Context->OwnerComp->GetOwner() : NULL, LogBehaviorTree, Verbosity, Format, ##__VA_ARGS__)
@@ -33,7 +45,7 @@ namespace FBlackboard
 }
 
 // delegate defines
-DECLARE_DELEGATE_TwoParams(FOnBlackboardChange, const class UBlackboardComponent*, FBlackboard::FKey /*key ID*/);
+DECLARE_DELEGATE_TwoParams(FOnBlackboardChange, const UBlackboardComponent*, FBlackboard::FKey /*key ID*/);
 
 namespace BTSpecialChild
 {
@@ -132,7 +144,7 @@ namespace EBTNodeUpdateMode
 struct FBehaviorTreeParallelTask
 {
 	/** worker object */
-	const class UBTTaskNode* TaskNode;
+	const UBTTaskNode* TaskNode;
 
 	/** additional mode data used for context switching */
 	TEnumAsByte<EBTTaskStatus::Type> Status;
@@ -178,10 +190,10 @@ struct FBehaviorTreeDebuggerInstance
 	FBehaviorTreeDebuggerInstance() : TreeAsset(NULL), RootNode(NULL) {}
 
 	/** behavior tree asset */
-	class UBehaviorTree* TreeAsset;
+	UBehaviorTree* TreeAsset;
 
 	/** root node in template */
-	class UBTCompositeNode* RootNode;
+	UBTCompositeNode* RootNode;
 
 	/** execution indices of active nodes */
 	TArray<uint16> ActivePath;
@@ -220,10 +232,10 @@ struct FBehaviorTreeExecutionStep
 struct FBehaviorTreeInstanceId
 {
 	/** behavior tree asset */
-	class UBehaviorTree* TreeAsset;
+	UBehaviorTree* TreeAsset;
 
 	/** root node in template for cleanup purposes */
-	class UBTCompositeNode* RootNode;
+	UBTCompositeNode* RootNode;
 
 	/** execution index path from root */
 	TArray<uint16> Path;
@@ -246,13 +258,13 @@ struct FBehaviorTreeInstanceId
 struct FBehaviorTreeInstance
 {
 	/** root node in template */
-	class UBTCompositeNode* RootNode;
+	UBTCompositeNode* RootNode;
 
 	/** active node in template */
-	class UBTNode* ActiveNode;
+	UBTNode* ActiveNode;
 
 	/** active auxiliary nodes */
-	TArray<class UBTAuxiliaryNode*> ActiveAuxNodes;
+	TArray<UBTAuxiliaryNode*> ActiveAuxNodes;
 
 	/** active parallel tasks */
 	TArray<FBehaviorTreeParallelTask> ParallelTasks;
@@ -285,18 +297,18 @@ struct FBehaviorTreeInstance
 #endif // STATS
 
 	/** initialize memory and create node instances */
-	void Initialize(class UBehaviorTreeComponent* OwnerComp, UBTCompositeNode* Node, int32& InstancedIndex, EBTMemoryInit::Type InitType);
+	void Initialize(UBehaviorTreeComponent* OwnerComp, UBTCompositeNode* Node, int32& InstancedIndex, EBTMemoryInit::Type InitType);
 
 	/** update injected nodes */
-	void InjectNodes(class UBehaviorTreeComponent* OwnerComp, UBTCompositeNode* Node, int32& InstancedIndex);
+	void InjectNodes(UBehaviorTreeComponent* OwnerComp, UBTCompositeNode* Node, int32& InstancedIndex);
 
 	/** cleanup node instances */
-	void Cleanup(class UBehaviorTreeComponent* OwnerComp, EBTMemoryClear::Type CleanupType);
+	void Cleanup(UBehaviorTreeComponent* OwnerComp, EBTMemoryClear::Type CleanupType);
 
 protected:
 
 	/** worker for updating all nodes */
-	void CleanupNodes(class UBehaviorTreeComponent* OwnerComp, UBTCompositeNode* Node, EBTMemoryClear::Type CleanupType);
+	void CleanupNodes(UBehaviorTreeComponent* OwnerComp, UBTCompositeNode* Node, EBTMemoryClear::Type CleanupType);
 };
 
 struct FBTNodeIndex
@@ -322,8 +334,8 @@ struct FBTNodeIndex
 /** node update data */
 struct FBehaviorTreeSearchUpdate
 {
-	class UBTAuxiliaryNode* AuxNode;
-	class UBTTaskNode* TaskNode;
+	UBTAuxiliaryNode* AuxNode;
+	UBTTaskNode* TaskNode;
 
 	uint16 InstanceIndex;
 
@@ -333,11 +345,11 @@ struct FBehaviorTreeSearchUpdate
 	uint8 bPostUpdate : 1;
 
 	FBehaviorTreeSearchUpdate() : AuxNode(0), TaskNode(0), InstanceIndex(0) {}
-	FBehaviorTreeSearchUpdate(const class UBTAuxiliaryNode* InAuxNode, uint16 InInstanceIndex, EBTNodeUpdateMode::Type InMode) :
-		AuxNode((class UBTAuxiliaryNode*)InAuxNode), TaskNode(0), InstanceIndex(InInstanceIndex), Mode(InMode) 
+	FBehaviorTreeSearchUpdate(const UBTAuxiliaryNode* InAuxNode, uint16 InInstanceIndex, EBTNodeUpdateMode::Type InMode) :
+		AuxNode((UBTAuxiliaryNode*)InAuxNode), TaskNode(0), InstanceIndex(InInstanceIndex), Mode(InMode) 
 	{}
-	FBehaviorTreeSearchUpdate(const class UBTTaskNode* InTaskNode, uint16 InInstanceIndex, EBTNodeUpdateMode::Type InMode) :
-		AuxNode(0), TaskNode((class UBTTaskNode*)InTaskNode), InstanceIndex(InInstanceIndex), Mode(InMode) 
+	FBehaviorTreeSearchUpdate(const UBTTaskNode* InTaskNode, uint16 InInstanceIndex, EBTNodeUpdateMode::Type InMode) :
+		AuxNode(0), TaskNode((UBTTaskNode*)InTaskNode), InstanceIndex(InInstanceIndex), Mode(InMode) 
 	{}
 };
 
@@ -345,17 +357,17 @@ struct FBehaviorTreeSearchUpdate
 struct FBehaviorTreeSearchData
 {
 	/** BT component */
-	class UBehaviorTreeComponent* OwnerComp;
+	UBehaviorTreeComponent* OwnerComp;
 
 	/** requested updates of additional nodes (preconditions, services, parallels)
 	 *  buffered during search to prevent instant add & remove pairs */
-	TArray<struct FBehaviorTreeSearchUpdate> PendingUpdates;
+	TArray<FBehaviorTreeSearchUpdate> PendingUpdates;
 
 	/** first node allowed in search */
-	struct FBTNodeIndex SearchStart;
+	FBTNodeIndex SearchStart;
 
 	/** last node allowed in search */
-	struct FBTNodeIndex SearchEnd;
+	FBTNodeIndex SearchEnd;
 
 	/** search unique number */
 	int32 SearchId;
@@ -396,7 +408,7 @@ struct AIMODULE_API FBlackboardKeySelector
 	/** array of allowed types with additional properties (e.g. uobject's base class) 
 	  * EditDefaults is required for FBlackboardSelectorDetails::CacheBlackboardData() */
 	UPROPERTY(transient, EditDefaultsOnly, BlueprintReadWrite, Category = Blackboard)
-	TArray<class UBlackboardKeyType*> AllowedTypes;
+	TArray<UBlackboardKeyType*> AllowedTypes;
 
 	/** name of selected key */
 	UPROPERTY(EditInstanceOnly, BlueprintReadWrite, Category = Blackboard)
@@ -404,7 +416,7 @@ struct AIMODULE_API FBlackboardKeySelector
 
 	/** class of selected key  */
 	UPROPERTY(transient, EditInstanceOnly, BlueprintReadWrite, Category = Blackboard)
-	TSubclassOf<class UBlackboardKeyType> SelectedKeyType;
+	TSubclassOf<UBlackboardKeyType> SelectedKeyType;
 
 protected:
 	/** ID of selected key */
@@ -443,7 +455,7 @@ public:
 
 	FORCEINLINE bool IsNone() const { return bNoneIsAllowedValue && SelectedKeyID == FBlackboard::InvalidKey; }
 
-	friend class FBlackboardDecoratorDetails;
+	friend FBlackboardDecoratorDetails;
 };
 
 UCLASS(Abstract)
diff --git a/Engine/Source/Runtime/AIModule/Classes/BehaviorTree/BlackboardComponent.h b/Engine/Source/Runtime/AIModule/Classes/BehaviorTree/BlackboardComponent.h
--- a/Engine/Source/Runtime/AIModule/Classes/BehaviorTree/BlackboardComponent.h
+++ b/Engine/Source/Runtime/AIModule/Classes/BehaviorTree/BlackboardComponent.h
@@ -19,6 +19,7 @@
 
 class UBlackboardData;
 class UBrainComponent;
+class UBlackboardKeyType;
 
 namespace EBlackboardDescription
 {
@@ -68,7 +69,7 @@ class AIMODULE_API UBlackboardComponent : public UActorComponent
 	void ResumeUpdates();
 
 	/** @return associated behavior tree component */
-	class UBrainComponent* GetBrainComponent() const;
+	UBrainComponent* GetBrainComponent() const;
 
 	/** @return blackboard data asset */
 	UBlackboardData* GetBlackboardAsset() const;
@@ -199,7 +200,7 @@ class AIMODULE_API UBlackboardComponent : public UActorComponent
 	FORCEINLINE bool IsValidKey(FBlackboard::FKey KeyID) const { check(BlackboardAsset); return KeyID != FBlackboard::InvalidKey && BlackboardAsset->Keys.IsValidIndex(KeyID); }
 
 	/** compares blackboard's values under specified keys */
-	EBlackboardCompare::Type CompareKeyValues(TSubclassOf<class UBlackboardKeyType> KeyType, FBlackboard::FKey KeyA, FBlackboard::FKey KeyB) const;
+	EBlackboardCompare::Type CompareKeyValues(TSubclassOf<UBlackboardKeyType> KeyType, FBlackboard::FKey KeyA, FBlackboard::FKey KeyB) const;
 
 	FString GetDebugInfoString(EBlackboardDescription::Type Mode) const;
 
diff --git a/Engine/Source/Runtime/AIModule/Classes/BehaviorTree/Composites/BTComposite_SimpleParallel.h b/Engine/Source/Runtime/AIModule/Classes/BehaviorTree/Composites/BTComposite_SimpleParallel.h
--- a/Engine/Source/Runtime/AIModule/Classes/BehaviorTree/Composites/BTComposite_SimpleParallel.h
+++ b/Engine/Source/Runtime/AIModule/Classes/BehaviorTree/Composites/BTComposite_SimpleParallel.h
@@ -4,6 +4,9 @@
 #include "BehaviorTree/BTCompositeNode.h"
 #include "BTComposite_SimpleParallel.generated.h"
 
+class UBehaviorTreeComponent;
+struct FBehaviorTreeSearchData;
+
 namespace EBTParallelChild
 {
 	enum Type
@@ -51,14 +54,14 @@ class AIMODULE_API UBTComposite_SimpleParallel : public UBTCompositeNode
 	TEnumAsByte<EBTParallelMode::Type> FinishMode;
 
 	/** handle child updates */
-	int32 GetNextChildHandler(struct FBehaviorTreeSearchData& SearchData, int32 PrevChild, EBTNodeResult::Type LastResult) const;
+	int32 GetNextChildHandler(FBehaviorTreeSearchData& SearchData, int32 PrevChild, EBTNodeResult::Type LastResult) const;
 
-	virtual void NotifyChildExecution(class UBehaviorTreeComponent* OwnerComp, uint8* NodeMemory, int32 ChildIdx, EBTNodeResult::Type& NodeResult) const override;
-	virtual void NotifyNodeDeactivation(struct FBehaviorTreeSearchData& SearchData, EBTNodeResult::Type& NodeResult) const override;
-	virtual bool CanPushSubtree(class UBehaviorTreeComponent* OwnerComp, uint8* NodeMemory, int32 ChildIdx) const override;
+	virtual void NotifyChildExecution(UBehaviorTreeComponent* OwnerComp, uint8* NodeMemory, int32 ChildIdx, EBTNodeResult::Type& NodeResult) const override;
+	virtual void NotifyNodeDeactivation(FBehaviorTreeSearchData& SearchData, EBTNodeResult::Type& NodeResult) const override;
+	virtual bool CanPushSubtree(UBehaviorTreeComponent* OwnerComp, uint8* NodeMemory, int32 ChildIdx) const override;
 	virtual uint16 GetInstanceMemorySize() const override;
 	virtual FString GetStaticDescription() const override;
-	virtual void DescribeRuntimeValues(const class UBehaviorTreeComponent* OwnerComp, uint8* NodeMemory, EBTDescriptionVerbosity::Type Verbosity, TArray<FString>& Values) const override;
+	virtual void DescribeRuntimeValues(const UBehaviorTreeComponent* OwnerComp, uint8* NodeMemory, EBTDescriptionVerbosity::Type Verbosity, TArray<FString>& Values) const override;
 
 	/** helper for showing values of EBTParallelMode enum */
 	static FString DescribeFinishMode(EBTParallelMode::Type Mode);
diff --git a/Engine/Source/Runtime/AIModule/Classes/BehaviorTree/Decorators/BTDecorator_Blackboard.h b/Engine/Source/Runtime/AIModule/Classes/BehaviorTree/Decorators/BTDecorator_Blackboard.h
--- a/Engine/Source/Runtime/AIModule/Classes/BehaviorTree/Decorators/BTDecorator_Blackboard.h
+++ b/Engine/Source/Runtime/AIModule/Classes/BehaviorTree/Decorators/BTDecorator_Blackboard.h
@@ -6,6 +6,10 @@
 #include "BTDecorator_BlackboardBase.h"
 #include "BTDecorator_Blackboard.generated.h"
 
+class UBlackboardComponent;
+class UBehaviorTreeComponent;
+class FBlackboardDecoratorDetails;
+
 /**
  *  Decorator for accessing blackboard values
  */
@@ -29,9 +33,9 @@ class AIMODULE_API UBTDecorator_Blackboard : public UBTDecorator_BlackboardBase
 {
 	GENERATED_UCLASS_BODY()
 
-	virtual bool CalculateRawConditionValue(class UBehaviorTreeComponent* OwnerComp, uint8* NodeMemory) const override;
-	virtual void OnBlackboardChange(const class UBlackboardComponent* Blackboard, FBlackboard::FKey ChangedKeyID) override;
-	virtual void DescribeRuntimeValues(const class UBehaviorTreeComponent* OwnerComp, uint8* NodeMemory, EBTDescriptionVerbosity::Type Verbosity, TArray<FString>& Values) const override;
+	virtual bool CalculateRawConditionValue(UBehaviorTreeComponent* OwnerComp, uint8* NodeMemory) const override;
+	virtual void OnBlackboardChange(const UBlackboardComponent* Blackboard, FBlackboard::FKey ChangedKeyID) override;
+	virtual void DescribeRuntimeValues(const UBehaviorTreeComponent* OwnerComp, uint8* NodeMemory, EBTDescriptionVerbosity::Type Verbosity, TArray<FString>& Values) const override;
 	virtual FString GetStaticDescription() const override;
 
 protected:
@@ -84,9 +88,7 @@ protected:
 #endif
 
 	/** take blackboard value and evaluate decorator's condition */
-	bool EvaluateOnBlackboard(const class UBlackboardComponent* BlackboardComp) const;
+	bool EvaluateOnBlackboard(const UBlackboardComponent* BlackboardComp) const;
 
-#if CPP
-	friend class FBlackboardDecoratorDetails;
-#endif
+	friend FBlackboardDecoratorDetails;
 };
diff --git a/Engine/Source/Runtime/AIModule/Classes/BehaviorTree/Decorators/BTDecorator_BlackboardBase.h b/Engine/Source/Runtime/AIModule/Classes/BehaviorTree/Decorators/BTDecorator_BlackboardBase.h
--- a/Engine/Source/Runtime/AIModule/Classes/BehaviorTree/Decorators/BTDecorator_BlackboardBase.h
+++ b/Engine/Source/Runtime/AIModule/Classes/BehaviorTree/Decorators/BTDecorator_BlackboardBase.h
@@ -3,6 +3,8 @@
 #include "BehaviorTree/BTDecorator.h"
 #include "BTDecorator_BlackboardBase.generated.h"
 
+class UBlackboardComponent;
+
 UCLASS(Abstract)
 class AIMODULE_API UBTDecorator_BlackboardBase : public UBTDecorator
 {
@@ -12,7 +14,7 @@ class AIMODULE_API UBTDecorator_BlackboardBase : public UBTDecorator
 	virtual void InitializeFromAsset(UBehaviorTree& Asset) override;
 
 	/** notify about change in blackboard keys */
-	virtual void OnBlackboardChange(const class UBlackboardComponent* Blackboard, FBlackboard::FKey ChangedKeyID);
+	virtual void OnBlackboardChange(const UBlackboardComponent* Blackboard, FBlackboard::FKey ChangedKeyID);
 
 #if WITH_EDITOR
 	virtual FName GetNodeIconName() const override;
@@ -25,7 +27,7 @@ protected:
 
 	/** blackboard key selector */
 	UPROPERTY(EditAnywhere, Category=Blackboard)
-	struct FBlackboardKeySelector BlackboardKey;
+	FBlackboardKeySelector BlackboardKey;
 
 	FOnBlackboardChange BBKeyObserver;
 
diff --git a/Engine/Source/Runtime/AIModule/Classes/BehaviorTree/Decorators/BTDecorator_BlueprintBase.h b/Engine/Source/Runtime/AIModule/Classes/BehaviorTree/Decorators/BTDecorator_BlueprintBase.h
--- a/Engine/Source/Runtime/AIModule/Classes/BehaviorTree/Decorators/BTDecorator_BlueprintBase.h
+++ b/Engine/Source/Runtime/AIModule/Classes/BehaviorTree/Decorators/BTDecorator_BlueprintBase.h
@@ -4,6 +4,11 @@
 #include "BehaviorTree/BTDecorator.h"
 #include "BTDecorator_BlueprintBase.generated.h"
 
+class UBlackboardComponent;
+class UBehaviorTreeComponent;
+class FBehaviorBlueprintDetails;
+struct FBehaviorTreeSearchData;
+
 /**
  *  Base class for blueprint based decorator nodes. Do NOT use it for creating native c++ classes!
  *
@@ -30,12 +35,12 @@ class AIMODULE_API UBTDecorator_BlueprintBase : public UBTDecorator
 	virtual void PostInitProperties() override;
 
 	/** notify about changes in blackboard */
-	void OnBlackboardChange(const class UBlackboardComponent* Blackboard, FBlackboard::FKey ChangedKeyID);
+	void OnBlackboardChange(const UBlackboardComponent* Blackboard, FBlackboard::FKey ChangedKeyID);
 
 	virtual FString GetStaticDescription() const override;
-	virtual void DescribeRuntimeValues(const class UBehaviorTreeComponent* OwnerComp, uint8* NodeMemory, EBTDescriptionVerbosity::Type Verbosity, TArray<FString>& Values) const override;
-	virtual bool CalculateRawConditionValue(class UBehaviorTreeComponent* OwnerComp, uint8* NodeMemory) const override;
-	virtual void OnInstanceDestroyed(class UBehaviorTreeComponent* OwnerComp) override;
+	virtual void DescribeRuntimeValues(const UBehaviorTreeComponent* OwnerComp, uint8* NodeMemory, EBTDescriptionVerbosity::Type Verbosity, TArray<FString>& Values) const override;
+	virtual bool CalculateRawConditionValue(UBehaviorTreeComponent* OwnerComp, uint8* NodeMemory) const override;
+	virtual void OnInstanceDestroyed(UBehaviorTreeComponent* OwnerComp) override;
 
 	virtual void SetOwner(AActor* ActorOwner) override;
 
@@ -88,8 +93,8 @@ protected:
 	
 	virtual void OnBecomeRelevant(UBehaviorTreeComponent* OwnerComp, uint8* NodeMemory) override;
 	virtual void OnCeaseRelevant(UBehaviorTreeComponent* OwnerComp, uint8* NodeMemory) override;
-	virtual void OnNodeActivation(struct FBehaviorTreeSearchData& SearchData) override;
-	virtual void OnNodeDeactivation(struct FBehaviorTreeSearchData& SearchData, EBTNodeResult::Type NodeResult) override;
+	virtual void OnNodeActivation(FBehaviorTreeSearchData& SearchData) override;
+	virtual void OnNodeDeactivation(FBehaviorTreeSearchData& SearchData, EBTNodeResult::Type NodeResult) override;
 	virtual void TickNode(UBehaviorTreeComponent* OwnerComp, uint8* NodeMemory, float DeltaSeconds) override;
 
 	/** tick function
@@ -182,5 +187,5 @@ protected:
 	UFUNCTION(BlueprintCallable, Category="AI|BehaviorTree")
 	bool IsDecoratorObserverActive() const;
 
-	friend class FBehaviorBlueprintDetails;
+	friend FBehaviorBlueprintDetails;
 };
diff --git a/Engine/Source/Runtime/AIModule/Classes/BehaviorTree/Decorators/BTDecorator_CompareBBEntries.h b/Engine/Source/Runtime/AIModule/Classes/BehaviorTree/Decorators/BTDecorator_CompareBBEntries.h
--- a/Engine/Source/Runtime/AIModule/Classes/BehaviorTree/Decorators/BTDecorator_CompareBBEntries.h
+++ b/Engine/Source/Runtime/AIModule/Classes/BehaviorTree/Decorators/BTDecorator_CompareBBEntries.h
@@ -4,6 +4,9 @@
 #include "BehaviorTree/BTDecorator.h"
 #include "BTDecorator_CompareBBEntries.generated.h"
 
+class UBlackboardComponent;
+struct FBlackboardKeySelector;
+
 UENUM()
 namespace EBlackBoardEntryComparison
 {
@@ -31,11 +34,11 @@ protected:
 
 	/** blackboard key selector */
 	UPROPERTY(EditAnywhere, Category=Blackboard)
-	struct FBlackboardKeySelector BlackboardKeyA;
+	FBlackboardKeySelector BlackboardKeyA;
 
 	/** blackboard key selector */
 	UPROPERTY(EditAnywhere, Category=Blackboard)
-	struct FBlackboardKeySelector BlackboardKeyB;
+	FBlackboardKeySelector BlackboardKeyB;
 
 	FOnBlackboardChange BBKeyObserver;
 
@@ -43,12 +46,12 @@ public:
 
 	virtual void PostInitProperties() override;
 
-	virtual bool CalculateRawConditionValue(class UBehaviorTreeComponent* OwnerComp, uint8* NodeMemory) const override;
+	virtual bool CalculateRawConditionValue(UBehaviorTreeComponent* OwnerComp, uint8* NodeMemory) const override;
 	virtual FString GetStaticDescription() const override;
 
 	virtual void InitializeFromAsset(UBehaviorTree& Asset) override;
 
-	virtual void OnBlackboardChange(const class UBlackboardComponent* Blackboard, FBlackboard::FKey ChangedKeyID);
+	virtual void OnBlackboardChange(const UBlackboardComponent* Blackboard, FBlackboard::FKey ChangedKeyID);
 	virtual void OnBecomeRelevant(UBehaviorTreeComponent* OwnerComp, uint8* NodeMemory) override;
 	virtual void OnCeaseRelevant(UBehaviorTreeComponent* OwnerComp, uint8* NodeMemory) override;
 
diff --git a/Engine/Source/Runtime/AIModule/Classes/BehaviorTree/Decorators/BTDecorator_ConditionalLoop.h b/Engine/Source/Runtime/AIModule/Classes/BehaviorTree/Decorators/BTDecorator_ConditionalLoop.h
--- a/Engine/Source/Runtime/AIModule/Classes/BehaviorTree/Decorators/BTDecorator_ConditionalLoop.h
+++ b/Engine/Source/Runtime/AIModule/Classes/BehaviorTree/Decorators/BTDecorator_ConditionalLoop.h
@@ -13,9 +13,9 @@ class AIMODULE_API UBTDecorator_ConditionalLoop : public UBTDecorator_Blackboard
 {
 	GENERATED_UCLASS_BODY()
 
-	virtual bool CalculateRawConditionValue(class UBehaviorTreeComponent* OwnerComp, uint8* NodeMemory) const override;
-	virtual void OnBlackboardChange(const class UBlackboardComponent* Blackboard, FBlackboard::FKey ChangedKeyID) override;
-	virtual void OnNodeDeactivation(struct FBehaviorTreeSearchData& SearchData, EBTNodeResult::Type NodeResult) override;
+	virtual bool CalculateRawConditionValue(UBehaviorTreeComponent* OwnerComp, uint8* NodeMemory) const override;
+	virtual void OnBlackboardChange(const UBlackboardComponent* Blackboard, FBlackboard::FKey ChangedKeyID) override;
+	virtual void OnNodeDeactivation(FBehaviorTreeSearchData& SearchData, EBTNodeResult::Type NodeResult) override;
 
 #if WITH_EDITOR
 	virtual FName GetNodeIconName() const override;
diff --git a/Engine/Source/Runtime/AIModule/Classes/BehaviorTree/Decorators/BTDecorator_ConeCheck.h b/Engine/Source/Runtime/AIModule/Classes/BehaviorTree/Decorators/BTDecorator_ConeCheck.h
--- a/Engine/Source/Runtime/AIModule/Classes/BehaviorTree/Decorators/BTDecorator_ConeCheck.h
+++ b/Engine/Source/Runtime/AIModule/Classes/BehaviorTree/Decorators/BTDecorator_ConeCheck.h
@@ -4,6 +4,9 @@
 #include "BehaviorTree/BTDecorator.h"
 #include "BTDecorator_ConeCheck.generated.h"
 
+class UBlackboardComponent;
+struct FBlackboardKeySelector;
+
 struct FBTConeCheckDecoratorMemory
 {
 	bool bLastRawResult;
@@ -26,21 +29,21 @@ class AIMODULE_API UBTDecorator_ConeCheck : public UBTDecorator
 	
 	/** blackboard key selector */
 	UPROPERTY(EditAnywhere, Category=Blackboard)
-	struct FBlackboardKeySelector ConeOrigin;
+	FBlackboardKeySelector ConeOrigin;
 
 	/** "None" means "use ConeOrigin's direction" */
 	UPROPERTY(EditAnywhere, Category=Blackboard)
-	struct FBlackboardKeySelector ConeDirection;
+	FBlackboardKeySelector ConeDirection;
 
 	/** blackboard key selector */
 	UPROPERTY(EditAnywhere, Category=Blackboard)
-	struct FBlackboardKeySelector Observed;
+	FBlackboardKeySelector Observed;
 	
 	float ConeHalfAngleDot;
 
 	virtual void InitializeFromAsset(UBehaviorTree& Asset) override;
 	virtual uint16 GetInstanceMemorySize() const override;
-	virtual void DescribeRuntimeValues(const class UBehaviorTreeComponent* OwnerComp, uint8* NodeMemory, EBTDescriptionVerbosity::Type Verbosity, TArray<FString>& Values) const override;
+	virtual void DescribeRuntimeValues(const UBehaviorTreeComponent* OwnerComp, uint8* NodeMemory, EBTDescriptionVerbosity::Type Verbosity, TArray<FString>& Values) const override;
 	virtual FString GetStaticDescription() const override;
 
 #if WITH_EDITOR
@@ -49,14 +52,14 @@ class AIMODULE_API UBTDecorator_ConeCheck : public UBTDecorator
 
 protected:
 
-	virtual bool CalculateRawConditionValue(class UBehaviorTreeComponent* OwnerComp, uint8* NodeMemory) const override;
+	virtual bool CalculateRawConditionValue(UBehaviorTreeComponent* OwnerComp, uint8* NodeMemory) const override;
 	virtual void OnBecomeRelevant(UBehaviorTreeComponent* OwnerComp, uint8* NodeMemory) override;
-	void OnBlackboardChange(const class UBlackboardComponent* Blackboard, FBlackboard::FKey ChangedKeyID);
+	void OnBlackboardChange(const UBlackboardComponent* Blackboard, FBlackboard::FKey ChangedKeyID);
 
 	virtual void TickNode(UBehaviorTreeComponent* OwnerComp, uint8* NodeMemory, float DeltaSeconds) override;
 	
 	bool CalculateDirection(const UBlackboardComponent* BlackboardComp, const FBlackboardKeySelector& Origin, const FBlackboardKeySelector& End, FVector& Direction) const;
 
 private:
-	bool CalcConditionImpl(class UBehaviorTreeComponent* OwnerComp, uint8* NodeMemory) const;
+	bool CalcConditionImpl(UBehaviorTreeComponent* OwnerComp, uint8* NodeMemory) const;
 };
diff --git a/Engine/Source/Runtime/AIModule/Classes/BehaviorTree/Decorators/BTDecorator_Cooldown.h b/Engine/Source/Runtime/AIModule/Classes/BehaviorTree/Decorators/BTDecorator_Cooldown.h
--- a/Engine/Source/Runtime/AIModule/Classes/BehaviorTree/Decorators/BTDecorator_Cooldown.h
+++ b/Engine/Source/Runtime/AIModule/Classes/BehaviorTree/Decorators/BTDecorator_Cooldown.h
@@ -27,10 +27,10 @@ class AIMODULE_API UBTDecorator_Cooldown : public UBTDecorator
 	virtual void PostLoad() override;
 	// End UObject Interface
 
-	virtual bool CalculateRawConditionValue(class UBehaviorTreeComponent* OwnerComp, uint8* NodeMemory) const override;
-	virtual void InitializeMemory(class UBehaviorTreeComponent* OwnerComp, uint8* NodeMemory, EBTMemoryInit::Type InitType) const override;
+	virtual bool CalculateRawConditionValue(UBehaviorTreeComponent* OwnerComp, uint8* NodeMemory) const override;
+	virtual void InitializeMemory(UBehaviorTreeComponent* OwnerComp, uint8* NodeMemory, EBTMemoryInit::Type InitType) const override;
 	virtual uint16 GetInstanceMemorySize() const override;
-	virtual void DescribeRuntimeValues(const class UBehaviorTreeComponent* OwnerComp, uint8* NodeMemory, EBTDescriptionVerbosity::Type Verbosity, TArray<FString>& Values) const override;
+	virtual void DescribeRuntimeValues(const UBehaviorTreeComponent* OwnerComp, uint8* NodeMemory, EBTDescriptionVerbosity::Type Verbosity, TArray<FString>& Values) const override;
 	virtual FString GetStaticDescription() const override;
 
 #if WITH_EDITOR
@@ -39,6 +39,6 @@ class AIMODULE_API UBTDecorator_Cooldown : public UBTDecorator
 
 protected:
 
-	virtual void OnNodeDeactivation(struct FBehaviorTreeSearchData& SearchData, EBTNodeResult::Type NodeResult) override;
-	virtual void TickNode(class UBehaviorTreeComponent* OwnerComp, uint8* NodeMemory, float DeltaSeconds) override;
+	virtual void OnNodeDeactivation(FBehaviorTreeSearchData& SearchData, EBTNodeResult::Type NodeResult) override;
+	virtual void TickNode(UBehaviorTreeComponent* OwnerComp, uint8* NodeMemory, float DeltaSeconds) override;
 };
diff --git a/Engine/Source/Runtime/AIModule/Classes/BehaviorTree/Decorators/BTDecorator_DoesPathExist.h b/Engine/Source/Runtime/AIModule/Classes/BehaviorTree/Decorators/BTDecorator_DoesPathExist.h
--- a/Engine/Source/Runtime/AIModule/Classes/BehaviorTree/Decorators/BTDecorator_DoesPathExist.h
+++ b/Engine/Source/Runtime/AIModule/Classes/BehaviorTree/Decorators/BTDecorator_DoesPathExist.h
@@ -4,6 +4,8 @@
 #include "BehaviorTree/BTDecorator.h"
 #include "BTDecorator_DoesPathExist.generated.h"
 
+class UNavigationQueryFilter;
+
 UENUM()
 namespace EPathExistanceQueryType
 {
@@ -28,11 +30,11 @@ protected:
 
 	/** blackboard key selector */
 	UPROPERTY(EditAnywhere, Category=Condition)
-	struct FBlackboardKeySelector BlackboardKeyA;
+	FBlackboardKeySelector BlackboardKeyA;
 
 	/** blackboard key selector */
 	UPROPERTY(EditAnywhere, Category=Condition)
-	struct FBlackboardKeySelector BlackboardKeyB;
+	FBlackboardKeySelector BlackboardKeyB;
 
 public:
 
@@ -45,9 +47,9 @@ public:
 
 	/** "None" will result in default filter being used */
 	UPROPERTY(Category=Node, EditAnywhere)
-	TSubclassOf<class UNavigationQueryFilter> FilterClass;
+	TSubclassOf<UNavigationQueryFilter> FilterClass;
 
-	virtual bool CalculateRawConditionValue(class UBehaviorTreeComponent* OwnerComp, uint8* NodeMemory) const override;
+	virtual bool CalculateRawConditionValue(UBehaviorTreeComponent* OwnerComp, uint8* NodeMemory) const override;
 	virtual FString GetStaticDescription() const override;
 	virtual void InitializeFromAsset(UBehaviorTree& Asset) override;
 
diff --git a/Engine/Source/Runtime/AIModule/Classes/BehaviorTree/Decorators/BTDecorator_KeepInCone.h b/Engine/Source/Runtime/AIModule/Classes/BehaviorTree/Decorators/BTDecorator_KeepInCone.h
--- a/Engine/Source/Runtime/AIModule/Classes/BehaviorTree/Decorators/BTDecorator_KeepInCone.h
+++ b/Engine/Source/Runtime/AIModule/Classes/BehaviorTree/Decorators/BTDecorator_KeepInCone.h
@@ -26,11 +26,11 @@ class AIMODULE_API UBTDecorator_KeepInCone : public UBTDecorator
 	
 	/** blackboard key selector */
 	UPROPERTY(EditAnywhere, Category=Blackboard)
-	struct FBlackboardKeySelector ConeOrigin;
+	FBlackboardKeySelector ConeOrigin;
 
 	/** blackboard key selector */
 	UPROPERTY(EditAnywhere, Category=Blackboard)
-	struct FBlackboardKeySelector Observed;
+	FBlackboardKeySelector Observed;
 
 	// deprecated, set value of ConeOrigin on initialization
 	UPROPERTY()
@@ -44,7 +44,7 @@ class AIMODULE_API UBTDecorator_KeepInCone : public UBTDecorator
 
 	virtual void InitializeFromAsset(UBehaviorTree& Asset) override;
 	virtual uint16 GetInstanceMemorySize() const override;
-	virtual void DescribeRuntimeValues(const class UBehaviorTreeComponent* OwnerComp, uint8* NodeMemory, EBTDescriptionVerbosity::Type Verbosity, TArray<FString>& Values) const override;
+	virtual void DescribeRuntimeValues(const UBehaviorTreeComponent* OwnerComp, uint8* NodeMemory, EBTDescriptionVerbosity::Type Verbosity, TArray<FString>& Values) const override;
 	virtual FString GetStaticDescription() const override;
 
 #if WITH_EDITOR
diff --git a/Engine/Source/Runtime/AIModule/Classes/BehaviorTree/Decorators/BTDecorator_Loop.h b/Engine/Source/Runtime/AIModule/Classes/BehaviorTree/Decorators/BTDecorator_Loop.h
--- a/Engine/Source/Runtime/AIModule/Classes/BehaviorTree/Decorators/BTDecorator_Loop.h
+++ b/Engine/Source/Runtime/AIModule/Classes/BehaviorTree/Decorators/BTDecorator_Loop.h
@@ -27,7 +27,7 @@ class AIMODULE_API UBTDecorator_Loop : public UBTDecorator
 	bool bInfiniteLoop;
 
 	virtual uint16 GetInstanceMemorySize() const override;
-	virtual void DescribeRuntimeValues(const class UBehaviorTreeComponent* OwnerComp, uint8* NodeMemory, EBTDescriptionVerbosity::Type Verbosity, TArray<FString>& Values) const override;
+	virtual void DescribeRuntimeValues(const UBehaviorTreeComponent* OwnerComp, uint8* NodeMemory, EBTDescriptionVerbosity::Type Verbosity, TArray<FString>& Values) const override;
 	virtual FString GetStaticDescription() const override;
 
 #if WITH_EDITOR
@@ -36,5 +36,5 @@ class AIMODULE_API UBTDecorator_Loop : public UBTDecorator
 
 protected:
 
-	virtual void OnNodeActivation(struct FBehaviorTreeSearchData& SearchData) override;
+	virtual void OnNodeActivation(FBehaviorTreeSearchData& SearchData) override;
 };
diff --git a/Engine/Source/Runtime/AIModule/Classes/BehaviorTree/Decorators/BTDecorator_ReachedMoveGoal.h b/Engine/Source/Runtime/AIModule/Classes/BehaviorTree/Decorators/BTDecorator_ReachedMoveGoal.h
--- a/Engine/Source/Runtime/AIModule/Classes/BehaviorTree/Decorators/BTDecorator_ReachedMoveGoal.h
+++ b/Engine/Source/Runtime/AIModule/Classes/BehaviorTree/Decorators/BTDecorator_ReachedMoveGoal.h
@@ -13,7 +13,7 @@ class AIMODULE_API UBTDecorator_ReachedMoveGoal : public UBTDecorator
 {
 	GENERATED_UCLASS_BODY()
 
-	virtual bool CalculateRawConditionValue(class UBehaviorTreeComponent* OwnerComp, uint8* NodeMemory) const override;
+	virtual bool CalculateRawConditionValue(UBehaviorTreeComponent* OwnerComp, uint8* NodeMemory) const override;
 
 #if WITH_EDITOR
 	virtual FName GetNodeIconName() const override;
diff --git a/Engine/Source/Runtime/AIModule/Classes/BehaviorTree/Decorators/BTDecorator_TimeLimit.h b/Engine/Source/Runtime/AIModule/Classes/BehaviorTree/Decorators/BTDecorator_TimeLimit.h
--- a/Engine/Source/Runtime/AIModule/Classes/BehaviorTree/Decorators/BTDecorator_TimeLimit.h
+++ b/Engine/Source/Runtime/AIModule/Classes/BehaviorTree/Decorators/BTDecorator_TimeLimit.h
@@ -16,7 +16,7 @@ class AIMODULE_API UBTDecorator_TimeLimit : public UBTDecorator
 	UPROPERTY(Category=Decorator, EditAnywhere)
 	float TimeLimit;
 
-	virtual void DescribeRuntimeValues(const class UBehaviorTreeComponent* OwnerComp, uint8* NodeMemory, EBTDescriptionVerbosity::Type Verbosity, TArray<FString>& Values) const override;
+	virtual void DescribeRuntimeValues(const UBehaviorTreeComponent* OwnerComp, uint8* NodeMemory, EBTDescriptionVerbosity::Type Verbosity, TArray<FString>& Values) const override;
 	virtual FString GetStaticDescription() const override;
 
 #if WITH_EDITOR
diff --git a/Engine/Source/Runtime/AIModule/Classes/BehaviorTree/Services/BTService_BlueprintBase.h b/Engine/Source/Runtime/AIModule/Classes/BehaviorTree/Services/BTService_BlueprintBase.h
--- a/Engine/Source/Runtime/AIModule/Classes/BehaviorTree/Services/BTService_BlueprintBase.h
+++ b/Engine/Source/Runtime/AIModule/Classes/BehaviorTree/Services/BTService_BlueprintBase.h
@@ -4,6 +4,8 @@
 #include "BehaviorTree/BTService.h"
 #include "BTService_BlueprintBase.generated.h"
 
+class FBehaviorBlueprintDetails;
+
 /**
  *  Base class for blueprint based service nodes. Do NOT use it for creating native c++ classes!
  *
@@ -19,8 +21,8 @@ class AIMODULE_API UBTService_BlueprintBase : public UBTService
 
 	virtual void PostInitProperties() override;
 
-	virtual void DescribeRuntimeValues(const class UBehaviorTreeComponent* OwnerComp, uint8* NodeMemory, EBTDescriptionVerbosity::Type Verbosity, TArray<FString>& Values) const override;
-	virtual void OnInstanceDestroyed(class UBehaviorTreeComponent* OwnerComp) override;
+	virtual void DescribeRuntimeValues(const UBehaviorTreeComponent* OwnerComp, uint8* NodeMemory, EBTDescriptionVerbosity::Type Verbosity, TArray<FString>& Values) const override;
+	virtual void OnInstanceDestroyed(UBehaviorTreeComponent* OwnerComp) override;
 
 	virtual void SetOwner(AActor* ActorOwner) override;
 
@@ -59,10 +61,10 @@ protected:
 	/** set if ReceiveSearchStart is implemented by blueprint */
 	uint32 ReceiveSearchStartImplementations : 2;
 
-	virtual void OnBecomeRelevant(class UBehaviorTreeComponent* OwnerComp, uint8* NodeMemory) override;
-	virtual void OnCeaseRelevant(class UBehaviorTreeComponent* OwnerComp, uint8* NodeMemory) override;
-	virtual void TickNode(class UBehaviorTreeComponent* OwnerComp, uint8* NodeMemory, float DeltaSeconds) override;
-	virtual void OnSearchStart(struct FBehaviorTreeSearchData& SearchData) override;
+	virtual void OnBecomeRelevant(UBehaviorTreeComponent* OwnerComp, uint8* NodeMemory) override;
+	virtual void OnCeaseRelevant(UBehaviorTreeComponent* OwnerComp, uint8* NodeMemory) override;
+	virtual void TickNode(UBehaviorTreeComponent* OwnerComp, uint8* NodeMemory, float DeltaSeconds) override;
+	virtual void OnSearchStart(FBehaviorTreeSearchData& SearchData) override;
 
 	/** tick function
 	 *	@Note that if both generic and AI event versions are implemented only the more
@@ -120,5 +122,5 @@ protected:
 	UFUNCTION(BlueprintCallable, Category="AI|BehaviorTree")
 	bool IsServiceActive() const;
 
-	friend class FBehaviorBlueprintDetails;
+	friend FBehaviorBlueprintDetails;
 };
diff --git a/Engine/Source/Runtime/AIModule/Classes/BehaviorTree/Services/BTService_DefaultFocus.h b/Engine/Source/Runtime/AIModule/Classes/BehaviorTree/Services/BTService_DefaultFocus.h
--- a/Engine/Source/Runtime/AIModule/Classes/BehaviorTree/Services/BTService_DefaultFocus.h
+++ b/Engine/Source/Runtime/AIModule/Classes/BehaviorTree/Services/BTService_DefaultFocus.h
@@ -26,11 +26,11 @@ class AIMODULE_API UBTService_DefaultFocus : public UBTService_BlackboardBase
 	GENERATED_UCLASS_BODY()
 
 	virtual uint16 GetInstanceMemorySize() const override { return sizeof(FBTFocusMemory); }
-	virtual void OnBecomeRelevant(class UBehaviorTreeComponent* OwnerComp, uint8* NodeMemory) override;
-	virtual void OnCeaseRelevant(class UBehaviorTreeComponent* OwnerComp, uint8* NodeMemory) override;
+	virtual void OnBecomeRelevant(UBehaviorTreeComponent* OwnerComp, uint8* NodeMemory) override;
+	virtual void OnCeaseRelevant(UBehaviorTreeComponent* OwnerComp, uint8* NodeMemory) override;
 
 	virtual FString GetStaticDescription() const override;
-	virtual void DescribeRuntimeValues(const class UBehaviorTreeComponent* OwnerComp, uint8* NodeMemory, EBTDescriptionVerbosity::Type Verbosity, TArray<FString>& Values) const override;
+	virtual void DescribeRuntimeValues(const UBehaviorTreeComponent* OwnerComp, uint8* NodeMemory, EBTDescriptionVerbosity::Type Verbosity, TArray<FString>& Values) const override;
 
 #if WITH_EDITOR
 	virtual FName GetNodeIconName() const override;
diff --git a/Engine/Source/Runtime/AIModule/Classes/BehaviorTree/Tasks/BTTask_BlueprintBase.h b/Engine/Source/Runtime/AIModule/Classes/BehaviorTree/Tasks/BTTask_BlueprintBase.h
--- a/Engine/Source/Runtime/AIModule/Classes/BehaviorTree/Tasks/BTTask_BlueprintBase.h
+++ b/Engine/Source/Runtime/AIModule/Classes/BehaviorTree/Tasks/BTTask_BlueprintBase.h
@@ -1,8 +1,12 @@
 // Copyright 1998-2014 Epic Games, Inc. All Rights Reserved.
 
 #pragma once
+
+#include "BehaviorTree/BTTaskNode.h"
 #include "BTTask_BlueprintBase.generated.h"
 
+class FBehaviorBlueprintDetails;
+
 /**
  *  Base class for blueprint based task nodes. Do NOT use it for creating native c++ classes!
  *
@@ -19,12 +23,12 @@ class AIMODULE_API UBTTask_BlueprintBase : public UBTTaskNode
 	/** setup node name */
 	virtual void PostInitProperties() override;
 
-	virtual EBTNodeResult::Type ExecuteTask(class UBehaviorTreeComponent* OwnerComp, uint8* NodeMemory) override;
-	virtual EBTNodeResult::Type AbortTask(class UBehaviorTreeComponent* OwnerComp, uint8* NodeMemory) override;
+	virtual EBTNodeResult::Type ExecuteTask(UBehaviorTreeComponent* OwnerComp, uint8* NodeMemory) override;
+	virtual EBTNodeResult::Type AbortTask(UBehaviorTreeComponent* OwnerComp, uint8* NodeMemory) override;
 
 	virtual FString GetStaticDescription() const override;
-	virtual void DescribeRuntimeValues(const class UBehaviorTreeComponent* OwnerComp, uint8* NodeMemory, EBTDescriptionVerbosity::Type Verbosity, TArray<FString>& Values) const override;
-	virtual void OnInstanceDestroyed(class UBehaviorTreeComponent* OwnerComp) override;
+	virtual void DescribeRuntimeValues(const UBehaviorTreeComponent* OwnerComp, uint8* NodeMemory, EBTDescriptionVerbosity::Type Verbosity, TArray<FString>& Values) const override;
+	virtual void OnInstanceDestroyed(UBehaviorTreeComponent* OwnerComp) override;
 
 	virtual void SetOwner(AActor* ActorOwner) override;
 
@@ -124,7 +128,7 @@ protected:
 	bool IsTaskExecuting() const;
 	
 	/** ticks this task */
-	virtual void TickTask(class UBehaviorTreeComponent* OwnerComp, uint8* NodeMemory, float DeltaSeconds) override;
+	virtual void TickTask(UBehaviorTreeComponent* OwnerComp, uint8* NodeMemory, float DeltaSeconds) override;
 
-	friend class FBehaviorBlueprintDetails;
+	friend FBehaviorBlueprintDetails;
 };
diff --git a/Engine/Source/Runtime/AIModule/Classes/BehaviorTree/Tasks/BTTask_MoveDirectlyToward.h b/Engine/Source/Runtime/AIModule/Classes/BehaviorTree/Tasks/BTTask_MoveDirectlyToward.h
--- a/Engine/Source/Runtime/AIModule/Classes/BehaviorTree/Tasks/BTTask_MoveDirectlyToward.h
+++ b/Engine/Source/Runtime/AIModule/Classes/BehaviorTree/Tasks/BTTask_MoveDirectlyToward.h
@@ -34,10 +34,10 @@ class AIMODULE_API UBTTask_MoveDirectlyToward : public UBTTask_BlackboardBase
 	UPROPERTY(Category=Node, EditAnywhere)
 	uint32 bAllowStrafe : 1;
 
-	virtual EBTNodeResult::Type ExecuteTask(class UBehaviorTreeComponent* OwnerComp, uint8* NodeMemory) override;
-	virtual EBTNodeResult::Type AbortTask(class UBehaviorTreeComponent* OwnerComp, uint8* NodeMemory) override;
+	virtual EBTNodeResult::Type ExecuteTask(UBehaviorTreeComponent* OwnerComp, uint8* NodeMemory) override;
+	virtual EBTNodeResult::Type AbortTask(UBehaviorTreeComponent* OwnerComp, uint8* NodeMemory) override;
 	virtual uint16 GetInstanceMemorySize() const override;
-	virtual void DescribeRuntimeValues(const class UBehaviorTreeComponent* OwnerComp, uint8* NodeMemory, EBTDescriptionVerbosity::Type Verbosity, TArray<FString>& Values) const override;
+	virtual void DescribeRuntimeValues(const UBehaviorTreeComponent* OwnerComp, uint8* NodeMemory, EBTDescriptionVerbosity::Type Verbosity, TArray<FString>& Values) const override;
 	virtual FString GetStaticDescription() const override;
 
 #if WITH_EDITOR
diff --git a/Engine/Source/Runtime/AIModule/Classes/BehaviorTree/Tasks/BTTask_MoveTo.h b/Engine/Source/Runtime/AIModule/Classes/BehaviorTree/Tasks/BTTask_MoveTo.h
--- a/Engine/Source/Runtime/AIModule/Classes/BehaviorTree/Tasks/BTTask_MoveTo.h
+++ b/Engine/Source/Runtime/AIModule/Classes/BehaviorTree/Tasks/BTTask_MoveTo.h
@@ -4,6 +4,8 @@
 #include "BehaviorTree/Tasks/BTTask_BlackboardBase.h"
 #include "BTTask_MoveTo.generated.h"
 
+class UNavigationQueryFilter;
+
 struct FBTMoveToTaskMemory
 {
 	/** Move request ID */
@@ -26,19 +28,19 @@ class AIMODULE_API UBTTask_MoveTo : public UBTTask_BlackboardBase
 
 	/** "None" will result in default filter being used */
 	UPROPERTY(Category=Node, EditAnywhere)
-	TSubclassOf<class UNavigationQueryFilter> FilterClass;
+	TSubclassOf<UNavigationQueryFilter> FilterClass;
 
 	UPROPERTY(Category=Node, EditAnywhere)
 	uint32 bAllowStrafe : 1;
 	
-	virtual EBTNodeResult::Type ExecuteTask(class UBehaviorTreeComponent* OwnerComp, uint8* NodeMemory) override;
-	virtual EBTNodeResult::Type AbortTask(class UBehaviorTreeComponent* OwnerComp, uint8* NodeMemory) override;
-	virtual void TickTask(class UBehaviorTreeComponent* OwnerComp, uint8* NodeMemory, float DeltaSeconds) override;
+	virtual EBTNodeResult::Type ExecuteTask(UBehaviorTreeComponent* OwnerComp, uint8* NodeMemory) override;
+	virtual EBTNodeResult::Type AbortTask(UBehaviorTreeComponent* OwnerComp, uint8* NodeMemory) override;
+	virtual void TickTask(UBehaviorTreeComponent* OwnerComp, uint8* NodeMemory, float DeltaSeconds) override;
 	virtual uint16 GetInstanceMemorySize() const override;
 
-	virtual void OnMessage(class UBehaviorTreeComponent* OwnerComp, uint8* NodeMemory, FName Message, int32 RequestID, bool bSuccess) override;
+	virtual void OnMessage(UBehaviorTreeComponent* OwnerComp, uint8* NodeMemory, FName Message, int32 RequestID, bool bSuccess) override;
 
-	virtual void DescribeRuntimeValues(const class UBehaviorTreeComponent* OwnerComp, uint8* NodeMemory, EBTDescriptionVerbosity::Type Verbosity, TArray<FString>& Values) const override;
+	virtual void DescribeRuntimeValues(const UBehaviorTreeComponent* OwnerComp, uint8* NodeMemory, EBTDescriptionVerbosity::Type Verbosity, TArray<FString>& Values) const override;
 	virtual FString GetStaticDescription() const override;
 
 #if WITH_EDITOR
@@ -47,5 +49,5 @@ class AIMODULE_API UBTTask_MoveTo : public UBTTask_BlackboardBase
 
 protected:
 
-	EBTNodeResult::Type PerformMoveTask(class UBehaviorTreeComponent* OwnerComp, uint8* NodeMemory);
+	EBTNodeResult::Type PerformMoveTask(UBehaviorTreeComponent* OwnerComp, uint8* NodeMemory);
 };
diff --git a/Engine/Source/Runtime/AIModule/Classes/BehaviorTree/Tasks/BTTask_RunBehavior.h b/Engine/Source/Runtime/AIModule/Classes/BehaviorTree/Tasks/BTTask_RunBehavior.h
--- a/Engine/Source/Runtime/AIModule/Classes/BehaviorTree/Tasks/BTTask_RunBehavior.h
+++ b/Engine/Source/Runtime/AIModule/Classes/BehaviorTree/Tasks/BTTask_RunBehavior.h
@@ -25,7 +25,7 @@ class AIMODULE_API UBTTask_RunBehavior : public UBTTaskNode
 
 	virtual EBTNodeResult::Type ExecuteTask(UBehaviorTreeComponent* OwnerComp, uint8* NodeMemory) override;
 	virtual uint16 GetInstanceMemorySize() const override;
-	virtual void CleanupMemory(class UBehaviorTreeComponent* OwnerComp, uint8* NodeMemory, EBTMemoryClear::Type CleanupType) const override;
+	virtual void CleanupMemory(UBehaviorTreeComponent* OwnerComp, uint8* NodeMemory, EBTMemoryClear::Type CleanupType) const override;
 	virtual FString GetStaticDescription() const override;
 
 #if WITH_EDITOR
@@ -39,19 +39,19 @@ class AIMODULE_API UBTTask_RunBehavior : public UBTTaskNode
 	int32 GetInjectedNodesCount() const;
 
 	/** @returns subtree asset */
-	class UBehaviorTree* GetSubtreeAsset() const;
+	UBehaviorTree* GetSubtreeAsset() const;
 
 protected:
 
 	/** behavior to run */
 	UPROPERTY(Category = Node, EditAnywhere)
-	class UBehaviorTree* BehaviorAsset;
+	UBehaviorTree* BehaviorAsset;
 };
 
 //////////////////////////////////////////////////////////////////////////
 // Inlines
 
-FORCEINLINE class UBehaviorTree* UBTTask_RunBehavior::GetSubtreeAsset() const
+FORCEINLINE UBehaviorTree* UBTTask_RunBehavior::GetSubtreeAsset() const
 {
 	return BehaviorAsset;
 }
diff --git a/Engine/Source/Runtime/AIModule/Classes/BehaviorTree/Tasks/BTTask_RunEQSQuery.h b/Engine/Source/Runtime/AIModule/Classes/BehaviorTree/Tasks/BTTask_RunEQSQuery.h
--- a/Engine/Source/Runtime/AIModule/Classes/BehaviorTree/Tasks/BTTask_RunEQSQuery.h
+++ b/Engine/Source/Runtime/AIModule/Classes/BehaviorTree/Tasks/BTTask_RunEQSQuery.h
@@ -1,9 +1,11 @@
 // Copyright 1998-2014 Epic Games, Inc. All Rights Reserved.
 
 #pragma once
-#include "../../EnvironmentQuery/EnvQueryTypes.h"
+#include "EnvironmentQuery/EnvQueryTypes.h"
 #include "BTTask_RunEQSQuery.generated.h"
 
+class UEnvQuery;
+
 struct FBTEnvQueryTaskMemory
 {
 	/** Query request ID */
@@ -21,21 +23,21 @@ class AIMODULE_API UBTTask_RunEQSQuery : public UBTTask_BlackboardBase
 
 	/** query to run */
 	UPROPERTY(Category=Node, EditAnywhere)
-	class UEnvQuery* QueryTemplate;
+	UEnvQuery* QueryTemplate;
 
 	/** optional parameters for query */
 	UPROPERTY(Category=Node, EditAnywhere)
-	TArray<struct FEnvNamedValue> QueryParams;
+	TArray<FEnvNamedValue> QueryParams;
 
-	virtual EBTNodeResult::Type ExecuteTask(class UBehaviorTreeComponent* OwnerComp, uint8* NodeMemory) override;
-	virtual EBTNodeResult::Type AbortTask(class UBehaviorTreeComponent* OwnerComp, uint8* NodeMemory) override;
+	virtual EBTNodeResult::Type ExecuteTask(UBehaviorTreeComponent* OwnerComp, uint8* NodeMemory) override;
+	virtual EBTNodeResult::Type AbortTask(UBehaviorTreeComponent* OwnerComp, uint8* NodeMemory) override;
 
-	virtual void DescribeRuntimeValues(const class UBehaviorTreeComponent* OwnerComp, uint8* NodeMemory, EBTDescriptionVerbosity::Type Verbosity, TArray<FString>& Values) const override;
+	virtual void DescribeRuntimeValues(const UBehaviorTreeComponent* OwnerComp, uint8* NodeMemory, EBTDescriptionVerbosity::Type Verbosity, TArray<FString>& Values) const override;
 	virtual FString GetStaticDescription() const override;
 	virtual uint16 GetInstanceMemorySize() const override;
 
 	/** finish task */
-	void OnQueryFinished(TSharedPtr<struct FEnvQueryResult> Result);
+	void OnQueryFinished(TSharedPtr<FEnvQueryResult> Result);
 
 #if WITH_EDITOR
 	/** prepare query params */
diff --git a/Engine/Source/Runtime/AIModule/Classes/BehaviorTree/Tasks/BTTask_Wait.h b/Engine/Source/Runtime/AIModule/Classes/BehaviorTree/Tasks/BTTask_Wait.h
--- a/Engine/Source/Runtime/AIModule/Classes/BehaviorTree/Tasks/BTTask_Wait.h
+++ b/Engine/Source/Runtime/AIModule/Classes/BehaviorTree/Tasks/BTTask_Wait.h
@@ -22,9 +22,9 @@ class AIMODULE_API UBTTask_Wait : public UBTTaskNode
 	UPROPERTY(Category=Wait, EditAnywhere)
 	float WaitTime;
 
-	virtual EBTNodeResult::Type ExecuteTask(class UBehaviorTreeComponent* OwnerComp, uint8* NodeMemory) override;
+	virtual EBTNodeResult::Type ExecuteTask(UBehaviorTreeComponent* OwnerComp, uint8* NodeMemory) override;
 	virtual uint16 GetInstanceMemorySize() const override;
-	virtual void DescribeRuntimeValues(const class UBehaviorTreeComponent* OwnerComp, uint8* NodeMemory, EBTDescriptionVerbosity::Type Verbosity, TArray<FString>& Values) const override;
+	virtual void DescribeRuntimeValues(const UBehaviorTreeComponent* OwnerComp, uint8* NodeMemory, EBTDescriptionVerbosity::Type Verbosity, TArray<FString>& Values) const override;
 	virtual FString GetStaticDescription() const override;
 
 #if WITH_EDITOR
@@ -33,5 +33,5 @@ class AIMODULE_API UBTTask_Wait : public UBTTaskNode
 
 protected:
 
-	virtual void TickTask(class UBehaviorTreeComponent* OwnerComp, uint8* NodeMemory, float DeltaSeconds) override;
+	virtual void TickTask(UBehaviorTreeComponent* OwnerComp, uint8* NodeMemory, float DeltaSeconds) override;
 };
diff --git a/Engine/Source/Runtime/AIModule/Classes/BehaviorTree/Tasks/BTTask_WaitBlackboardTime.h b/Engine/Source/Runtime/AIModule/Classes/BehaviorTree/Tasks/BTTask_WaitBlackboardTime.h
--- a/Engine/Source/Runtime/AIModule/Classes/BehaviorTree/Tasks/BTTask_WaitBlackboardTime.h
+++ b/Engine/Source/Runtime/AIModule/Classes/BehaviorTree/Tasks/BTTask_WaitBlackboardTime.h
@@ -14,7 +14,7 @@ class AIMODULE_API UBTTask_WaitBlackboardTime : public UBTTask_Wait
 	GENERATED_UCLASS_BODY()
 
 	virtual void InitializeFromAsset(UBehaviorTree& Asset) override;
-	virtual EBTNodeResult::Type ExecuteTask(class UBehaviorTreeComponent* OwnerComp, uint8* NodeMemory) override;
+	virtual EBTNodeResult::Type ExecuteTask(UBehaviorTreeComponent* OwnerComp, uint8* NodeMemory) override;
 	virtual FString GetStaticDescription() const override;
 
 	/** get name of selected blackboard key */
diff --git a/Engine/Source/Runtime/AIModule/Classes/Blueprint/AIAsyncTaskBlueprintProxy.h b/Engine/Source/Runtime/AIModule/Classes/Blueprint/AIAsyncTaskBlueprintProxy.h
--- a/Engine/Source/Runtime/AIModule/Classes/Blueprint/AIAsyncTaskBlueprintProxy.h
+++ b/Engine/Source/Runtime/AIModule/Classes/Blueprint/AIAsyncTaskBlueprintProxy.h
@@ -1,6 +1,7 @@
 // Copyright 1998-2014 Epic Games, Inc. All Rights Reserved.
 #pragma once
 
+#include "AITypes.h"
 #include "AI/Navigation/NavLinkDefinition.h"
 #include "Navigation/PathFollowingComponent.h"
 #include "AIAsyncTaskBlueprintProxy.generated.h"
diff --git a/Engine/Source/Runtime/AIModule/Classes/Blueprint/AIBlueprintHelperLibrary.h b/Engine/Source/Runtime/AIModule/Classes/Blueprint/AIBlueprintHelperLibrary.h
--- a/Engine/Source/Runtime/AIModule/Classes/Blueprint/AIBlueprintHelperLibrary.h
+++ b/Engine/Source/Runtime/AIModule/Classes/Blueprint/AIBlueprintHelperLibrary.h
@@ -8,30 +8,36 @@
 #pragma once
 #include "AIBlueprintHelperLibrary.generated.h"
 
+class UBehaviorTree;
+class UAnimInstance;
+class APawn;
+class UBlackboardComponent;
+class UAIAsyncTaskBlueprintProxy;
+
 UCLASS()
 class AIMODULE_API UAIBlueprintHelperLibrary : public UBlueprintFunctionLibrary
 {
 	GENERATED_UCLASS_BODY()
 
 	UFUNCTION(BlueprintCallable, meta=(WorldContext="WorldContextObject", BlueprintInternalUseOnly = "TRUE"))
-	static class UAIAsyncTaskBlueprintProxy* CreateMoveToProxyObject(class UObject* WorldContextObject, APawn* Pawn, FVector Destination, AActor* TargetActor = NULL, float AcceptanceRadius = 0, bool bStopOnOverlap = false);
+	static UAIAsyncTaskBlueprintProxy* CreateMoveToProxyObject(UObject* WorldContextObject, APawn* Pawn, FVector Destination, AActor* TargetActor = NULL, float AcceptanceRadius = 0, bool bStopOnOverlap = false);
 
 	UFUNCTION(BlueprintCallable, Category="AI", meta=(DefaultToSelf="MessageSource"))
 	static void SendAIMessage(APawn* Target, FName Message, UObject* MessageSource, bool bSuccess = true);
 
 	UFUNCTION(BlueprintCallable, Category="AI", meta=(WorldContext="WorldContextObject", UnsafeDuringActorConstruction="true"))
-	static class APawn* SpawnAIFromClass(class UObject* WorldContextObject, TSubclassOf<APawn> PawnClass, class UBehaviorTree* BehaviorTree, FVector Location, FRotator Rotation=FRotator::ZeroRotator, bool bNoCollisionFail=false);
+	static APawn* SpawnAIFromClass(UObject* WorldContextObject, TSubclassOf<APawn> PawnClass, UBehaviorTree* BehaviorTree, FVector Location, FRotator Rotation=FRotator::ZeroRotator, bool bNoCollisionFail=false);
 
 	UFUNCTION(BlueprintPure, Category="AI", meta=(DefaultToSelf="Target"))
-	static class UBlackboardComponent* GetBlackboard(AActor* Target);
+	static UBlackboardComponent* GetBlackboard(AActor* Target);
 
 	/** locks indicated AI resources of animated pawn */
 	UFUNCTION(BlueprintCallable, Category = "Animation", BlueprintAuthorityOnly, meta = (DefaultToSelf = "AnimInstance"))
-	static void LockAIResourcesWithAnimation(class UAnimInstance* AnimInstance, bool bLockMovement, bool LockAILogic);
+	static void LockAIResourcesWithAnimation(UAnimInstance* AnimInstance, bool bLockMovement, bool LockAILogic);
 
 	/** unlocks indicated AI resources of animated pawn. Will unlock only animation-locked resources */
 	UFUNCTION(BlueprintCallable, Category = "Animation", BlueprintAuthorityOnly, meta = (DefaultToSelf = "AnimInstance"))
-	static void UnlockAIResourcesWithAnimation(class UAnimInstance* AnimInstance, bool bUnlockMovement, bool UnlockAILogic);
+	static void UnlockAIResourcesWithAnimation(UAnimInstance* AnimInstance, bool bUnlockMovement, bool UnlockAILogic);
 
 	UFUNCTION(BlueprintPure, Category = "AI")
 	static bool IsValidAILocation(FVector Location);
diff --git a/Engine/Source/Runtime/AIModule/Classes/BrainComponent.h b/Engine/Source/Runtime/AIModule/Classes/BrainComponent.h
--- a/Engine/Source/Runtime/AIModule/Classes/BrainComponent.h
+++ b/Engine/Source/Runtime/AIModule/Classes/BrainComponent.h
@@ -6,10 +6,12 @@
 #include "AIResourceInterface.h"
 #include "BrainComponent.generated.h"
 
+class AController;
 class AAIController;
 class UBlackboardComponent;
 class UBrainComponent;
 struct FAIMessage;
+class APawn;
 
 DECLARE_DELEGATE_TwoParams(FOnAIMessage, UBrainComponent*, const FAIMessage&);
 
@@ -50,11 +52,11 @@ struct AIMODULE_API FAIMessage
 	void ClearFlag(uint8 Flag) { MessageFlags &= ~Flag; }
 	bool HasFlag(uint8 Flag) const { return (MessageFlags & Flag) != 0; }
 
-	static void Send(class AController* Controller, const FAIMessage& Message);
-	static void Send(class APawn* Pawn, const FAIMessage& Message);
-	static void Send(class UBrainComponent* BrainComp, const FAIMessage& Message);
+	static void Send(AController* Controller, const FAIMessage& Message);
+	static void Send(APawn* Pawn, const FAIMessage& Message);
+	static void Send(UBrainComponent* BrainComp, const FAIMessage& Message);
 
-	static void Broadcast(class UObject* WorldContextObject, const FAIMessage& Message);
+	static void Broadcast(UObject* WorldContextObject, const FAIMessage& Message);
 };
 
 typedef TSharedPtr<struct FAIMessageObserver, ESPMode::Fast> FAIMessageObserverHandle;
@@ -63,8 +65,8 @@ struct AIMODULE_API FAIMessageObserver : public TSharedFromThis<FAIMessageObserv
 {
 public:
 
-	static FAIMessageObserverHandle Create(class AController* Controller, FName MessageType, FOnAIMessage const& Delegate);
-	static FAIMessageObserverHandle Create(class AController* Controller, FName MessageType, FAIRequestID MessageID, FOnAIMessage const& Delegate);
+	static FAIMessageObserverHandle Create(AController* Controller, FName MessageType, FOnAIMessage const& Delegate);
+	static FAIMessageObserverHandle Create(AController* Controller, FName MessageType, FAIRequestID MessageID, FOnAIMessage const& Delegate);
 
 	static FAIMessageObserverHandle Create(APawn* Pawn, FName MessageType, FOnAIMessage const& Delegate);
 	static FAIMessageObserverHandle Create(APawn* Pawn, FName MessageType, FAIRequestID MessageID, FOnAIMessage const& Delegate);
@@ -83,7 +85,7 @@ public:
 
 private:
 
-	void Register(class UBrainComponent* OwnerComp);
+	void Register(UBrainComponent* OwnerComp);
 	void Unregister();
 
 	/** observed message type */
diff --git a/Engine/Source/Runtime/AIModule/Classes/EnvironmentQuery/EQSRenderingComponent.h b/Engine/Source/Runtime/AIModule/Classes/EnvironmentQuery/EQSRenderingComponent.h
--- a/Engine/Source/Runtime/AIModule/Classes/EnvironmentQuery/EQSRenderingComponent.h
+++ b/Engine/Source/Runtime/AIModule/Classes/EnvironmentQuery/EQSRenderingComponent.h
@@ -7,6 +7,9 @@
 #include "EnvQueryDebugHelpers.h"
 #include "EQSRenderingComponent.generated.h"
 
+class IEQSQueryResultSourceInterface;
+struct FEnvQueryInstance;
+
 class AIMODULE_API FEQSSceneProxy : public FDebugRenderSceneProxy
 {
 public:
@@ -18,14 +21,14 @@ public:
 	virtual FPrimitiveViewRelevance GetViewRelevance(const FSceneView* View) override;
 
 #if  USE_EQS_DEBUGGER 
-	static void CollectEQSData(const UPrimitiveComponent* InComponent, const class IEQSQueryResultSourceInterface* QueryDataSource, TArray<FSphere>& Spheres, TArray<FText3d>& Texts, TArray<EQSDebug::FDebugHelper>& DebugItems);
-	static void CollectEQSData(const FEnvQueryResult* ResultItems, const struct FEnvQueryInstance* QueryInstance, TArray<FSphere>& Spheres, TArray<FText3d>& Texts, bool ShouldDrawFailedItems, TArray<EQSDebug::FDebugHelper>& DebugItems);
+	static void CollectEQSData(const UPrimitiveComponent* InComponent, const IEQSQueryResultSourceInterface* QueryDataSource, TArray<FSphere>& Spheres, TArray<FText3d>& Texts, TArray<EQSDebug::FDebugHelper>& DebugItems);
+	static void CollectEQSData(const FEnvQueryResult* ResultItems, const FEnvQueryInstance* QueryInstance, TArray<FSphere>& Spheres, TArray<FText3d>& Texts, bool ShouldDrawFailedItems, TArray<EQSDebug::FDebugHelper>& DebugItems);
 #endif
 private:
 	FEnvQueryResult QueryResult;	
 	// can be 0
 	AActor* ActorOwner;
-	const class IEQSQueryResultSourceInterface* QueryDataSource;
+	const IEQSQueryResultSourceInterface* QueryDataSource;
 	bool bDrawOnlyWhenSelected;
 
 	static const FVector ItemDrawRadius;
diff --git a/Engine/Source/Runtime/AIModule/Classes/Navigation/CrowdFollowingComponent.h b/Engine/Source/Runtime/AIModule/Classes/Navigation/CrowdFollowingComponent.h
--- a/Engine/Source/Runtime/AIModule/Classes/Navigation/CrowdFollowingComponent.h
+++ b/Engine/Source/Runtime/AIModule/Classes/Navigation/CrowdFollowingComponent.h
@@ -6,6 +6,10 @@
 #include "Navigation/PathFollowingComponent.h"
 #include "CrowdFollowingComponent.generated.h"
 
+class INavLinkCustomInterface;
+class UCharacterMovementComponent; 
+class UCrowdManager;
+
 namespace ECrowdAvoidanceQuality
 {
 	enum Type
@@ -40,7 +44,7 @@ class AIMODULE_API UCrowdFollowingComponent : public UPathFollowingComponent, pu
 	virtual void ResumeMove(FAIRequestID RequestID = FAIRequestID::CurrentRequest) override;
 	virtual FVector GetMoveFocus(bool bAllowStrafe) const override;
 	virtual void OnLanded() override;
-	virtual void FinishUsingCustomLink(class INavLinkCustomInterface* CustomNavLink) override;
+	virtual void FinishUsingCustomLink(INavLinkCustomInterface* CustomNavLink) override;
 	virtual void OnPathFinished(EPathFollowingResult::Type Result) override;
 	virtual void OnPathUpdated() override;
 	virtual int32 GetCurrentPathElement() const override { return LastPathPolyIndex; }
@@ -108,7 +112,7 @@ class AIMODULE_API UCrowdFollowingComponent : public UPathFollowingComponent, pu
 protected:
 
 	UPROPERTY(transient)
-	class UCharacterMovementComponent* CharacterMovement;
+	UCharacterMovementComponent* CharacterMovement;
 
 	/** Group mask for this agent */
 	UPROPERTY(Category = "Avoidance", EditAnywhere, BlueprintReadOnly, AdvancedDisplay)
@@ -179,5 +183,5 @@ protected:
 	bool HasMovedDuringPause() const;
 	void UpdateCachedDirections(const FVector& NewVelocity, const FVector& NextPathCorner, bool bTraversingLink);
 
-	friend class UCrowdManager;
+	friend UCrowdManager;
 };
diff --git a/Engine/Source/Runtime/AIModule/Classes/Navigation/NavigationComponent.h b/Engine/Source/Runtime/AIModule/Classes/Navigation/NavigationComponent.h
--- a/Engine/Source/Runtime/AIModule/Classes/Navigation/NavigationComponent.h
+++ b/Engine/Source/Runtime/AIModule/Classes/Navigation/NavigationComponent.h
@@ -9,6 +9,11 @@
 
 class AActor;
 class AController;
+class ANavigationData;
+class INavAgentInterface;
+class UNavLinkCustomComponent;
+class UCanvas;
+class UPathFollowingComponent;
 
 UCLASS(config=Engine)
 class AIMODULE_API UNavigationComponent : public UActorComponent, public INavigationPathGenerator
@@ -34,7 +39,7 @@ public:
 	// End UActorComponent Interface
 
 	// Begin INavigationPathGenerator Interface
-	virtual FNavPathSharedPtr GetGeneratedPath(class INavAgentInterface* Agent) override { return Path; }
+	virtual FNavPathSharedPtr GetGeneratedPath(INavAgentInterface* Agent) override { return Path; }
 	// End INavigationPathGenerator Interface
 
 	/** updates cached pointers to relevant owner's components */
@@ -48,10 +53,10 @@ public:
 	//----------------------------------------------------------------------//
 	
 	/** Retrieves ANavigationData instance being a context to this NavigationComponent. */
-	FORCEINLINE const class ANavigationData* GetNavData() const { return MyNavData ? MyNavData : PickNavData(); }
+	FORCEINLINE const ANavigationData* GetNavData() const { return MyNavData ? MyNavData : PickNavData(); }
 	
 	/** Retrieves ANavigationData instance being a context to this NavigationComponent. */
-	FORCEINLINE class ANavigationData* GetNavData() { return MyNavData ? MyNavData : PickNavData(); }
+	FORCEINLINE ANavigationData* GetNavData() { return MyNavData ? MyNavData : PickNavData(); }
 
 	/** Returns a random location on the navmesh. */
 	FNavLocation GetRandomPointOnNavMesh() const;
@@ -157,7 +162,7 @@ public:
 	void SetRepathWhenInvalid(bool bEnabled);
 
 	UFUNCTION()
-	virtual void OnNavDataRegistered(class ANavigationData* NavData);
+	virtual void OnNavDataRegistered(ANavigationData* NavData);
 
 	/** 
 	 *	Called when given path becomes invalid (via @see PathObserverDelegate)
@@ -166,7 +171,7 @@ public:
 	virtual void OnPathEvent(FNavigationPath* InvalidatedPath, ENavPathEvent::Type Event);
 
 	/** Called when nearby custom link changes its state */
-	virtual void OnCustomLinkBroadcast(class UNavLinkCustomComponent* NearbyLink);
+	virtual void OnCustomLinkBroadcast(UNavLinkCustomComponent* NearbyLink);
 
 	/** set whether this component should receive update broadcasts from nearby smart links or not */
 	void SetReceiveSmartLinkUpdates(bool bEnabled);
@@ -194,14 +199,14 @@ public:
 #endif // ENABLE_VISUAL_LOG
 
 	/** quick and dirty debug render. If Canvas not NULL additional data (if present) will be drawn */
-	void DrawDebugCurrentPath(class UCanvas* Canvas = NULL, FColor PathColor = FColor::White, const uint32 NextPathPointIndex = 0);
+	void DrawDebugCurrentPath(UCanvas* Canvas = NULL, FColor PathColor = FColor::White, const uint32 NextPathPointIndex = 0);
 
 
 protected:
 
 	/** associated path following component */
 	UPROPERTY(transient)
-	class UPathFollowingComponent* PathFollowComp;
+	UPathFollowingComponent* PathFollowComp;
 
 	/** default query extent to use. It's a sum of navigation data's default query extent and nav agent's size */
 	mutable FVector NavigationQueryExtent;
@@ -283,7 +288,7 @@ protected:
 	FDeferredRepath RepathData;
 
 	/** asks NavigationSystem for NavigationData appropriate for self and caches it in MyNavData */
-	class ANavigationData* PickNavData() const;
+	ANavigationData* PickNavData() const;
 
 	/** called as part of PickNavData to update NavigationQueryExtent */
 	virtual void CacheNavQueryExtent() const;
diff --git a/Engine/Source/Runtime/AIModule/Classes/Navigation/PathFollowingComponent.h b/Engine/Source/Runtime/AIModule/Classes/Navigation/PathFollowingComponent.h
--- a/Engine/Source/Runtime/AIModule/Classes/Navigation/PathFollowingComponent.h
+++ b/Engine/Source/Runtime/AIModule/Classes/Navigation/PathFollowingComponent.h
@@ -9,6 +9,13 @@
 
 DECLARE_LOG_CATEGORY_EXTERN(LogPathFollowing, Warning, All);
 
+class UNavMovementComponent;
+class UCanvas;
+class AActor;
+class INavLinkCustomInterface;
+class INavAgentInterface;
+class UNavigationComponent;
+
 UENUM(BlueprintType)
 namespace EPathFollowingStatus
 {
@@ -147,7 +154,7 @@ class AIMODULE_API UPathFollowingComponent : public UActorComponent, public IAIR
 	virtual void OnPathUpdated();
 
 	/** set associated movement component */
-	virtual void SetMovementComponent(class UNavMovementComponent* MoveComp);
+	virtual void SetMovementComponent(UNavMovementComponent* MoveComp);
 
 	/** get current focal point of movement */
 	virtual FVector GetMoveFocus(bool bAllowStrafe) const;
@@ -229,14 +236,14 @@ class AIMODULE_API UPathFollowingComponent : public UActorComponent, public IAIR
 	virtual void GetDebugStringTokens(TArray<FString>& Tokens, TArray<EPathFollowingDebugTokens::Type>& Flags) const;
 	virtual FString GetDebugString() const;
 
-	virtual void DisplayDebug(class UCanvas* Canvas, const FDebugDisplayInfo& DebugDisplay, float& YL, float& YPos) const;
+	virtual void DisplayDebug(UCanvas* Canvas, const FDebugDisplayInfo& DebugDisplay, float& YL, float& YPos) const;
 #if ENABLE_VISUAL_LOG
 	virtual void DescribeSelfToVisLog(struct FVisualLogEntry* Snapshot) const;
 #endif // ENABLE_VISUAL_LOG
 
 	/** called when moving agent collides with another actor */
 	UFUNCTION()
-	virtual void OnActorBump(class AActor* SelfActor, class AActor* OtherActor, FVector NormalImpulse, const FHitResult& Hit);
+	virtual void OnActorBump(AActor* SelfActor, AActor* OtherActor, FVector NormalImpulse, const FHitResult& Hit);
 
 	/** Called when movement is blocked by a collision with another actor.  */
 	virtual void OnMoveBlockedBy(const FHitResult& BlockingImpact) {}
@@ -245,7 +252,7 @@ class AIMODULE_API UPathFollowingComponent : public UActorComponent, public IAIR
 	virtual void OnLanded() {}
 
 	/** call when moving agent finishes using custom nav link, returns control back to path following */
-	virtual void FinishUsingCustomLink(class INavLinkCustomInterface* CustomNavLink);
+	virtual void FinishUsingCustomLink(INavLinkCustomInterface* CustomNavLink);
 
 	// IAIResourceInterface begin
 	virtual void LockResource(EAIRequestPriority::Type LockSource) override;
@@ -258,11 +265,11 @@ protected:
 
 	/** associated movement component */
 	UPROPERTY(transient)
-	class UNavMovementComponent* MovementComp;
+	UNavMovementComponent* MovementComp;
 
 	/** associated navigation component */
 	UPROPERTY(transient)
-	class UNavigationComponent* NavComp;
+	UNavigationComponent* NavComp;
 
 	/** currently traversed custom nav link */
 	FWeakObjectPtr CurrentCustomLinkOb;
@@ -302,7 +309,7 @@ protected:
 	TWeakObjectPtr<AActor> DestinationActor;
 
 	/** cached DestinationActor cast to INavAgentInterface. Use SetDestinationActor to set this */
-	const class INavAgentInterface* DestinationAgent;
+	const INavAgentInterface* DestinationAgent;
 
 	/** destination for current path segment */
 	FBasedPosition CurrentDestination;
@@ -384,7 +391,7 @@ protected:
 	virtual void UpdatePathSegment();
 
 	/** next path segment if custom nav link, try passing control to it */
-	virtual void StartUsingCustomLink(class INavLinkCustomInterface* CustomNavLink, const FVector& DestPoint);
+	virtual void StartUsingCustomLink(INavLinkCustomInterface* CustomNavLink, const FVector& DestPoint);
 
 	/** update blocked movement detection, @returns true if new sample was added */
 	virtual bool UpdateBlockDetection();
diff --git a/Engine/Source/Runtime/AIModule/Classes/Perception/AISense.h b/Engine/Source/Runtime/AIModule/Classes/Perception/AISense.h
--- a/Engine/Source/Runtime/AIModule/Classes/Perception/AISense.h
+++ b/Engine/Source/Runtime/AIModule/Classes/Perception/AISense.h
@@ -5,6 +5,7 @@
 #include "AIPerceptionTypes.h"
 #include "AISense.generated.h"
 
+class AActor;
 class UAIPerceptionSystem;
 
 DECLARE_DELEGATE_OneParam(FOnPerceptionListenerUpdateDelegate, const FPerceptionListener&);
diff --git a/Engine/Source/Runtime/AIModule/Classes/Perception/AISense_Prediction.h b/Engine/Source/Runtime/AIModule/Classes/Perception/AISense_Prediction.h
--- a/Engine/Source/Runtime/AIModule/Classes/Perception/AISense_Prediction.h
+++ b/Engine/Source/Runtime/AIModule/Classes/Perception/AISense_Prediction.h
@@ -6,24 +6,26 @@
 #include "Perception/AISense.h"
 #include "AISense_Prediction.generated.h"
 
+class UAISense_Prediction;
+
 USTRUCT()
 struct AIMODULE_API FAIPredictionEvent
 {	
 	GENERATED_USTRUCT_BODY()
 
-	typedef class UAISense_Prediction FSenseClass;
+	typedef UAISense_Prediction FSenseClass;
 	
 	UPROPERTY()
-	class AActor* Requestor;
+	AActor* Requestor;
 
 	UPROPERTY()
-	class AActor* PredictedActor;
+	AActor* PredictedActor;
 
 	float TimeToPredict;
 		
 	FAIPredictionEvent(){}
 	
-	FAIPredictionEvent(class AActor* InRequestor, class AActor* InPredictedActor, float PredictionTime)
+	FAIPredictionEvent(AActor* InRequestor, AActor* InPredictedActor, float PredictionTime)
 		: Requestor(InRequestor), PredictedActor(InPredictedActor), TimeToPredict(PredictionTime)
 	{
 	}
diff --git a/Engine/Source/Runtime/AIModule/Classes/Perception/AISense_Sight.h b/Engine/Source/Runtime/AIModule/Classes/Perception/AISense_Sight.h
--- a/Engine/Source/Runtime/AIModule/Classes/Perception/AISense_Sight.h
+++ b/Engine/Source/Runtime/AIModule/Classes/Perception/AISense_Sight.h
@@ -8,6 +8,7 @@
 
 class IAISightTargetInterface;
 class UAISenseConfig_Sight;
+class UAISense_Sight;
 
 namespace ESightPerceptionEventName
 {
@@ -24,20 +25,20 @@ struct AIMODULE_API FAISightEvent
 {
 	GENERATED_USTRUCT_BODY()
 
-	typedef class UAISense_Sight FSenseClass;
+	typedef UAISense_Sight FSenseClass;
 
 	float Age;
 	ESightPerceptionEventName::Type EventType;	
 
 	UPROPERTY()
-	class AActor* SeenActor;
+	AActor* SeenActor;
 
 	UPROPERTY()
-	class AActor* Observer;
+	AActor* Observer;
 
 	FAISightEvent(){}
 
-	FAISightEvent(class AActor* InSeenActor, class AActor* InObserver, ESightPerceptionEventName::Type InEventType)
+	FAISightEvent(AActor* InSeenActor, AActor* InObserver, ESightPerceptionEventName::Type InEventType)
 		: Age(0.f), EventType(InEventType), SeenActor(InSeenActor), Observer(InObserver)
 	{
 	}
@@ -53,7 +54,7 @@ struct FAISightTarget
 	FGenericTeamId TeamId;
 	FTargetId TargetId;
 
-	FAISightTarget(class AActor* InTarget = NULL, FGenericTeamId InTeamId = FGenericTeamId::NoTeam);
+	FAISightTarget(AActor* InTarget = NULL, FGenericTeamId InTeamId = FGenericTeamId::NoTeam);
 
 	FORCEINLINE FVector GetLocationSimple() const
 	{
diff --git a/Engine/Source/Runtime/AIModule/Classes/Perception/AISense_Team.h b/Engine/Source/Runtime/AIModule/Classes/Perception/AISense_Team.h
--- a/Engine/Source/Runtime/AIModule/Classes/Perception/AISense_Team.h
+++ b/Engine/Source/Runtime/AIModule/Classes/Perception/AISense_Team.h
@@ -6,12 +6,14 @@
 #include "Perception/AISense.h"
 #include "AISense_Team.generated.h"
 
+class UAISense_Team;
+
 USTRUCT()
 struct AIMODULE_API FAITeamStimulusEvent
 {	
 	GENERATED_USTRUCT_BODY()
 
-	typedef class UAISense_Team FSenseClass;
+	typedef UAISense_Team FSenseClass;
 
 	FVector LastKnowLocation;
 private:
@@ -22,13 +24,13 @@ public:
 	FGenericTeamId TeamIdentifier;
 private:
 	UPROPERTY()
-	class AActor* Broadcaster;
+	AActor* Broadcaster;
 public:
 	UPROPERTY()
-	class AActor* Enemy;
+	AActor* Enemy;
 		
 	FAITeamStimulusEvent(){}	
-	FAITeamStimulusEvent(class AActor* InBroadcaster, class AActor* InEnemy, const FVector& InLastKnowLocation, float EventRange, float PassedInfoAge = 0.f);
+	FAITeamStimulusEvent(AActor* InBroadcaster, AActor* InEnemy, const FVector& InLastKnowLocation, float EventRange, float PassedInfoAge = 0.f);
 
 	FORCEINLINE void CacheBroadcastLocation()
 	{
diff --git a/Engine/Source/Runtime/AIModule/Classes/Perception/AISense_Touch.h b/Engine/Source/Runtime/AIModule/Classes/Perception/AISense_Touch.h
--- a/Engine/Source/Runtime/AIModule/Classes/Perception/AISense_Touch.h
+++ b/Engine/Source/Runtime/AIModule/Classes/Perception/AISense_Touch.h
@@ -6,24 +6,26 @@
 #include "Perception/AISense.h"
 #include "AISense_Touch.generated.h"
 
+class UAISense_Touch;
+
 USTRUCT()
 struct AIMODULE_API FAITouchEvent
 {	
 	GENERATED_USTRUCT_BODY()
 
-	typedef class UAISense_Touch FSenseClass;
+	typedef UAISense_Touch FSenseClass;
 
 	FVector Location;
 	
 	UPROPERTY()
-	class AActor* TouchReceiver;
+	AActor* TouchReceiver;
 
 	UPROPERTY()
-	class AActor* OtherActor;
+	AActor* OtherActor;
 		
 	FAITouchEvent(){}
 	
-	FAITouchEvent(class AActor* InTouchReceiver, class AActor* InOtherActor, const FVector& EventLocation)
+	FAITouchEvent(AActor* InTouchReceiver, AActor* InOtherActor, const FVector& EventLocation)
 		: Location(EventLocation), TouchReceiver(InTouchReceiver), OtherActor(InOtherActor)
 	{
 	}
diff --git a/Engine/Source/Runtime/AIModule/Classes/Perception/PawnSensingComponent.h b/Engine/Source/Runtime/AIModule/Classes/Perception/PawnSensingComponent.h
--- a/Engine/Source/Runtime/AIModule/Classes/Perception/PawnSensingComponent.h
+++ b/Engine/Source/Runtime/AIModule/Classes/Perception/PawnSensingComponent.h
@@ -6,6 +6,8 @@
 #include "PawnSensingComponent.generated.h"
 
 class UPawnNoiseEmitterComponent;
+class AActor;
+class APawn;
 
 /**
  * SensingComponent encapsulates sensory (ie sight and hearing) settings and functionality for an Actor,
@@ -83,23 +85,23 @@ public:
 	uint32 bHearNoises:1;
 
 	/** Is the given actor our owner? Used to ensure that we are not trying to sense our self / our owner. */
-	virtual bool IsSensorActor(const class AActor* Actor) const;
+	virtual bool IsSensorActor(const AActor* Actor) const;
 
 	/** Are we capable of sensing anything (and do we have any callbacks that care about sensing)? If so UpdateAISensing() will be called every sensing interval. */
 	virtual bool CanSenseAnything() const;
 
 	/** Returns true if we should check whether the given Pawn is visible (because we can see things, the Pawn is not hidden, and if the Pawn is a player and we only see players) */
-	virtual bool ShouldCheckVisibilityOf(class APawn* Pawn) const;
+	virtual bool ShouldCheckVisibilityOf(APawn* Pawn) const;
 
 	/** 
 	 * Chance of seeing other pawn decreases with increasing distance or angle in peripheral vision
 	 * @param bMaySkipChecks if true allows checks to be sometimes skipped if target is far away (used by periodic automatic visibility checks)
 	 * @return true if the specified pawn Other is potentially visible (within peripheral vision, etc.) - still need to do LineOfSightTo() check to see if actually visible.  
 	 */
-	virtual bool CouldSeePawn(const class APawn* Other, bool bMaySkipChecks = false) const;
+	virtual bool CouldSeePawn(const APawn* Other, bool bMaySkipChecks = false) const;
 
 	/** Returns true if we should check whether we can hear the given Pawn (because we are able to hear, and the Pawn has the correct team relationship to us) */
-	virtual bool ShouldCheckAudibilityOf(class APawn* Pawn) const;
+	virtual bool ShouldCheckAudibilityOf(APawn* Pawn) const;
 
 	/**
 	 * Check line to other actor.
@@ -107,7 +109,7 @@ public:
 	 * @param ViewPoint is eye position visibility is being checked from.
 	 * @return true if controller's pawn can see Other actor.
 	 */
-	virtual bool HasLineOfSightTo(const class AActor* Other) const;
+	virtual bool HasLineOfSightTo(const AActor* Other) const;
 
 	/** Test whether the noise is loud enough and recent enough to care about.  bSourceWithinNoiseEmitter is true iff the 
 	 * noise was made by the pawn itself or within close proximity (its collision volume).  Otherwise the noise was made
diff --git a/Engine/Source/Runtime/AIModule/Public/BehaviorTreeDelegates.h b/Engine/Source/Runtime/AIModule/Public/BehaviorTreeDelegates.h
--- a/Engine/Source/Runtime/AIModule/Public/BehaviorTreeDelegates.h
+++ b/Engine/Source/Runtime/AIModule/Public/BehaviorTreeDelegates.h
@@ -2,13 +2,17 @@
 
 #pragma once
 
+class UBehaviorTreeComponent;
+class UBehaviorTree;
+class APawn;
+
 struct AIMODULE_API FBehaviorTreeDelegates
 {
 	/** delegate type for tree execution events (Params: const UBehaviorTreeComponent* OwnerComp, const UBehaviorTree* TreeAsset) */
-	DECLARE_MULTICAST_DELEGATE_TwoParams(FOnTreeStarted, const class UBehaviorTreeComponent*, const class UBehaviorTree* );
+	DECLARE_MULTICAST_DELEGATE_TwoParams(FOnTreeStarted, const UBehaviorTreeComponent*, const UBehaviorTree* );
 
 	/** delegate type for locking AI debugging tool on pawn */
-	DECLARE_MULTICAST_DELEGATE_OneParam(FOnTreeDebugTool, const class APawn*);
+	DECLARE_MULTICAST_DELEGATE_OneParam(FOnTreeDebugTool, const APawn*);
 
 	/** Called when the behavior tree starts execution. */
 	static FOnTreeStarted OnTreeStarted;
diff --git a/Engine/Source/Runtime/AIModule/Public/BlueprintNodeHelpers.h b/Engine/Source/Runtime/AIModule/Public/BlueprintNodeHelpers.h
--- a/Engine/Source/Runtime/AIModule/Public/BlueprintNodeHelpers.h
+++ b/Engine/Source/Runtime/AIModule/Public/BlueprintNodeHelpers.h
@@ -2,6 +2,8 @@
 
 #pragma once
 
+class UBTNode;
+
 namespace BlueprintNodeHelpers
 {
 	FString CollectPropertyDescription(const UObject* Ob, const UClass* StopAtClass, const TArray<UProperty*>& PropertyData);
@@ -16,7 +18,7 @@ namespace BlueprintNodeHelpers
 	void CopyPropertiesToContext(const TArray<UProperty*>& PropertyData, uint8* ObjectMemory, uint8* ContextMemory);
 	void CopyPropertiesFromContext(const TArray<UProperty*>& PropertyData, uint8* ObjectMemory, uint8* ContextMemory);
 
-	bool FindNodeOwner(AActor* OwningActor, class UBTNode* Node, UBehaviorTreeComponent*& OwningComp, int32& OwningInstanceIdx);
+	bool FindNodeOwner(AActor* OwningActor, UBTNode* Node, UBehaviorTreeComponent*& OwningComp, int32& OwningInstanceIdx);
 
 	void AbortLatentActions(UActorComponent* OwnerOb, const UObject* Ob);
 
diff --git a/Engine/Source/Runtime/Engine/Private/AI/Navigation/NavGraphGenerator.h b/Engine/Source/Runtime/Engine/Private/AI/Navigation/NavGraphGenerator.h
--- a/Engine/Source/Runtime/Engine/Private/AI/Navigation/NavGraphGenerator.h
+++ b/Engine/Source/Runtime/Engine/Private/AI/Navigation/NavGraphGenerator.h
@@ -4,13 +4,15 @@
 
 #include "AI/NavDataGenerator.h"
 
+class ANavigationGraph;
+
 /**
  * Class that handles generation of the ANavigationGraph data
  */
 class FNavGraphGenerator : public FNavDataGenerator
 {
 public:
-	FNavGraphGenerator(class ANavigationGraph* InDestNavGraph);
+	FNavGraphGenerator(ANavigationGraph* InDestNavGraph);
 	virtual ~FNavGraphGenerator();
 
 private:
@@ -36,7 +38,7 @@ private:
 
 	FCriticalSection GraphChangingLock;
 
-	class ANavigationGraph* DestNavGraph;
+	ANavigationGraph* DestNavGraph;
 
 	uint32 bInitialized:1;
 };
