Refactor blackboard observer's signature to take BlackboardComponent's reference rather then pointer #UE4

diff --git a/Engine/Source/Runtime/AIModule/Classes/BehaviorTree/BehaviorTreeTypes.h b/Engine/Source/Runtime/AIModule/Classes/BehaviorTree/BehaviorTreeTypes.h
--- a/Engine/Source/Runtime/AIModule/Classes/BehaviorTree/BehaviorTreeTypes.h
+++ b/Engine/Source/Runtime/AIModule/Classes/BehaviorTree/BehaviorTreeTypes.h
@@ -45,7 +45,7 @@ namespace FBlackboard
 }
 
 // delegate defines
-DECLARE_DELEGATE_TwoParams(FOnBlackboardChange, const UBlackboardComponent*, FBlackboard::FKey /*key ID*/);
+DECLARE_DELEGATE_TwoParams(FOnBlackboardChange, const UBlackboardComponent&, FBlackboard::FKey /*key ID*/);
 
 namespace BTSpecialChild
 {
diff --git a/Engine/Source/Runtime/AIModule/Classes/BehaviorTree/Decorators/BTDecorator_Blackboard.h b/Engine/Source/Runtime/AIModule/Classes/BehaviorTree/Decorators/BTDecorator_Blackboard.h
--- a/Engine/Source/Runtime/AIModule/Classes/BehaviorTree/Decorators/BTDecorator_Blackboard.h
+++ b/Engine/Source/Runtime/AIModule/Classes/BehaviorTree/Decorators/BTDecorator_Blackboard.h
@@ -34,7 +34,7 @@ class AIMODULE_API UBTDecorator_Blackboard : public UBTDecorator_BlackboardBase
 	GENERATED_UCLASS_BODY()
 
 	virtual bool CalculateRawConditionValue(UBehaviorTreeComponent* OwnerComp, uint8* NodeMemory) const override;
-	virtual void OnBlackboardChange(const UBlackboardComponent* Blackboard, FBlackboard::FKey ChangedKeyID) override;
+	virtual void OnBlackboardChange(const UBlackboardComponent& Blackboard, FBlackboard::FKey ChangedKeyID) override;
 	virtual void DescribeRuntimeValues(const UBehaviorTreeComponent* OwnerComp, uint8* NodeMemory, EBTDescriptionVerbosity::Type Verbosity, TArray<FString>& Values) const override;
 	virtual FString GetStaticDescription() const override;
 
@@ -88,7 +88,7 @@ protected:
 #endif
 
 	/** take blackboard value and evaluate decorator's condition */
-	bool EvaluateOnBlackboard(const UBlackboardComponent* BlackboardComp) const;
+	bool EvaluateOnBlackboard(const UBlackboardComponent& BlackboardComp) const;
 
 	friend FBlackboardDecoratorDetails;
 };
diff --git a/Engine/Source/Runtime/AIModule/Classes/BehaviorTree/Decorators/BTDecorator_BlackboardBase.h b/Engine/Source/Runtime/AIModule/Classes/BehaviorTree/Decorators/BTDecorator_BlackboardBase.h
--- a/Engine/Source/Runtime/AIModule/Classes/BehaviorTree/Decorators/BTDecorator_BlackboardBase.h
+++ b/Engine/Source/Runtime/AIModule/Classes/BehaviorTree/Decorators/BTDecorator_BlackboardBase.h
@@ -14,7 +14,7 @@ class AIMODULE_API UBTDecorator_BlackboardBase : public UBTDecorator
 	virtual void InitializeFromAsset(UBehaviorTree& Asset) override;
 
 	/** notify about change in blackboard keys */
-	virtual void OnBlackboardChange(const UBlackboardComponent* Blackboard, FBlackboard::FKey ChangedKeyID);
+	virtual void OnBlackboardChange(const UBlackboardComponent& Blackboard, FBlackboard::FKey ChangedKeyID);
 
 #if WITH_EDITOR
 	virtual FName GetNodeIconName() const override;
diff --git a/Engine/Source/Runtime/AIModule/Classes/BehaviorTree/Decorators/BTDecorator_BlueprintBase.h b/Engine/Source/Runtime/AIModule/Classes/BehaviorTree/Decorators/BTDecorator_BlueprintBase.h
--- a/Engine/Source/Runtime/AIModule/Classes/BehaviorTree/Decorators/BTDecorator_BlueprintBase.h
+++ b/Engine/Source/Runtime/AIModule/Classes/BehaviorTree/Decorators/BTDecorator_BlueprintBase.h
@@ -35,7 +35,7 @@ class AIMODULE_API UBTDecorator_BlueprintBase : public UBTDecorator
 	virtual void PostInitProperties() override;
 
 	/** notify about changes in blackboard */
-	void OnBlackboardChange(const UBlackboardComponent* Blackboard, FBlackboard::FKey ChangedKeyID);
+	void OnBlackboardChange(const UBlackboardComponent& Blackboard, FBlackboard::FKey ChangedKeyID);
 
 	virtual FString GetStaticDescription() const override;
 	virtual void DescribeRuntimeValues(const UBehaviorTreeComponent* OwnerComp, uint8* NodeMemory, EBTDescriptionVerbosity::Type Verbosity, TArray<FString>& Values) const override;
diff --git a/Engine/Source/Runtime/AIModule/Classes/BehaviorTree/Decorators/BTDecorator_CompareBBEntries.h b/Engine/Source/Runtime/AIModule/Classes/BehaviorTree/Decorators/BTDecorator_CompareBBEntries.h
--- a/Engine/Source/Runtime/AIModule/Classes/BehaviorTree/Decorators/BTDecorator_CompareBBEntries.h
+++ b/Engine/Source/Runtime/AIModule/Classes/BehaviorTree/Decorators/BTDecorator_CompareBBEntries.h
@@ -51,7 +51,7 @@ public:
 
 	virtual void InitializeFromAsset(UBehaviorTree& Asset) override;
 
-	virtual void OnBlackboardChange(const UBlackboardComponent* Blackboard, FBlackboard::FKey ChangedKeyID);
+	virtual void OnBlackboardChange(const UBlackboardComponent& Blackboard, FBlackboard::FKey ChangedKeyID);
 	virtual void OnBecomeRelevant(UBehaviorTreeComponent* OwnerComp, uint8* NodeMemory) override;
 	virtual void OnCeaseRelevant(UBehaviorTreeComponent* OwnerComp, uint8* NodeMemory) override;
 
diff --git a/Engine/Source/Runtime/AIModule/Classes/BehaviorTree/Decorators/BTDecorator_ConditionalLoop.h b/Engine/Source/Runtime/AIModule/Classes/BehaviorTree/Decorators/BTDecorator_ConditionalLoop.h
--- a/Engine/Source/Runtime/AIModule/Classes/BehaviorTree/Decorators/BTDecorator_ConditionalLoop.h
+++ b/Engine/Source/Runtime/AIModule/Classes/BehaviorTree/Decorators/BTDecorator_ConditionalLoop.h
@@ -14,7 +14,7 @@ class AIMODULE_API UBTDecorator_ConditionalLoop : public UBTDecorator_Blackboard
 	GENERATED_UCLASS_BODY()
 
 	virtual bool CalculateRawConditionValue(UBehaviorTreeComponent* OwnerComp, uint8* NodeMemory) const override;
-	virtual void OnBlackboardChange(const UBlackboardComponent* Blackboard, FBlackboard::FKey ChangedKeyID) override;
+	virtual void OnBlackboardChange(const UBlackboardComponent& Blackboard, FBlackboard::FKey ChangedKeyID) override;
 	virtual void OnNodeDeactivation(FBehaviorTreeSearchData& SearchData, EBTNodeResult::Type NodeResult) override;
 
 #if WITH_EDITOR
diff --git a/Engine/Source/Runtime/AIModule/Classes/BehaviorTree/Decorators/BTDecorator_ConeCheck.h b/Engine/Source/Runtime/AIModule/Classes/BehaviorTree/Decorators/BTDecorator_ConeCheck.h
--- a/Engine/Source/Runtime/AIModule/Classes/BehaviorTree/Decorators/BTDecorator_ConeCheck.h
+++ b/Engine/Source/Runtime/AIModule/Classes/BehaviorTree/Decorators/BTDecorator_ConeCheck.h
@@ -54,7 +54,7 @@ protected:
 
 	virtual bool CalculateRawConditionValue(UBehaviorTreeComponent* OwnerComp, uint8* NodeMemory) const override;
 	virtual void OnBecomeRelevant(UBehaviorTreeComponent* OwnerComp, uint8* NodeMemory) override;
-	void OnBlackboardChange(const UBlackboardComponent* Blackboard, FBlackboard::FKey ChangedKeyID);
+	void OnBlackboardChange(const UBlackboardComponent& Blackboard, FBlackboard::FKey ChangedKeyID);
 
 	virtual void TickNode(UBehaviorTreeComponent* OwnerComp, uint8* NodeMemory, float DeltaSeconds) override;
 	
diff --git a/Engine/Source/Runtime/AIModule/Private/BehaviorTree/BlackboardComponent.cpp b/Engine/Source/Runtime/AIModule/Private/BehaviorTree/BlackboardComponent.cpp
--- a/Engine/Source/Runtime/AIModule/Private/BehaviorTree/BlackboardComponent.cpp
+++ b/Engine/Source/Runtime/AIModule/Private/BehaviorTree/BlackboardComponent.cpp
@@ -200,7 +200,7 @@ void UBlackboardComponent::NotifyObservers(FBlackboard::FKey KeyID) const
 		for (TMultiMap<uint8, FOnBlackboardChange>::TConstKeyIterator KeyIt(Observers, KeyID); KeyIt; ++KeyIt)
 		{
 			const FOnBlackboardChange& ObserverDelegate = KeyIt.Value();
-			ObserverDelegate.ExecuteIfBound(this, KeyID);
+			ObserverDelegate.ExecuteIfBound(*this, KeyID);
 		}
 	}
 }
diff --git a/Engine/Source/Runtime/AIModule/Private/BehaviorTree/Decorators/BTDecorator_Blackboard.cpp b/Engine/Source/Runtime/AIModule/Private/BehaviorTree/Decorators/BTDecorator_Blackboard.cpp
--- a/Engine/Source/Runtime/AIModule/Private/BehaviorTree/Decorators/BTDecorator_Blackboard.cpp
+++ b/Engine/Source/Runtime/AIModule/Private/BehaviorTree/Decorators/BTDecorator_Blackboard.cpp
@@ -14,16 +14,18 @@ UBTDecorator_Blackboard::UBTDecorator_Blackboard(const FObjectInitializer& Objec
 bool UBTDecorator_Blackboard::CalculateRawConditionValue(UBehaviorTreeComponent* OwnerComp, uint8* NodeMemory) const
 {
 	const UBlackboardComponent* BlackboardComp = OwnerComp->GetBlackboardComponent();
-	return EvaluateOnBlackboard(BlackboardComp);
+	// note that this may produce unexpected logical results. FALSE is a valid return value here as well
+	// @todo signal it
+	return BlackboardComp && EvaluateOnBlackboard(*BlackboardComp);
 }
 
-bool UBTDecorator_Blackboard::EvaluateOnBlackboard(const UBlackboardComponent* BlackboardComp) const
+bool UBTDecorator_Blackboard::EvaluateOnBlackboard(const UBlackboardComponent& BlackboardComp) const
 {
 	bool bResult = false;
-	if (BlackboardComp && BlackboardKey.SelectedKeyType)
+	if (BlackboardKey.SelectedKeyType)
 	{
 		UBlackboardKeyType* KeyCDO = BlackboardKey.SelectedKeyType->GetDefaultObject<UBlackboardKeyType>();
-		const uint8* KeyMemory = BlackboardComp->GetKeyRawData(BlackboardKey.GetSelectedKeyID());
+		const uint8* KeyMemory = BlackboardComp.GetKeyRawData(BlackboardKey.GetSelectedKeyID());
 
 		const EBlackboardKeyOperation::Type Op = KeyCDO->GetTestOperation();
 		switch (Op)
@@ -48,9 +50,9 @@ bool UBTDecorator_Blackboard::EvaluateOnBlackboard(const UBlackboardComponent* B
 	return bResult;
 }
 
-void UBTDecorator_Blackboard::OnBlackboardChange(const UBlackboardComponent* Blackboard, FBlackboard::FKey ChangedKeyID)
+void UBTDecorator_Blackboard::OnBlackboardChange(const UBlackboardComponent& Blackboard, FBlackboard::FKey ChangedKeyID)
 {
-	UBehaviorTreeComponent* BehaviorComp = Blackboard ? (UBehaviorTreeComponent*)Blackboard->GetBrainComponent() : NULL;
+	UBehaviorTreeComponent* BehaviorComp = (UBehaviorTreeComponent*)Blackboard.GetBrainComponent();
 	if (BlackboardKey.GetSelectedKeyID() == ChangedKeyID &&
 		GetFlowAbortMode() != EBTFlowAbortMode::None &&
 		BehaviorComp)
@@ -62,7 +64,7 @@ void UBTDecorator_Blackboard::OnBlackboardChange(const UBlackboardComponent* Bla
 
 			UE_VLOG(BehaviorComp->GetOwner(), LogBehaviorTree, Verbose, TEXT("%s, OnBlackboardChange[%s] pass:%d executing:%d => %s"),
 				*UBehaviorTreeTypes::DescribeNodeHelper(this),
-				*Blackboard->GetKeyName(ChangedKeyID).ToString(), bPass, bIsExecutingBranch,
+				*Blackboard.GetKeyName(ChangedKeyID).ToString(), bPass, bIsExecutingBranch,
 				(bIsExecutingBranch && !bPass) || (!bIsExecutingBranch && bPass) ? TEXT("restart") : TEXT("skip"));
 
 			if ((bIsExecutingBranch && !bPass) ||
@@ -75,7 +77,7 @@ void UBTDecorator_Blackboard::OnBlackboardChange(const UBlackboardComponent* Bla
 		{
 			UE_VLOG(BehaviorComp->GetOwner(), LogBehaviorTree, Verbose, TEXT("%s, OnBlackboardChange[%s] => restart"),
 				*UBehaviorTreeTypes::DescribeNodeHelper(this),
-				*Blackboard->GetKeyName(ChangedKeyID).ToString());
+				*Blackboard.GetKeyName(ChangedKeyID).ToString());
 
 			// force result Aborted to restart from this decorator
 			// can't use helper function
@@ -101,7 +103,7 @@ void UBTDecorator_Blackboard::DescribeRuntimeValues(const UBehaviorTreeComponent
 		DescKeyValue = BlackboardComp->DescribeKeyValue(BlackboardKey.GetSelectedKeyID(), EBlackboardDescription::OnlyValue);
 	}
 
-	const bool bResult = EvaluateOnBlackboard(BlackboardComp);
+	const bool bResult = BlackboardComp && EvaluateOnBlackboard(*BlackboardComp);
 	Values.Add(FString::Printf(TEXT("value: %s (%s)"), *DescKeyValue, bResult ? TEXT("pass") : TEXT("fail")));
 }
 
diff --git a/Engine/Source/Runtime/AIModule/Private/BehaviorTree/Decorators/BTDecorator_BlackboardBase.cpp b/Engine/Source/Runtime/AIModule/Private/BehaviorTree/Decorators/BTDecorator_BlackboardBase.cpp
--- a/Engine/Source/Runtime/AIModule/Private/BehaviorTree/Decorators/BTDecorator_BlackboardBase.cpp
+++ b/Engine/Source/Runtime/AIModule/Private/BehaviorTree/Decorators/BTDecorator_BlackboardBase.cpp
@@ -44,9 +44,9 @@ void UBTDecorator_BlackboardBase::OnCeaseRelevant(UBehaviorTreeComponent* OwnerC
 	}
 }
 
-void UBTDecorator_BlackboardBase::OnBlackboardChange(const UBlackboardComponent* Blackboard, FBlackboard::FKey ChangedKeyID)
+void UBTDecorator_BlackboardBase::OnBlackboardChange(const UBlackboardComponent& Blackboard, FBlackboard::FKey ChangedKeyID)
 {
-	UBehaviorTreeComponent* BehaviorComp = Blackboard ? (UBehaviorTreeComponent*)Blackboard->GetBrainComponent() : NULL;
+	UBehaviorTreeComponent* BehaviorComp = (UBehaviorTreeComponent*)Blackboard.GetBrainComponent();
 	if (BlackboardKey.GetSelectedKeyID() == ChangedKeyID && BehaviorComp)
 	{
 		BehaviorComp->RequestExecution(this);		
diff --git a/Engine/Source/Runtime/AIModule/Private/BehaviorTree/Decorators/BTDecorator_BlueprintBase.cpp b/Engine/Source/Runtime/AIModule/Private/BehaviorTree/Decorators/BTDecorator_BlueprintBase.cpp
--- a/Engine/Source/Runtime/AIModule/Private/BehaviorTree/Decorators/BTDecorator_BlueprintBase.cpp
+++ b/Engine/Source/Runtime/AIModule/Private/BehaviorTree/Decorators/BTDecorator_BlueprintBase.cpp
@@ -222,9 +222,9 @@ void UBTDecorator_BlueprintBase::DescribeRuntimeValues(const UBehaviorTreeCompon
 	}
 }
 
-void UBTDecorator_BlueprintBase::OnBlackboardChange(const UBlackboardComponent* Blackboard, FBlackboard::FKey ChangedKeyID)
+void UBTDecorator_BlueprintBase::OnBlackboardChange(const UBlackboardComponent& Blackboard, FBlackboard::FKey ChangedKeyID)
 {
-	UBehaviorTreeComponent* BehaviorComp = Blackboard ? (UBehaviorTreeComponent*)Blackboard->GetBrainComponent() : NULL;
+	UBehaviorTreeComponent* BehaviorComp = (UBehaviorTreeComponent*)Blackboard.GetBrainComponent();
 	if (BehaviorComp)
 	{
 		BehaviorComp->RequestExecution(this);		
diff --git a/Engine/Source/Runtime/AIModule/Private/BehaviorTree/Decorators/BTDecorator_CompareBBEntries.cpp b/Engine/Source/Runtime/AIModule/Private/BehaviorTree/Decorators/BTDecorator_CompareBBEntries.cpp
--- a/Engine/Source/Runtime/AIModule/Private/BehaviorTree/Decorators/BTDecorator_CompareBBEntries.cpp
+++ b/Engine/Source/Runtime/AIModule/Private/BehaviorTree/Decorators/BTDecorator_CompareBBEntries.cpp
@@ -85,9 +85,9 @@ void UBTDecorator_CompareBBEntries::OnCeaseRelevant(UBehaviorTreeComponent* Owne
 	}
 }
 
-void UBTDecorator_CompareBBEntries::OnBlackboardChange(const UBlackboardComponent* Blackboard, FBlackboard::FKey ChangedKeyID)
+void UBTDecorator_CompareBBEntries::OnBlackboardChange(const UBlackboardComponent& Blackboard, FBlackboard::FKey ChangedKeyID)
 {
-	UBehaviorTreeComponent* BehaviorComp = Blackboard ? (UBehaviorTreeComponent*)Blackboard->GetBrainComponent() : NULL;
+	UBehaviorTreeComponent* BehaviorComp = (UBehaviorTreeComponent*)Blackboard.GetBrainComponent();
 	if (BehaviorComp && (BlackboardKeyA.GetSelectedKeyID() == ChangedKeyID || BlackboardKeyB.GetSelectedKeyID() == ChangedKeyID))
 	{
 		BehaviorComp->RequestExecution(this);		
diff --git a/Engine/Source/Runtime/AIModule/Private/BehaviorTree/Decorators/BTDecorator_ConditionalLoop.cpp b/Engine/Source/Runtime/AIModule/Private/BehaviorTree/Decorators/BTDecorator_ConditionalLoop.cpp
--- a/Engine/Source/Runtime/AIModule/Private/BehaviorTree/Decorators/BTDecorator_ConditionalLoop.cpp
+++ b/Engine/Source/Runtime/AIModule/Private/BehaviorTree/Decorators/BTDecorator_ConditionalLoop.cpp
@@ -20,7 +20,7 @@ bool UBTDecorator_ConditionalLoop::CalculateRawConditionValue(UBehaviorTreeCompo
 	return true;
 }
 
-void UBTDecorator_ConditionalLoop::OnBlackboardChange(const UBlackboardComponent* Blackboard, FBlackboard::FKey ChangedKeyID)
+void UBTDecorator_ConditionalLoop::OnBlackboardChange(const UBlackboardComponent& Blackboard, FBlackboard::FKey ChangedKeyID)
 {
 	// empty, don't react to blackboard value changes
 }
@@ -29,7 +29,8 @@ void UBTDecorator_ConditionalLoop::OnNodeDeactivation(FBehaviorTreeSearchData& S
 {
 	if (NodeResult != EBTNodeResult::Aborted && SearchData.OwnerComp)
 	{
-		const bool bEvalResult = EvaluateOnBlackboard(SearchData.OwnerComp->GetBlackboardComponent());
+		const UBlackboardComponent* BlackboardComp = SearchData.OwnerComp->GetBlackboardComponent();
+		const bool bEvalResult = BlackboardComp && EvaluateOnBlackboard(*BlackboardComp);
 		UE_VLOG(SearchData.OwnerComp->GetOwner(), LogBehaviorTree, Verbose, TEXT("Loop condition: %s -> %s"),
 			bEvalResult ? TEXT("true") : TEXT("false"), (bEvalResult != IsInversed()) ? TEXT("run again!") : TEXT("break"));
 
diff --git a/Engine/Source/Runtime/AIModule/Private/BehaviorTree/Decorators/BTDecorator_ConeCheck.cpp b/Engine/Source/Runtime/AIModule/Private/BehaviorTree/Decorators/BTDecorator_ConeCheck.cpp
--- a/Engine/Source/Runtime/AIModule/Private/BehaviorTree/Decorators/BTDecorator_ConeCheck.cpp
+++ b/Engine/Source/Runtime/AIModule/Private/BehaviorTree/Decorators/BTDecorator_ConeCheck.cpp
@@ -87,7 +87,7 @@ void UBTDecorator_ConeCheck::OnBecomeRelevant(UBehaviorTreeComponent* OwnerComp,
 	DecoratorMemory->bLastRawResult = CalcConditionImpl(OwnerComp, NodeMemory);
 }
 
-void UBTDecorator_ConeCheck::OnBlackboardChange(const UBlackboardComponent* Blackboard, FBlackboard::FKey ChangedKeyID)
+void UBTDecorator_ConeCheck::OnBlackboardChange(const UBlackboardComponent& Blackboard, FBlackboard::FKey ChangedKeyID)
 {
 	check(false);
 }
