Removed unused EnvQueryGenerator_PathingGrid::TraceType property #UE4

diff --git a/Engine/Source/Runtime/AIModule/Classes/EnvironmentQuery/Generators/EnvQueryGenerator_PathingGrid.h b/Engine/Source/Runtime/AIModule/Classes/EnvironmentQuery/Generators/EnvQueryGenerator_PathingGrid.h
--- a/Engine/Source/Runtime/AIModule/Classes/EnvironmentQuery/Generators/EnvQueryGenerator_PathingGrid.h
+++ b/Engine/Source/Runtime/AIModule/Classes/EnvironmentQuery/Generators/EnvQueryGenerator_PathingGrid.h
@@ -36,11 +36,7 @@ class UEnvQueryGenerator_PathingGrid : public UEnvQueryGenerator
 	/** context */
 	UPROPERTY(EditAnywhere, Category=Generator)
 	TSubclassOf<UEnvQueryContext> GenerateAround;
-
-	/** context */
-	UPROPERTY(EditAnywhere, Category=Generator)
-	TEnumAsByte<EEnvQueryTrace::Type> TraceType;
-
+	
 	/** navigation filter for tracing */
 	UPROPERTY(EditDefaultsOnly, Category=Generator)
 	TSubclassOf<UNavigationQueryFilter> NavigationFilter;
