Made EnvQueryContext_BlueprintBase abstract to not display in EQS' drob-downs #UE4

diff --git a/Engine/Source/Runtime/AIModule/Classes/EnvironmentQuery/Contexts/EnvQueryContext_BlueprintBase.h b/Engine/Source/Runtime/AIModule/Classes/EnvironmentQuery/Contexts/EnvQueryContext_BlueprintBase.h
--- a/Engine/Source/Runtime/AIModule/Classes/EnvironmentQuery/Contexts/EnvQueryContext_BlueprintBase.h
+++ b/Engine/Source/Runtime/AIModule/Classes/EnvironmentQuery/Contexts/EnvQueryContext_BlueprintBase.h
@@ -9,7 +9,7 @@ class UWorld;
 struct FEnvQueryInstance;
 struct FEnvQueryContextData;
 
-UCLASS(MinimalAPI, Blueprintable)
+UCLASS(MinimalAPI, Abstract, Blueprintable)
 class UEnvQueryContext_BlueprintBase : public UEnvQueryContext
 {
 	GENERATED_UCLASS_BODY()
