Fixed Blackboard not displaying proper values for user-defined enums #UE4

diff --git a/Engine/Source/Runtime/AIModule/Private/BehaviorTree/Blackboard/BlackboardKeyType_Enum.cpp b/Engine/Source/Runtime/AIModule/Private/BehaviorTree/Blackboard/BlackboardKeyType_Enum.cpp
--- a/Engine/Source/Runtime/AIModule/Private/BehaviorTree/Blackboard/BlackboardKeyType_Enum.cpp
+++ b/Engine/Source/Runtime/AIModule/Private/BehaviorTree/Blackboard/BlackboardKeyType_Enum.cpp
@@ -21,7 +21,7 @@ bool UBlackboardKeyType_Enum::SetValue(uint8* RawData, uint8 Value)
 
 FString UBlackboardKeyType_Enum::DescribeValue(const uint8* RawData) const
 {
-	return EnumType ? EnumType->GetEnumName(GetValue(RawData)) : FString("UNKNOWN!");
+	return EnumType ? EnumType->GetEnumText(GetValue(RawData)).ToString() : FString("UNKNOWN!");
 }
 
 FString UBlackboardKeyType_Enum::DescribeSelf() const
@@ -60,5 +60,5 @@ bool UBlackboardKeyType_Enum::TestArithmeticOperation(const uint8* MemoryBlock,
 
 FString UBlackboardKeyType_Enum::DescribeArithmeticParam(int32 IntValue, float FloatValue) const
 {
-	return EnumType ? EnumType->GetEnumName(IntValue) : FString("UNKNOWN!");
+	return EnumType ? EnumType->GetEnumText(IntValue).ToString() : FString("UNKNOWN!");
 }
