Rollback //depot/UE4/Engine/Source/Runtime/AIModule/Private/BehaviorTree/BTCompositeNode.cpp to revision 10

diff --git a/Engine/Source/Runtime/AIModule/Private/BehaviorTree/BTCompositeNode.cpp b/Engine/Source/Runtime/AIModule/Private/BehaviorTree/BTCompositeNode.cpp
--- a/Engine/Source/Runtime/AIModule/Private/BehaviorTree/BTCompositeNode.cpp
+++ b/Engine/Source/Runtime/AIModule/Private/BehaviorTree/BTCompositeNode.cpp
@@ -45,7 +45,7 @@ int32 UBTCompositeNode::FindChildToExecute(FBehaviorTreeSearchData& SearchData,
 	return RetIdx;
 }
 
-int32 UBTCompositeNode::GetChildIndex(struct FBehaviorTreeSearchData& SearchData, const class UBTNode* ChildNode) const
+int32 UBTCompositeNode::GetChildIndex(FBehaviorTreeSearchData& SearchData, const UBTNode* ChildNode) const
 {
 	if (ChildNode->GetParentNode() != this)
 	{
@@ -56,7 +56,7 @@ int32 UBTCompositeNode::GetChildIndex(struct FBehaviorTreeSearchData& SearchData
 	return GetChildIndex(ChildNode);
 }
 
-int32 UBTCompositeNode::GetChildIndex(const class UBTNode* ChildNode) const
+int32 UBTCompositeNode::GetChildIndex(const UBTNode* ChildNode) const
 {
 	for (int32 ChildIndex = 0; ChildIndex < Children.Num(); ChildIndex++)
 	{
@@ -70,7 +70,7 @@ int32 UBTCompositeNode::GetChildIndex(const class UBTNode* ChildNode) const
 	return BTSpecialChild::ReturnToParent;
 }
 
-void UBTCompositeNode::OnChildActivation(struct FBehaviorTreeSearchData& SearchData, const class UBTNode* ChildNode) const
+void UBTCompositeNode::OnChildActivation(FBehaviorTreeSearchData& SearchData, const UBTNode* ChildNode) const
 {
 	OnChildActivation(SearchData, GetChildIndex(SearchData, ChildNode));
 }
