navigation refactor: removed NavigationComponent, added abstract navdata for handling direct path updates

diff --git a/Engine/Config/BaseEngine.ini b/Engine/Config/BaseEngine.ini
--- a/Engine/Config/BaseEngine.ini
+++ b/Engine/Config/BaseEngine.ini
@@ -585,7 +585,6 @@ bUseStreamingPause=false
 +ActiveClassRedirects=(OldClassName="CrowdAgentInterface",NewClassName="/Script/AIModule.CrowdAgentInterface")
 +ActiveClassRedirects=(OldClassName="CrowdFollowingComponent",NewClassName="/Script/AIModule.CrowdFollowingComponent")
 +ActiveClassRedirects=(OldClassName="CrowdManager",NewClassName="/Script/AIModule.CrowdManager")
-+ActiveClassRedirects=(OldClassName="NavigationComponent",NewClassName="/Script/AIModule.NavigationComponent")
 +ActiveClassRedirects=(OldClassName="PathFollowingComponent",NewClassName="/Script/AIModule.PathFollowingComponent")
 +ActiveClassRedirects=(OldClassName="PawnSensingComponent",NewClassName="/Script/AIModule.PawnSensingComponent")
 +K2FieldRedirects=(OldFieldName="BTTask_MoveDirectlyToward.bForceMoveToLocation",NewFieldName="bDisablePathUpdateOnGoalLocationChange")
@@ -2198,11 +2197,6 @@ bBuildNavigationAtRuntime=true
 bAddPlayersToGenerationSeeds=true
 +RequiredNavigationDataClassNames=/Script/Engine.RecastNavMesh
 
-[/Script/Engine.NavigationComponent]
-bDoAsyncPathfinding=false
-bDoHierarchicalPathfinding=false
-bSearchFromGoalWhenOutOfNodes=false
-
 [/Script/Engine.RecastNavMesh]
 ; runtime params
 bRebuildAtRuntime=false
diff --git a/Engine/Source/Developer/GameplayDebugger/Private/GameplayDebuggingComponent.cpp b/Engine/Source/Developer/GameplayDebugger/Private/GameplayDebuggingComponent.cpp
--- a/Engine/Source/Developer/GameplayDebugger/Private/GameplayDebuggingComponent.cpp
+++ b/Engine/Source/Developer/GameplayDebugger/Private/GameplayDebuggingComponent.cpp
@@ -14,7 +14,6 @@
 #include "AIController.h"
 #include "BrainComponent.h"
 #include "BehaviorTreeDelegates.h"
-#include "Navigation/NavigationComponent.h"
 #include "GameFramework/PlayerState.h"
 #include "Engine/Channel.h"
 
@@ -484,9 +483,9 @@ void UGameplayDebuggingComponent::CollectPathData()
 	bool bRefreshRendering = false;
 	if (AAIController *MyController = Cast<AAIController>(MyPawn->Controller))
 	{
-		if ( MyController->GetNavComponent() && MyController->GetNavComponent()->HasValidPath())
+		if (MyController->PathFollowingComponent && MyController->PathFollowingComponent->GetPath().IsValid())
 		{
-			const FNavPathSharedPtr& NewPath = MyController->GetNavComponent()->GetPath();
+			const FNavPathSharedPtr& NewPath = MyController->PathFollowingComponent->GetPath();
 			if (!CurrentPath.HasSameObject(NewPath.Get()))
 			{
 				FVisualLogEntry Snapshot;
diff --git a/Engine/Source/Runtime/AIModule/Classes/AIController.h b/Engine/Source/Runtime/AIModule/Classes/AIController.h
--- a/Engine/Source/Runtime/AIModule/Classes/AIController.h
+++ b/Engine/Source/Runtime/AIModule/Classes/AIController.h
@@ -12,7 +12,6 @@
 #include "AIController.generated.h"
 
 class APawn;
-class UNavigationComponent;
 class UPathFollowingComponent;
 class UBrainComponent;
 class UAIPerceptionComponent;
@@ -26,6 +25,7 @@ class UCanvas;
 #if ENABLE_VISUAL_LOG
 struct FVisualLogEntry;
 #endif // ENABLE_VISUAL_LOG
+struct FPathFindingQuery;
 
 DECLARE_DYNAMIC_MULTICAST_DELEGATE_TwoParams(FAIMoveCompletedSignature, FAIRequestID, RequestID, EPathFollowingResult::Type, Result);
 
@@ -98,10 +98,6 @@ public:
 	uint32 bWantsPlayerState:1;
 
 private_subobject:
-	/** Component used for pathfinding and querying environment's navigation. */
-	DEPRECATED_FORGAME(4.6, "NavComponent should not be accessed directly, please use GetNavComponent() function instead. NavComponent will soon be private and your code will not compile.")
-	UPROPERTY()
-	UNavigationComponent* NavComponent;
 
 	/** Component used for moving along a path. */
 	DEPRECATED_FORGAME(4.6, "PathFollowingComponent should not be accessed directly, please use GetPathFollowingComponent() function instead. PathFollowingComponent will soon be private and your code will not compile.")
@@ -153,6 +149,25 @@ public:
 	// EPathFollowingRequestResult::Type MoveToLocation(const FVector& Dest, float AcceptanceRadius = /*UPathFollowingComponent::DefaultAcceptanceRadius==*/-1, bool bStopOnOverlap = true, bool bUsePathfinding = true, bool bCanStrafe = true);
 	// but parser doesn't like this (it's a bug, when fixed this will be changed)
 
+	/** Prepares a query for pathfinding
+	 *  @param Query - query struct to fill
+	 *  @param Dest - destination location
+	 *  @param Goal - goal actor (if applies)
+	 *  @param bUsePathfinding - calculate paths vs move in straight line
+	 *  @param QueryFilter - optional filter for path finding
+	 *  @return true if query was filled in successfully
+	 */
+	virtual bool PreparePathfinding(FPathFindingQuery& Query, const FVector& Dest, AActor* Goal, bool bUsePathfinding = true, TSubclassOf<class UNavigationQueryFilter> FilterClass = NULL);
+
+	/** Executes pathfinding query and starts move request
+	 *  @param Query - query struct holding pathfinding data
+	 *  @param Goal - goal actor if following actor
+	 *  @param AcceptanceRadius - finish move if pawn gets close enough
+	 *  @param CustomData - game specific data, that will be passed to pawn's movement component
+	 *  @return RequestID, or 0 when failed
+	 */
+	virtual FAIRequestID RequestPathAndMove(FPathFindingQuery& Query, AActor* Goal, float AcceptanceRadius, bool bStopOnOverlap, FCustomMoveSharedPtr CustomData);
+
 	/** Handle move requests
 	 *  @param Path - path to follow (can use incomplete)
 	 *  @param Goal - goal actor if following actor
@@ -181,12 +196,6 @@ public:
 	UPROPERTY(BlueprintAssignable, meta=(DisplayName="MoveCompleted"))
 	FAIMoveCompletedSignature ReceiveMoveCompleted;
 
-	/** @returns path to actor Goal (can be incomplete if async pathfinding is used) */
-	FNavPathSharedPtr FindPath(AActor* Goal, bool bUsePathfinding = true, TSharedPtr<const FNavigationQueryFilter> QueryFilter = NULL);
-
-	/** @returns path to point Dest (can be incomplete if async pathfinding is used) */
-	FNavPathSharedPtr FindPath(const FVector& Dest, bool bUsePathfinding = true, TSharedPtr<const FNavigationQueryFilter> QueryFilter = NULL);
-
 	/** Returns status of path following */
 	UFUNCTION(BlueprintCallable, Category="AI|Navigation")
 	EPathFollowingStatus::Type GetMoveStatus() const;
@@ -204,7 +213,7 @@ public:
 	void SetMoveBlockDetection(bool bEnable);
 
 	/** Prepares path finding and path following components. */
-	virtual void InitNavigationControl(UNavigationComponent*& PathFindingComp, UPathFollowingComponent*& PathFollowingComp) override;
+	virtual void InitNavigationControl(UPathFollowingComponent*& PathFollowingComp) override;
 
 	/** Starts executing behavior tree. */
 	UFUNCTION(BlueprintCallable, Category="AI")
diff --git a/Engine/Source/Runtime/AIModule/Classes/Navigation/CrowdFollowingComponent.h b/Engine/Source/Runtime/AIModule/Classes/Navigation/CrowdFollowingComponent.h
--- a/Engine/Source/Runtime/AIModule/Classes/Navigation/CrowdFollowingComponent.h
+++ b/Engine/Source/Runtime/AIModule/Classes/Navigation/CrowdFollowingComponent.h
@@ -47,6 +47,7 @@ class AIMODULE_API UCrowdFollowingComponent : public UPathFollowingComponent, pu
 	virtual void FinishUsingCustomLink(INavLinkCustomInterface* CustomNavLink) override;
 	virtual void OnPathFinished(EPathFollowingResult::Type Result) override;
 	virtual void OnPathUpdated() override;
+	virtual void OnPathfindingQuery(FPathFindingQuery& Query) override;
 	virtual int32 GetCurrentPathElement() const override { return LastPathPolyIndex; }
 	// PathFollowingComponent END
 
diff --git a/Engine/Source/Runtime/AIModule/Classes/Navigation/NavigationComponent.h b/Engine/Source/Runtime/AIModule/Classes/Navigation/NavigationComponent.h
deleted file mode 100644
diff --git a/Engine/Source/Runtime/AIModule/Classes/Navigation/PathFollowingComponent.h b/Engine/Source/Runtime/AIModule/Classes/Navigation/PathFollowingComponent.h
--- a/Engine/Source/Runtime/AIModule/Classes/Navigation/PathFollowingComponent.h
+++ b/Engine/Source/Runtime/AIModule/Classes/Navigation/PathFollowingComponent.h
@@ -257,6 +257,9 @@ class AIMODULE_API UPathFollowingComponent : public UActorComponent, public IAIR
 	/** call when moving agent finishes using custom nav link, returns control back to path following */
 	virtual void FinishUsingCustomLink(INavLinkCustomInterface* CustomNavLink);
 
+	/** called when owner is preparing new pathfinding request */
+	virtual void OnPathfindingQuery(FPathFindingQuery& Query) {}
+
 	// IAIResourceInterface begin
 	virtual void LockResource(EAIRequestPriority::Type LockSource) override;
 	virtual void ClearResourceLock(EAIRequestPriority::Type LockSource) override;
@@ -270,10 +273,6 @@ protected:
 	UPROPERTY(transient)
 	UNavMovementComponent* MovementComp;
 
-	/** associated navigation component */
-	UPROPERTY(transient)
-	UNavigationComponent* NavComp;
-
 	/** currently traversed custom nav link */
 	FWeakObjectPtr CurrentCustomLinkOb;
 
diff --git a/Engine/Source/Runtime/AIModule/Private/AIController.cpp b/Engine/Source/Runtime/AIModule/Private/AIController.cpp
--- a/Engine/Source/Runtime/AIModule/Private/AIController.cpp
+++ b/Engine/Source/Runtime/AIModule/Private/AIController.cpp
@@ -8,7 +8,6 @@
 #include "BehaviorTree/BehaviorTreeManager.h"
 #include "BehaviorTree/BlackboardComponent.h"
 #include "Navigation/PathFollowingComponent.h"
-#include "Navigation/NavigationComponent.h"
 #include "GameFramework/Pawn.h"
 #include "GameFramework/PawnMovementComponent.h"
 #include "Engine/Canvas.h"
@@ -36,9 +35,6 @@ DEFINE_LOG_CATEGORY(LogAINavigation);
 AAIController::AAIController(const FObjectInitializer& ObjectInitializer)
 	: Super(ObjectInitializer)
 {
-	// set up navigation component
-	NavComponent = ObjectInitializer.CreateDefaultSubobject<UNavigationComponent>(this, TEXT("NavComponent"));
-
 	PathFollowingComponent = ObjectInitializer.CreateDefaultSubobject<UPathFollowingComponent>(this, TEXT("PathFollowingComponent"));
 	PathFollowingComponent->OnMoveFinished.AddUObject(this, &AAIController::OnMoveCompleted);
 
@@ -127,11 +123,6 @@ void AAIController::GrabDebugSnapshot(FVisualLogEntry* Snapshot) const
 		Snapshot->Location = GetPawn()->GetActorLocation();
 	}
 
-	if (NavComponent)
-	{
-		NavComponent->DescribeSelfToVisLog(Snapshot);
-	}
-
 	if (PathFollowingComponent)
 	{
 		PathFollowingComponent->DescribeSelfToVisLog(Snapshot);
@@ -480,9 +471,8 @@ void AAIController::UnPossess()
 	}
 }
 
-void AAIController::InitNavigationControl(UNavigationComponent*& PathFindingComp, UPathFollowingComponent*& PathFollowingComp)
+void AAIController::InitNavigationControl(UPathFollowingComponent*& PathFollowingComp)
 {
-	PathFindingComp = NavComponent;
 	PathFollowingComp = PathFollowingComponent;
 }
 
@@ -510,13 +500,14 @@ EPathFollowingRequestResult::Type AAIController::MoveToActor(AActor* Goal, float
 		}
 		else
 		{
-			FNavPathSharedPtr Path = FindPath(Goal, bUsePathfinding, UNavigationQueryFilter::GetQueryFilter(NavComponent ? NavComponent->GetNavData() : NULL, FilterClass));
-			if (Path.IsValid())
+			FPathFindingQuery Query;
+			const bool bValidQuery = PreparePathfinding(Query, Goal->GetActorLocation(), Goal, bUsePathfinding, FilterClass);
+			const FAIRequestID RequestID = bValidQuery ? RequestPathAndMove(Query, Goal, AcceptanceRadius, bStopOnOverlap, NULL) : FAIRequestID::InvalidRequest;
+
+			if (RequestID.IsValid())
 			{
 				bAllowStrafe = bCanStrafe;
-
-				const uint32 RequestID = RequestMove(Path, Goal, AcceptanceRadius, bStopOnOverlap);
-				Result = (RequestID != FAIRequestID::InvalidRequest) ? EPathFollowingRequestResult::RequestSuccessful : EPathFollowingRequestResult::Failed;
+				Result = EPathFollowingRequestResult::RequestSuccessful;
 			}
 		}
 	}
@@ -556,11 +547,20 @@ EPathFollowingRequestResult::Type AAIController::MoveToLocation(const FVector& D
 
 	FVector GoalLocation = Dest;
 
-	// fail if projection to navigation is required but it either failed or there's not navigation component
-	if (bCanRequestMove && bProjectDestinationToNavigation && (NavComponent == NULL || !NavComponent->ProjectPointToNavigation(Dest, GoalLocation)))
+	// fail if projection to navigation is required but it failed
+	if (bCanRequestMove && bProjectDestinationToNavigation)
 	{
-		UE_VLOG_LOCATION(this, LogAINavigation, Error, Dest, 30.f, FLinearColor::Red, TEXT("AAIController::MoveToLocation failed to project destination location to navmesh"));
-		bCanRequestMove = false;
+		UNavigationSystem* NavSys = UNavigationSystem::GetCurrent(GetWorld());
+		const FNavAgentProperties& AgentProps = GetNavAgentProperties();
+		FNavLocation ProjectedLocation;
+
+		if (NavSys && !NavSys->ProjectPointToNavigation(Dest, ProjectedLocation, AgentProps.GetExtent(), &AgentProps))
+		{
+			UE_VLOG_LOCATION(this, LogAINavigation, Error, Dest, 30.f, FLinearColor::Red, TEXT("AAIController::MoveToLocation failed to project destination location to navmesh"));
+			bCanRequestMove = false;
+		}
+
+		GoalLocation = ProjectedLocation.Location;
 	}
 
 	if (bCanRequestMove && PathFollowingComponent && PathFollowingComponent->HasReached(GoalLocation, AcceptanceRadius, bStopOnOverlap))
@@ -579,13 +579,14 @@ EPathFollowingRequestResult::Type AAIController::MoveToLocation(const FVector& D
 
 	if (bCanRequestMove)
 	{
-		FNavPathSharedPtr Path = FindPath(GoalLocation, bUsePathfinding, UNavigationQueryFilter::GetQueryFilter(NavComponent != NULL ? NavComponent->GetNavData() : NULL, FilterClass));
-		if (Path.IsValid())
+		FPathFindingQuery Query;
+		const bool bValidQuery = PreparePathfinding(Query, GoalLocation, NULL, true, FilterClass);
+		const FAIRequestID RequestID = bValidQuery ? RequestPathAndMove(Query, NULL, AcceptanceRadius, bStopOnOverlap, NULL) : FAIRequestID::InvalidRequest;
+
+		if (RequestID.IsValid())
 		{
 			bAllowStrafe = bCanStrafe;
-
-			const FAIRequestID RequestID = RequestMove(Path, NULL, AcceptanceRadius, bStopOnOverlap);
-			Result = RequestID.IsValid() ? EPathFollowingRequestResult::RequestSuccessful : EPathFollowingRequestResult::Failed;
+			Result = EPathFollowingRequestResult::RequestSuccessful;
 		}
 	}
 
@@ -639,32 +640,61 @@ void AAIController::StopMovement()
 	PathFollowingComponent->AbortMove(TEXT("StopMovement"));
 }
 
-FNavPathSharedPtr AAIController::FindPath(AActor* Goal, bool bUsePathfinding, TSharedPtr<const FNavigationQueryFilter> QueryFilter)
+bool AAIController::PreparePathfinding(FPathFindingQuery& Query, const FVector& Dest, AActor* Goal, bool bUsePathfinding, TSubclassOf<class UNavigationQueryFilter> FilterClass)
 {
-	if (NavComponent)
+	UNavigationSystem* NavSys = UNavigationSystem::GetCurrent(GetWorld());
+	if (NavSys)
 	{
-		const bool bFoundPath = bUsePathfinding ? NavComponent->FindPathToActor(Goal, QueryFilter) : NavComponent->FindSimplePathToActor(Goal);
-		if (bFoundPath)
+		ANavigationData* NavData = bUsePathfinding ?
+			NavSys->GetNavDataForProps(GetNavAgentProperties()) :
+			NavSys->GetAbstractNavData();
+
+		FVector GoalLocation = Dest;
+		if (Goal)
+		{
+			const INavAgentInterface* NavGoal = Cast<const INavAgentInterface>(Goal);
+			if (NavGoal)
+			{
+				const FVector Offset = NavGoal->GetMoveGoalOffset(this);
+				GoalLocation = FRotationTranslationMatrix(Goal->GetActorRotation(), NavGoal->GetNavAgentLocation()).TransformPosition(Offset);
+			}
+		}
+
+		Query = FPathFindingQuery(this, NavData, GetNavAgentLocation(), GoalLocation, UNavigationQueryFilter::GetQueryFilter(NavData, FilterClass));
+		
+		if (PathFollowingComponent)
 		{
-			return NavComponent->GetPath();
+			PathFollowingComponent->OnPathfindingQuery(Query);
 		}
+
+		return true;
 	}
 
-	return NULL;
+	return false;
 }
 
-FNavPathSharedPtr AAIController::FindPath(const FVector& Dest, bool bUsePathfinding, TSharedPtr<const FNavigationQueryFilter> QueryFilter)
+FAIRequestID AAIController::RequestPathAndMove(FPathFindingQuery& Query, AActor* Goal, float AcceptanceRadius, bool bStopOnOverlap, FCustomMoveSharedPtr CustomData)
 {
-	if (NavComponent)
+	FAIRequestID RequestID;
+
+	UNavigationSystem* NavSys = UNavigationSystem::GetCurrent(GetWorld());
+	if (NavSys)
 	{
-		const bool bFoundPath = bUsePathfinding ? NavComponent->FindPathToLocation(Dest, QueryFilter) : NavComponent->FindSimplePathToLocation(Dest);
-		if (bFoundPath) 
+		FPathFindingResult PathResult = NavSys->FindPathSync(Query);
+		if (PathResult.IsSuccessful() && PathResult.Path.IsValid())
 		{
-			return NavComponent->GetPath();
+			if (Goal)
+			{
+				PathResult.Path->SetGoalActorObservation(*Goal, 100.0f);
+			}
+
+			PathResult.Path->EnableRecalculationOnInvalidation(true);
 		}
+
+		RequestID = RequestMove(PathResult.Path, Goal, AcceptanceRadius, bStopOnOverlap, CustomData);
 	}
 
-	return NULL;
+	return RequestID;
 }
 
 EPathFollowingStatus::Type AAIController::GetMoveStatus() const
@@ -797,8 +827,6 @@ FString AAIController::GetDebugIcon() const
 }
 
 
-/** Returns NavComponent subobject **/
-UNavigationComponent* AAIController::GetNavComponent() const { return NavComponent; }
 /** Returns PathFollowingComponent subobject **/
 UPathFollowingComponent* AAIController::GetPathFollowingComponent() const { return PathFollowingComponent; }
 /** Returns ActionsComp subobject **/
diff --git a/Engine/Source/Runtime/AIModule/Private/BehaviorTree/Decorators/BTDecorator_DoesPathExist.cpp b/Engine/Source/Runtime/AIModule/Private/BehaviorTree/Decorators/BTDecorator_DoesPathExist.cpp
--- a/Engine/Source/Runtime/AIModule/Private/BehaviorTree/Decorators/BTDecorator_DoesPathExist.cpp
+++ b/Engine/Source/Runtime/AIModule/Private/BehaviorTree/Decorators/BTDecorator_DoesPathExist.cpp
@@ -4,7 +4,6 @@
 #if WITH_RECAST
 #include "AI/Navigation/RecastNavMesh.h"
 #endif // WITH_RECAST
-#include "Navigation/NavigationComponent.h"
 #include "BehaviorTree/BlackboardComponent.h"
 #include "BehaviorTree/Decorators/BTDecorator_DoesPathExist.h"
 
@@ -61,7 +60,7 @@ bool UBTDecorator_DoesPathExist::CalculateRawConditionValue(UBehaviorTreeCompone
 	if (NavSys && bHasPointA && bHasPointB)
 	{
 		const AAIController* AIOwner = OwnerComp->GetAIOwner();
-		const ANavigationData* NavData = AIOwner && AIOwner->GetNavComponent() ? AIOwner->GetNavComponent()->GetNavData() : NULL;
+		const ANavigationData* NavData = AIOwner ? NavSys->GetNavDataForProps(AIOwner->GetNavAgentProperties()) : NULL;
 		TSharedPtr<const FNavigationQueryFilter> QueryFilter = UNavigationQueryFilter::GetQueryFilter(NavData, FilterClass);
 
 		if (PathQueryType == EPathExistanceQueryType::NavmeshRaycast2D)
diff --git a/Engine/Source/Runtime/AIModule/Private/Blueprint/AIBlueprintHelperLibrary.cpp b/Engine/Source/Runtime/AIModule/Private/Blueprint/AIBlueprintHelperLibrary.cpp
--- a/Engine/Source/Runtime/AIModule/Private/Blueprint/AIBlueprintHelperLibrary.cpp
+++ b/Engine/Source/Runtime/AIModule/Private/Blueprint/AIBlueprintHelperLibrary.cpp
@@ -91,16 +91,23 @@ UAIAsyncTaskBlueprintProxy* UAIBlueprintHelperLibrary::CreateMoveToProxyObject(U
 	{
 		UWorld* World = GEngine->GetWorldFromContextObject( WorldContextObject );
 		MyObj = NewObject<UAIAsyncTaskBlueprintProxy>(World);
-		FNavPathSharedPtr Path = TargetActor ? AIController->FindPath(TargetActor, true) : AIController->FindPath(Destination, true);
-		if (Path.IsValid())
-		{
-			MyObj->AIController = AIController;
-			MyObj->AIController->ReceiveMoveCompleted.AddDynamic(MyObj, &UAIAsyncTaskBlueprintProxy::OnMoveCompleted);
-			MyObj->MoveRequestId = MyObj->AIController->RequestMove(Path, TargetActor, AcceptanceRadius, bStopOnOverlap);
-		}
-		else
+
+		FPathFindingQuery Query;
+		const bool bValidQuery = AIController->PreparePathfinding(Query, Destination, TargetActor, true);
+
+		if (bValidQuery)
 		{
-			World->GetTimerManager().SetTimer(MyObj, &UAIAsyncTaskBlueprintProxy::OnNoPath, 0.1, false);
+			const FAIRequestID RequestID = AIController->RequestPathAndMove(Query, NULL, AcceptanceRadius, bStopOnOverlap, NULL);
+			if (RequestID.IsValid())
+			{
+				MyObj->AIController = AIController;
+				MyObj->AIController->ReceiveMoveCompleted.AddDynamic(MyObj, &UAIAsyncTaskBlueprintProxy::OnMoveCompleted);
+				MyObj->MoveRequestId = RequestID;
+			}
+			else
+			{
+				World->GetTimerManager().SetTimer(MyObj, &UAIAsyncTaskBlueprintProxy::OnNoPath, 0.1f, false);
+			}
 		}
 	}
 	return MyObj;
diff --git a/Engine/Source/Runtime/AIModule/Private/Navigation/CrowdFollowingComponent.cpp b/Engine/Source/Runtime/AIModule/Private/Navigation/CrowdFollowingComponent.cpp
--- a/Engine/Source/Runtime/AIModule/Private/Navigation/CrowdFollowingComponent.cpp
+++ b/Engine/Source/Runtime/AIModule/Private/Navigation/CrowdFollowingComponent.cpp
@@ -3,8 +3,8 @@
 #include "AIModulePrivate.h"
 #include "Navigation/CrowdFollowingComponent.h"
 #include "Navigation/CrowdManager.h"
-#include "Navigation/NavigationComponent.h"
 #include "AI/Navigation/NavLinkCustomInterface.h"
+#include "AI/Navigation/AbstractNavData.h"
 #include "GameFramework/CharacterMovementComponent.h"
 
 
@@ -261,20 +261,6 @@ void UCrowdFollowingComponent::SetCrowdSimulation(bool bEnable)
 
 	UE_VLOG(GetOwner(), LogCrowdFollowing, Log, TEXT("SetCrowdSimulation: %s"), bEnable ? TEXT("enabled") : TEXT("disabled"));
 	bEnableCrowdSimulation = bEnable;
-
-	if (NavComp)
-	{
-		if (bEnableCrowdSimulation)
-		{
-			// disable path post processing (string pulling), crowd simulation needs to handle 
-			// large paths by splitting into smaller parts and optimization gets in the way
-			NavComp->SetNavDataFlag(ERecastPathFlags::SkipStringPulling);
-		}
-		else
-		{
-			NavComp->ClearNavDataFlag(ERecastPathFlags::SkipStringPulling);
-		}
-	}
 }
 
 void UCrowdFollowingComponent::Initialize()
@@ -291,13 +277,6 @@ void UCrowdFollowingComponent::Initialize()
 	{
 		bEnableCrowdSimulation = false;
 	}
-
-	if (bEnableCrowdSimulation && NavComp)
-	{
-		// disable path post processing (string pulling), crowd simulation needs to handle 
-		// large paths by splitting into smaller parts and optimization gets in the way
-		NavComp->SetNavDataFlag(ERecastPathFlags::SkipStringPulling);
-	}
 }
 
 void UCrowdFollowingComponent::Cleanup()
@@ -418,6 +397,17 @@ void UCrowdFollowingComponent::OnPathUpdated()
 	LastPathPolyIndex = 0;
 }
 
+void UCrowdFollowingComponent::OnPathfindingQuery(FPathFindingQuery& Query)
+{
+	// disable path post processing (string pulling), crowd simulation needs to handle 
+	// large paths by splitting into smaller parts and optimization gets in the way
+
+	if (bEnableCrowdSimulation)
+	{
+		Query.NavDataFlags |= ERecastPathFlags::SkipStringPulling;
+	}
+}
+
 bool UCrowdFollowingComponent::ShouldCheckPathOnResume() const
 {
 	if (bEnableCrowdSimulation)
@@ -529,6 +519,7 @@ void UCrowdFollowingComponent::SetMoveSegment(int32 SegmentStartIndex)
 	FVector CurrentTargetPt = Path->GetPathPoints()[1].Location;
 
 	FNavMeshPath* NavMeshPath = Path->CastPath<FNavMeshPath>();
+	FAbstractNavigationPath* DirectPath = Path->CastPath<FAbstractNavigationPath>();
 	if (NavMeshPath)
 	{
 #if WITH_RECAST
@@ -594,7 +585,7 @@ void UCrowdFollowingComponent::SetMoveSegment(int32 SegmentStartIndex)
 		}
 #endif
 	}
-	else
+	else if (DirectPath)
 	{
 		// direct paths are not using any steering or avoidance
 		// pathfinding is replaced with simple velocity request 
@@ -618,6 +609,10 @@ void UCrowdFollowingComponent::SetMoveSegment(int32 SegmentStartIndex)
 			CrowdManager->SetAgentMoveDirection(this, CrowdAgentMoveDirection);
 		}
 	}
+	else
+	{
+		UE_VLOG(GetOwner(), LogCrowdFollowing, Error, TEXT("SetMoveSegment, unknown path type!"));
+	}
 }
 
 void UCrowdFollowingComponent::UpdatePathSegment()
@@ -636,7 +631,7 @@ void UCrowdFollowingComponent::UpdatePathSegment()
 
 	if (!Path->IsValid())
 	{
-		if (NavComp == NULL || !NavComp->IsWaitingForRepath())
+		if (!Path->IsWaitingForRepath())
 		{
 			AbortMove(TEXT("no path"), FAIRequestID::CurrentRequest, true, false, EPathFollowingMessage::NoPath);
 		}
@@ -653,7 +648,8 @@ void UCrowdFollowingComponent::UpdatePathSegment()
 		if (bFinalPathPart)
 		{
 			const FVector ToTarget = (GoalLocation - MovementComp->GetActorFeetLocation());
-			const float SegmentDot = FVector::DotProduct(ToTarget, Path->IsDirect() ? MovementComp->Velocity : CrowdAgentMoveDirection);
+			const bool bDirectPath = Path->CastPath<FAbstractNavigationPath>() != NULL;
+			const float SegmentDot = FVector::DotProduct(ToTarget, bDirectPath ? MovementComp->Velocity : CrowdAgentMoveDirection);
 			const bool bMovedTooFar = bCheckMovementAngle && (SegmentDot < 0.0);
 
 			// can't use HasReachedDestination here, because it will use last path point
@@ -814,10 +810,14 @@ void UCrowdFollowingComponent::GetDebugStringTokens(TArray<FString>& Tokens, TAr
 		{
 			StatusDesc += FString::Printf(TEXT(" (path:%d, visited:%d)"), PathStartIndex, LastPathPolyIndex);
 		}
-		else
+		else if (Path->CastPath<FAbstractNavigationPath>())
 		{
 			StatusDesc += TEXT(" (direct)");
 		}
+		else
+		{
+			StatusDesc += TEXT(" (unknown path)");
+		}
 	}
 
 	// get cylinder of moving agent
@@ -878,13 +878,15 @@ void UCrowdFollowingComponent::DescribeSelfToVisLog(FVisualLogEntry* Snapshot) c
 	FString StatusDesc = GetStatusDesc();
 
 	FNavMeshPath* NavMeshPath = Path.IsValid() ? Path->CastPath<FNavMeshPath>() : NULL;
+	FAbstractNavigationPath* DirectPath = Path.IsValid() ? Path->CastPath<FAbstractNavigationPath>() : NULL;
+
 	if (Status == EPathFollowingStatus::Moving)
 	{
 		StatusDesc += FString::Printf(TEXT(" [path:%d, visited:%d]"), PathStartIndex, LastPathPolyIndex);
 	}
 
 	Category.Add(TEXT("Status"), StatusDesc);
-	Category.Add(TEXT("Path"), !Path.IsValid() ? TEXT("none") : NavMeshPath ? TEXT("navmesh") : TEXT("direct"));
+	Category.Add(TEXT("Path"), !Path.IsValid() ? TEXT("none") : NavMeshPath ? TEXT("navmesh") : DirectPath ? TEXT("direct") : TEXT("unknown"));
 
 	UObject* CustomLinkOb = GetCurrentCustomLinkOb();
 	if (CustomLinkOb)
diff --git a/Engine/Source/Runtime/AIModule/Private/Navigation/NavigationComponent.cpp b/Engine/Source/Runtime/AIModule/Private/Navigation/NavigationComponent.cpp
deleted file mode 100644
diff --git a/Engine/Source/Runtime/AIModule/Private/Navigation/PathFollowingComponent.cpp b/Engine/Source/Runtime/AIModule/Private/Navigation/PathFollowingComponent.cpp
--- a/Engine/Source/Runtime/AIModule/Private/Navigation/PathFollowingComponent.cpp
+++ b/Engine/Source/Runtime/AIModule/Private/Navigation/PathFollowingComponent.cpp
@@ -5,7 +5,6 @@
 #	include "Detour/DetourNavMeshQuery.h"
 #endif
 #include "AI/Navigation/RecastNavMesh.h"
-#include "Navigation/NavigationComponent.h"
 #include "AI/Navigation/NavLinkCustomInterface.h"
 #include "GameFramework/NavMovementComponent.h"
 #include "GameFramework/Character.h"
@@ -263,11 +262,6 @@ void UPathFollowingComponent::AbortMove(const FString& Reason, FAIRequestID Requ
 			MovementComp->StopMovementKeepPathing();
 		}
 
-		if (NavComp && FinishResult != EPathFollowingResult::Skipped)
-		{
-			NavComp->ResetTransientData();
-		}
-
 		// notify observers after state was reset (they can request another move)
 		SavedReqFinished.ExecuteIfBound(FinishResult);
 		OnMoveFinished.Broadcast(AbortedMoveId, FinishResult);
@@ -299,6 +293,8 @@ void UPathFollowingComponent::PauseMove(FAIRequestID RequestID, bool bResetVeloc
 
 		UpdateMoveFocus();
 	}
+
+	// TODO: pause path updates with goal movement
 }
 
 void UPathFollowingComponent::ResumeMove(FAIRequestID RequestID)
@@ -380,11 +376,6 @@ void UPathFollowingComponent::OnPathFinished(EPathFollowingResult::Type Result)
 		MovementComp->StopMovementKeepPathing();
 	}
 
-	if (NavComp && Result != EPathFollowingResult::Skipped)
-	{
-		NavComp->ResetTransientData();
-	}
-
 	// notify observers after state was reset (they can request another move)
 	SavedReqFinished.ExecuteIfBound(Result);
 	OnMoveFinished.Broadcast(FinishedMoveId, Result);
@@ -413,12 +404,6 @@ void UPathFollowingComponent::Cleanup()
 
 void UPathFollowingComponent::UpdateCachedComponents()
 {
-	AActor* MyOwner = GetOwner();
-	if (MyOwner)
-	{
-		NavComp = MyOwner->FindComponentByClass<UNavigationComponent>();
-	}
-
 	UpdateMovementComponent(/*bForce=*/true);
 }
 
@@ -623,13 +608,8 @@ int32 UPathFollowingComponent::OptimizeSegmentVisibility(int32 StartIndex)
 	//SCOPE_CYCLE_COUNTER(STAT_Navigation_PathVisibilityOptimisation);
 
 	const AAIController* MyAI = Cast<AAIController>(GetOwner());
-	const ANavigationData* NavData = MyAI && MyAI->GetNavComponent() ? MyAI->GetNavComponent()->GetNavData() : NULL;
-	if (Path.IsValid())
-	{
-		Path->ShortcutNodeRefs.Reset();
-	}
-
-	if (NavData == NULL)
+	const ANavigationData* NavData = Path.IsValid() ? Path->GetNavigationDataUsed() : NULL;
+	if (NavData == NULL || MyAI == NULL)
 	{
 		return StartIndex + 1;
 	}
@@ -651,7 +631,7 @@ int32 UPathFollowingComponent::OptimizeSegmentVisibility(int32 StartIndex)
 	}
 
 #endif
-	TSharedPtr<FNavigationQueryFilter> QueryFilter = (NavComp && NavComp->GetStoredQueryFilter().IsValid()) ? NavComp->GetStoredQueryFilter()->GetCopy() : NavData->GetDefaultQueryFilter()->GetCopy();
+	TSharedPtr<FNavigationQueryFilter> QueryFilter = Path->GetFilter()->GetCopy();
 #if WITH_RECAST
 	const uint8 StartArea = FNavMeshNodeFlags(Path->GetPathPoints()[StartIndex].Flags).Area;
 	TArray<float> CostArray;
@@ -662,6 +642,7 @@ int32 UPathFollowingComponent::OptimizeSegmentVisibility(int32 StartIndex)
 
 	const int32 MaxPoints = Path->GetPathPoints().Num();
 	int32 Index = StartIndex + 2;
+	Path->ShortcutNodeRefs.Reset();
 
 	for (; Index <= MaxPoints; ++Index)
 	{
@@ -725,7 +706,7 @@ void UPathFollowingComponent::UpdatePathSegment()
 
 	if (!Path->IsValid())
 	{
-		if (NavComp == NULL || !NavComp->IsWaitingForRepath())
+		if (!Path->IsWaitingForRepath())
 		{
 			AbortMove(TEXT("no path"), FAIRequestID::CurrentRequest, true, false, EPathFollowingMessage::NoPath);
 		}
@@ -1242,7 +1223,7 @@ EPathFollowingAction::Type UPathFollowingComponent::GetPathActionType() const
 		case EPathFollowingStatus::Paused:	
 		case EPathFollowingStatus::Moving:
 			return Path.IsValid() == false ? EPathFollowingAction::Error :
-				Path->IsDirect() ? EPathFollowingAction::DirectMove :
+				Path->CastPath<FAbstractNavigationPath>() ? EPathFollowingAction::DirectMove :
 				Path->IsPartial() ? EPathFollowingAction::PartialPath : EPathFollowingAction::PathToGoal;
 
 		default:
@@ -1323,7 +1304,10 @@ void UPathFollowingComponent::DescribeSelfToVisLog(FVisualLogEntry* Snapshot) co
 	}
 
 	Category.Add(TEXT("Status"), StatusDesc);
-	Category.Add(TEXT("Path"), !Path.IsValid() ? TEXT("none") : (Path->CastPath<FNavMeshPath>() != NULL) ? TEXT("navmesh") : TEXT("direct"));
+	Category.Add(TEXT("Path"), !Path.IsValid() ? TEXT("none") :
+		(Path->CastPath<FNavMeshPath>() != NULL) ? TEXT("navmesh") :
+		(Path->CastPath<FAbstractNavigationPath>() != NULL) ? TEXT("direct") :
+		TEXT("unknown"));
 	
 	UObject* CustomNavLinkOb = CurrentCustomLinkOb.Get();
 	if (CustomNavLinkOb)
@@ -1349,7 +1333,7 @@ void UPathFollowingComponent::GetDebugStringTokens(TArray<FString>& Tokens, TArr
 	if (Path.IsValid())
 	{
 		const int32 NumMoveSegments = (Path.IsValid() && Path->IsValid()) ? Path->GetPathPoints().Num() : -1;
-		const bool bIsDirect = HasDirectPath();
+		const bool bIsDirect = (Path->CastPath<FAbstractNavigationPath>() != NULL);
 		const bool bIsCustomLink = CurrentCustomLinkOb.IsValid();
 
 		if (!bIsDirect)
diff --git a/Engine/Source/Runtime/Engine/Classes/AI/Navigation/AbstractNavData.h b/Engine/Source/Runtime/Engine/Classes/AI/Navigation/AbstractNavData.h
new file mode 100644
--- /dev/null
+++ b/Engine/Source/Runtime/Engine/Classes/AI/Navigation/AbstractNavData.h
@@ -0,0 +1,38 @@
+// Copyright 1998-2014 Epic Games, Inc. All Rights Reserved.
+
+#pragma once
+#include "AI/Navigation/NavigationData.h"
+#include "AI/Navigation/NavigationTypes.h"
+#include "AbstractNavData.generated.h"
+
+struct ENGINE_API FAbstractNavigationPath : public FNavigationPath
+{
+	typedef FNavigationPath Super;
+
+	FAbstractNavigationPath();
+
+	static const FNavPathType Type;
+};
+
+UCLASS()
+class AAbstractNavData : public ANavigationData
+{
+	GENERATED_UCLASS_BODY()
+
+	// Begin ANavigationData overrides
+	virtual void BatchRaycast(TArray<FNavigationRaycastWork>& Workload, TSharedPtr<const FNavigationQueryFilter> QueryFilter, const UObject* Querier = NULL) const override {};
+	virtual FNavLocation GetRandomPoint(TSharedPtr<const FNavigationQueryFilter> Filter = NULL, const UObject* Querier = NULL) const override { return FNavLocation();  }
+	virtual bool GetRandomPointInRadius(const FVector& Origin, float Radius, FNavLocation& OutResult, TSharedPtr<const FNavigationQueryFilter> Filter = NULL, const UObject* Querier = NULL) const override { return false; }
+	virtual bool ProjectPoint(const FVector& Point, FNavLocation& OutLocation, const FVector& Extent, TSharedPtr<const FNavigationQueryFilter> Filter = NULL, const UObject* Querier = NULL) const override { return false; }
+	virtual void BatchProjectPoints(TArray<FNavigationProjectionWork>& Workload, const FVector& Extent, TSharedPtr<const FNavigationQueryFilter> Filter = NULL, const UObject* Querier = NULL) const override {};
+	virtual ENavigationQueryResult::Type CalcPathCost(const FVector& PathStart, const FVector& PathEnd, float& OutPathCost, TSharedPtr<const FNavigationQueryFilter> QueryFilter = NULL, const UObject* Querier = NULL) const override { return ENavigationQueryResult::Invalid; }
+	virtual ENavigationQueryResult::Type CalcPathLength(const FVector& PathStart, const FVector& PathEnd, float& OutPathLength, TSharedPtr<const FNavigationQueryFilter> QueryFilter = NULL, const UObject* Querier = NULL) const override { return ENavigationQueryResult::Invalid; }
+	virtual ENavigationQueryResult::Type CalcPathLengthAndCost(const FVector& PathStart, const FVector& PathEnd, float& OutPathLength, float& OutPathCost, TSharedPtr<const FNavigationQueryFilter> QueryFilter = NULL, const UObject* Querier = NULL) const override { return ENavigationQueryResult::Invalid;	}
+	virtual void OnNavAreaAdded(const UClass* NavAreaClass, int32 AgentIndex) override {}
+	virtual void OnNavAreaRemoved(const UClass* NavAreaClass) override {};
+	// End ANavigationData overrides
+
+	static FPathFindingResult FindPathAbstract(const FNavAgentProperties& AgentProperties, const FPathFindingQuery& Query);
+	static bool TestPathAbstract(const FNavAgentProperties& AgentProperties, const FPathFindingQuery& Query, int32* NumVisitedNodes);
+	static bool RaycastAbstract(const ANavigationData* NavDataInstance, const FVector& RayStart, const FVector& RayEnd, FVector& HitLocation, TSharedPtr<const FNavigationQueryFilter> QueryFilter, const UObject* Querier);
+};
diff --git a/Engine/Source/Runtime/Engine/Classes/AI/Navigation/NavAgentInterface.h b/Engine/Source/Runtime/Engine/Classes/AI/Navigation/NavAgentInterface.h
--- a/Engine/Source/Runtime/Engine/Classes/AI/Navigation/NavAgentInterface.h
+++ b/Engine/Source/Runtime/Engine/Classes/AI/Navigation/NavAgentInterface.h
@@ -38,4 +38,7 @@ class INavAgentInterface
 	 * @param GoalHalfHeight - cylinder half height
 	 */
 	virtual void GetMoveGoalReachTest(AActor* MovingActor, const FVector& MoveOffset, FVector& GoalOffset, float& GoalRadius, float& GoalHalfHeight) const {}
+
+	/** Allows delaying repath requests */
+	virtual bool ShouldPostponePathUpdates() const { return false; }
 };
diff --git a/Engine/Source/Runtime/Engine/Classes/AI/Navigation/NavLinkCustomComponent.h b/Engine/Source/Runtime/Engine/Classes/AI/Navigation/NavLinkCustomComponent.h
--- a/Engine/Source/Runtime/Engine/Classes/AI/Navigation/NavLinkCustomComponent.h
+++ b/Engine/Source/Runtime/Engine/Classes/AI/Navigation/NavLinkCustomComponent.h
@@ -27,7 +27,7 @@ class ENGINE_API UNavLinkCustomComponent : public UNavRelevantComponent, public
 	GENERATED_UCLASS_BODY()
 
 	DECLARE_DELEGATE_ThreeParams(FOnMoveReachedLink, UNavLinkCustomComponent* /*ThisComp*/, UPathFollowingComponent* /*PathComp*/, const FVector& /*DestPoint*/);
-	DECLARE_DELEGATE_TwoParams(FBroadcastFilter, UNavLinkCustomComponent* /*ThisComp*/, TArray<UNavigationComponent*>& /*NotifyList*/);
+	DECLARE_DELEGATE_TwoParams(FBroadcastFilter, UNavLinkCustomComponent* /*ThisComp*/, TArray<UPathFollowingComponent*>& /*NotifyList*/);
 
 	// BEGIN INavLinkCustomInterface
 	virtual void GetLinkData(FVector& LeftPt, FVector& RightPt, ENavLinkDirection::Type& Direction) const override;
@@ -206,5 +206,5 @@ protected:
 	void BroadcastStateChange();
 	
 	/** gather agents to notify about state change */
-	void CollectNearbyAgents(TArray<UNavigationComponent*>& NotifyList);
+	void CollectNearbyAgents(TArray<UPathFollowingComponent*>& NotifyList);
 };
diff --git a/Engine/Source/Runtime/Engine/Classes/AI/Navigation/NavigationData.h b/Engine/Source/Runtime/Engine/Classes/AI/Navigation/NavigationData.h
--- a/Engine/Source/Runtime/Engine/Classes/AI/Navigation/NavigationData.h
+++ b/Engine/Source/Runtime/Engine/Classes/AI/Navigation/NavigationData.h
@@ -85,9 +85,9 @@ struct ENGINE_API FNavigationPath : public TSharedFromThis<FNavigationPath, ESPM
 	{
 		return bReachedSearchLimit;
 	}
-	FORCEINLINE bool IsDirect() const
+	FORCEINLINE bool IsWaitingForRepath() const
 	{
-		return NavigationDataUsed.Get() == NULL;
+		return bWaitingForRepath;
 	}
 	FORCEINLINE FVector GetDestinationLocation() const
 	{
@@ -159,8 +159,11 @@ struct ENGINE_API FNavigationPath : public TSharedFromThis<FNavigationPath, ESPM
 
 	FORCEINLINE void DoneUpdating(ENavPathUpdateType::Type UpdateType)
 	{
-		bUpToDate = true; ObserverDelegate.Broadcast(this, UpdateType == ENavPathUpdateType::GoalMoved ? ENavPathEvent::UpdatedDueToGoalMoved : ENavPathEvent::UpdatedDueToNavigationChanged);
+		bUpToDate = true;
+		bWaitingForRepath = false;
+		ObserverDelegate.Broadcast(this, UpdateType == ENavPathUpdateType::GoalMoved ? ENavPathEvent::UpdatedDueToGoalMoved : ENavPathEvent::UpdatedDueToNavigationChanged);
 	}
+
 	void Invalidate();
 	void RePathFailed();
 
@@ -351,6 +354,9 @@ protected:
 	/** if true path will request re-pathing if it gets invalidated due to underlying navigation changed */
 	uint32 bDoAutoUpdateOnInvalidation : 1;
 
+	/** set when path is waiting for recalc from navigation data */
+	uint32 bWaitingForRepath : 1;
+
 	/** navigation data used to generate this path */
 	TWeakObjectPtr<ANavigationData> NavigationDataUsed;
 
diff --git a/Engine/Source/Runtime/Engine/Classes/AI/Navigation/NavigationSystem.h b/Engine/Source/Runtime/Engine/Classes/AI/Navigation/NavigationSystem.h
--- a/Engine/Source/Runtime/Engine/Classes/AI/Navigation/NavigationSystem.h
+++ b/Engine/Source/Runtime/Engine/Classes/AI/Navigation/NavigationSystem.h
@@ -133,6 +133,10 @@ class ENGINE_API UNavigationSystem : public UBlueprintFunctionLibrary
 	UPROPERTY()
 	ANavigationData* MainNavData;
 
+	/** special navigation data for managing direct paths, not part of NavDataSet! */
+	UPROPERTY()
+	ANavigationData* AbstractNavData;
+
 	/** Should navigation system spawn default Navigation Data when there's none and there are navigation bounds present? */
 	UPROPERTY(config, EditAnywhere, Category=NavigationSystem)
 	uint32 bAutoCreateNavigationData:1;
@@ -388,6 +392,8 @@ public:
 	/** Returns the world nav mesh object.  Creates one if it doesn't exist. */
 	const ANavigationData* GetMainNavData() const { return MainNavData; }
 
+	ANavigationData* GetAbstractNavData() const { return AbstractNavData; }
+
 	TSharedPtr<FNavigationQueryFilter> CreateDefaultQueryFilterCopy() const;
 
 	/** Super-hacky safety feature for threaded navmesh building. Will be gone once figure out why keeping TSharedPointer to Navigation Generator doesn't 
diff --git a/Engine/Source/Runtime/Engine/Classes/AI/Navigation/NavigationTypes.h b/Engine/Source/Runtime/Engine/Classes/AI/Navigation/NavigationTypes.h
--- a/Engine/Source/Runtime/Engine/Classes/AI/Navigation/NavigationTypes.h
+++ b/Engine/Source/Runtime/Engine/Classes/AI/Navigation/NavigationTypes.h
@@ -448,7 +448,10 @@ struct ENGINE_API FPathFindingQuery
 	/** additional flags passed to navigation data handling request */
 	int32 NavDataFlags;
 
-	FPathFindingQuery();
+	/** if set, allow partial paths as a result */
+	uint32 bAllowPartialPaths : 1;
+
+	FPathFindingQuery() {}
 
 	FPathFindingQuery(const FPathFindingQuery& Source);
 
diff --git a/Engine/Source/Runtime/Engine/Classes/GameFramework/Controller.h b/Engine/Source/Runtime/Engine/Classes/GameFramework/Controller.h
--- a/Engine/Source/Runtime/Engine/Classes/GameFramework/Controller.h
+++ b/Engine/Source/Runtime/Engine/Classes/GameFramework/Controller.h
@@ -4,7 +4,6 @@
 #include "AI/Navigation/NavAgentInterface.h"
 #include "Controller.generated.h"
 
-class UNavigationComponent;
 class UPathFollowingComponent;
 
 UDELEGATE(BlueprintAuthorityOnly)
@@ -267,13 +266,11 @@ public:
 	virtual const struct FNavAgentProperties& GetNavAgentProperties() const override;
 	virtual FVector GetNavAgentLocation() const override;
 	virtual void GetMoveGoalReachTest(class AActor* MovingActor, const FVector& MoveOffset, FVector& GoalOffset, float& GoalRadius, float& GoalHalfHeight) const override;
+	virtual bool ShouldPostponePathUpdates() const override;
 	// End INavAgentInterface Interface
 
-	/** Allows agent to postpone any path updates (e.g. locked by gameplay). */
-	bool ShouldPostponePathUpdates() const;
-	
-	/** prepares path finding and path following components */
-	virtual void InitNavigationControl(UNavigationComponent*& PathFindingComp, UPathFollowingComponent*& PathFollowingComp);
+	/** prepares path following component */
+	virtual void InitNavigationControl(UPathFollowingComponent*& PathFollowingComp);
 
 	/** If controller has any navigation-related components then this function 
 	 *	makes them update their cached data */
diff --git a/Engine/Source/Runtime/Engine/Classes/GameFramework/Pawn.h b/Engine/Source/Runtime/Engine/Classes/GameFramework/Pawn.h
--- a/Engine/Source/Runtime/Engine/Classes/GameFramework/Pawn.h
+++ b/Engine/Source/Runtime/Engine/Classes/GameFramework/Pawn.h
@@ -237,9 +237,6 @@ public:
 	virtual void GetMoveGoalReachTest(AActor* MovingActor, const FVector& MoveOffset, FVector& GoalOffset, float& GoalRadius, float& GoalHalfHeight) const override;
 	// End INavAgentInterface Interface
 
-	/** Allows agent to postpone any path updates (e.g. locked by gameplay) */
-	virtual bool ShouldPostponePathUpdates() const { return false; }
-
 	/** updates MovementComponent's parameters used by navigation system */
 	void UpdateNavAgent();
 
diff --git a/Engine/Source/Runtime/Engine/Private/AI/Navigation/AbstractNavData.cpp b/Engine/Source/Runtime/Engine/Private/AI/Navigation/AbstractNavData.cpp
new file mode 100644
--- /dev/null
+++ b/Engine/Source/Runtime/Engine/Private/AI/Navigation/AbstractNavData.cpp
@@ -0,0 +1,53 @@
+// Copyright 1998-2014 Epic Games, Inc. All Rights Reserved.
+
+#include "EnginePrivate.h"
+#include "AI/Navigation/AbstractNavData.h"
+
+const FNavPathType FAbstractNavigationPath::Type;
+
+FAbstractNavigationPath::FAbstractNavigationPath()
+{
+	PathType = FAbstractNavigationPath::Type;
+}
+
+AAbstractNavData::AAbstractNavData(const FObjectInitializer& ObjectInitializer) : Super(ObjectInitializer)
+{
+	if (HasAnyFlags(RF_ClassDefaultObject) == false)
+	{
+		FindPathImplementation = FindPathAbstract;
+		FindHierarchicalPathImplementation = FindPathAbstract;
+
+		TestPathImplementation = TestPathAbstract;
+		TestHierarchicalPathImplementation = TestPathAbstract;
+
+		RaycastImplementation = RaycastAbstract;
+	}
+}
+
+FPathFindingResult AAbstractNavData::FindPathAbstract(const FNavAgentProperties& AgentProperties, const FPathFindingQuery& Query)
+{
+	const ANavigationData* Self = Query.NavData.Get();
+	if (Self == NULL)
+	{
+		return ENavigationQueryResult::Error;
+	}
+
+	FPathFindingResult Result;
+	Result.Path = Query.PathInstanceToFill.IsValid() ? Query.PathInstanceToFill : Self->CreatePathInstance<FNavigationPath>(Query.Owner.Get());
+
+	Result.Path->GetPathPoints().Add(FNavPathPoint(Query.StartLocation));
+	Result.Path->GetPathPoints().Add(FNavPathPoint(Query.EndLocation));
+	Result.Result = ENavigationQueryResult::Success;
+
+	return Result;
+}
+
+bool AAbstractNavData::TestPathAbstract(const FNavAgentProperties& AgentProperties, const FPathFindingQuery& Query, int32* NumVisitedNodes)
+{
+	return false;
+}
+
+bool AAbstractNavData::RaycastAbstract(const ANavigationData* NavDataInstance, const FVector& RayStart, const FVector& RayEnd, FVector& HitLocation, TSharedPtr<const FNavigationQueryFilter> QueryFilter, const UObject* Querier)
+{
+	return false;
+}
diff --git a/Engine/Source/Runtime/Engine/Private/AI/Navigation/NavFilters/NavigationQueryFilter.cpp b/Engine/Source/Runtime/Engine/Private/AI/Navigation/NavFilters/NavigationQueryFilter.cpp
--- a/Engine/Source/Runtime/Engine/Private/AI/Navigation/NavFilters/NavigationQueryFilter.cpp
+++ b/Engine/Source/Runtime/Engine/Private/AI/Navigation/NavFilters/NavigationQueryFilter.cpp
@@ -2,7 +2,6 @@
 
 #include "EnginePrivate.h"
 // @todo AIModule circular dependency
-#include "Navigation/NavigationComponent.h"
 #include "AI/Navigation/NavFilters/NavigationQueryFilter.h"
 
 //----------------------------------------------------------------------//
diff --git a/Engine/Source/Runtime/Engine/Private/AI/Navigation/NavLinkCustomComponent.cpp b/Engine/Source/Runtime/Engine/Private/AI/Navigation/NavLinkCustomComponent.cpp
--- a/Engine/Source/Runtime/Engine/Private/AI/Navigation/NavLinkCustomComponent.cpp
+++ b/Engine/Source/Runtime/Engine/Private/AI/Navigation/NavLinkCustomComponent.cpp
@@ -6,7 +6,6 @@
 #include "AI/Navigation/NavAreas/NavArea_Default.h"
 // @todo to be addressed when removing AIModule circular dependency
 #include "Navigation/PathFollowingComponent.h"
-#include "Navigation/NavigationComponent.h"
 #include "AI/Navigation/NavLinkCustomComponent.h"
 #include "AI/Navigation/NavAreas/NavArea_Default.h"
 #include "AI/Navigation/NavAreas/NavArea_Null.h"
@@ -255,7 +254,7 @@ void UNavLinkCustomComponent::SendBroadcastWhenDisabled(bool bEnabled)
 	bNotifyWhenDisabled = bEnabled;
 }
 
-void UNavLinkCustomComponent::CollectNearbyAgents(TArray<UNavigationComponent*>& NotifyList)
+void UNavLinkCustomComponent::CollectNearbyAgents(TArray<UPathFollowingComponent*>& NotifyList)
 {
 	AActor* MyOwner = GetOwner();
 	if (BroadcastRadius < KINDA_SMALL_NUMBER || MyOwner == NULL)
@@ -302,8 +301,8 @@ void UNavLinkCustomComponent::CollectNearbyAgents(TArray<UNavigationComponent*>&
 
 	for (int32 i = 0; i < PawnList.Num(); i++)
 	{
-		UNavigationComponent* NavComp = PawnList[i]->GetController()->FindComponentByClass<UNavigationComponent>();
-		if (NavComp && NavComp->WantsSmartLinkUpdates())
+		UPathFollowingComponent* NavComp = PawnList[i]->GetController()->FindComponentByClass<UPathFollowingComponent>();
+		if (NavComp)// && NavComp->WantsSmartLinkUpdates())
 		{
 			NotifyList.Add(NavComp);
 		}
@@ -312,14 +311,14 @@ void UNavLinkCustomComponent::CollectNearbyAgents(TArray<UNavigationComponent*>&
 
 void UNavLinkCustomComponent::BroadcastStateChange()
 {
-	TArray<UNavigationComponent*> NearbyAgents;
+	TArray<UPathFollowingComponent*> NearbyAgents;
 
 	CollectNearbyAgents(NearbyAgents);
 	OnBroadcastFilter.ExecuteIfBound(this, NearbyAgents);
 
 	for (int32 i = 0; i < NearbyAgents.Num(); i++)
 	{
-		NearbyAgents[i]->OnCustomLinkBroadcast(this);
+//		NearbyAgents[i]->OnCustomLinkBroadcast(this);
 	}
 
 	if (BroadcastInterval > 0.0f)
diff --git a/Engine/Source/Runtime/Engine/Private/AI/Navigation/NavigationData.cpp b/Engine/Source/Runtime/Engine/Private/AI/Navigation/NavigationData.cpp
--- a/Engine/Source/Runtime/Engine/Private/AI/Navigation/NavigationData.cpp
+++ b/Engine/Source/Runtime/Engine/Private/AI/Navigation/NavigationData.cpp
@@ -16,6 +16,7 @@ FPathFindingQuery::FPathFindingQuery(const UObject* InOwner, const ANavigationDa
 , QueryFilter(SourceQueryFilter)
 , PathInstanceToFill(InPathInstanceToFill)
 , NavDataFlags(0)
+, bAllowPartialPaths(true)
 {
 	if (SourceQueryFilter.IsValid() == false && NavData.IsValid() == true)
 	{
@@ -31,6 +32,7 @@ FPathFindingQuery::FPathFindingQuery(const FPathFindingQuery& Source)
 , QueryFilter(Source.QueryFilter)
 , PathInstanceToFill(Source.PathInstanceToFill)
 , NavDataFlags(Source.NavDataFlags)
+, bAllowPartialPaths(Source.bAllowPartialPaths)
 {
 	if (Source.QueryFilter.IsValid() == false && NavData.IsValid() == true)
 	{
@@ -46,6 +48,7 @@ FPathFindingQuery::FPathFindingQuery(FNavPathSharedRef PathToRecalculate, const
 , QueryFilter(PathToRecalculate->GetFilter())
 , PathInstanceToFill(PathToRecalculate)
 , NavDataFlags(0)
+, bAllowPartialPaths(true)
 {
 	if (QueryFilter.IsValid() == false && NavData.IsValid() == true)
 	{
@@ -215,9 +218,19 @@ void ANavigationData::TickActor(float DeltaTime, enum ELevelTick TickType, FActo
 
 	if (RepathRequests.Num() > 0)
 	{
+		TArray<FNavPathRecalculationRequest> PostponedRequests;
+
 		// @todo batch-process it!
 		for (auto RecalcRequest : RepathRequests)
 		{
+			// check if it can be updated right now
+			const UObject* PathQuerier = RecalcRequest.Path->GetQuerier();
+			const INavAgentInterface* PathNavAgent = Cast<const INavAgentInterface>(PathQuerier);
+			if (PathNavAgent && PathNavAgent->ShouldPostponePathUpdates())
+			{
+				PostponedRequests.Add(RecalcRequest);
+			}
+
 			FPathFindingQuery Query(RecalcRequest.Path);
 			// @todo consider supplying NavAgentPropertied from path's querier
 			const FPathFindingResult Result = FindPath(FNavAgentProperties(), Query.SetPathInstanceToUpdate(RecalcRequest.Path));
@@ -239,6 +252,7 @@ void ANavigationData::TickActor(float DeltaTime, enum ELevelTick TickType, FActo
 		}
 
 		RepathRequests.Reset();
+		RepathRequests.Append(PostponedRequests);
 	}
 }
 
diff --git a/Engine/Source/Runtime/Engine/Private/AI/Navigation/NavigationPath.cpp b/Engine/Source/Runtime/Engine/Private/AI/Navigation/NavigationPath.cpp
--- a/Engine/Source/Runtime/Engine/Private/AI/Navigation/NavigationPath.cpp
+++ b/Engine/Source/Runtime/Engine/Private/AI/Navigation/NavigationPath.cpp
@@ -123,6 +123,7 @@ void FNavigationPath::Invalidate()
 	ObserverDelegate.Broadcast(this, ENavPathEvent::Invalidated);
 	if (bDoAutoUpdateOnInvalidation && NavigationDataUsed.IsValid())
 	{
+		bWaitingForRepath = true;
 		NavigationDataUsed->RequestRePath(AsShared(), ENavPathUpdateType::NavigationChanged);
 	}
 }
@@ -130,6 +131,7 @@ void FNavigationPath::Invalidate()
 void FNavigationPath::RePathFailed()
 {
 	ObserverDelegate.Broadcast(this, ENavPathEvent::RePathFailed);
+	bWaitingForRepath = false;
 }
 
 void FNavigationPath::DebugDraw(const ANavigationData* NavData, FColor PathColor, UCanvas* Canvas, bool bPersistent, const uint32 NextPathPointIndex) const
diff --git a/Engine/Source/Runtime/Engine/Private/AI/Navigation/NavigationSystem.cpp b/Engine/Source/Runtime/Engine/Private/AI/Navigation/NavigationSystem.cpp
--- a/Engine/Source/Runtime/Engine/Private/AI/Navigation/NavigationSystem.cpp
+++ b/Engine/Source/Runtime/Engine/Private/AI/Navigation/NavigationSystem.cpp
@@ -19,13 +19,13 @@
 // @todo this is here only due to circular dependency to AIModule. To be removed
 #include "Navigation/CrowdManager.h"
 #include "Navigation/PathFollowingComponent.h"
-#include "Navigation/NavigationComponent.h"
 #include "AI/Navigation/NavAreas/NavArea_Null.h"
 #include "AI/Navigation/NavAreas/NavArea_Default.h"
 #include "AI/Navigation/NavLinkCustomInterface.h"
 #include "AI/Navigation/NavigationSystem.h"
 #include "AI/Navigation/NavRelevantComponent.h"
 #include "AI/Navigation/NavigationPath.h"
+#include "AI/Navigation/AbstractNavData.h"
 
 static const uint32 INITIAL_ASYNC_QUERIES_SIZE = 32;
 static const uint32 REGISTRATION_QUEUE_SIZE = 16;	// and we'll not reallocate
@@ -278,6 +278,25 @@ void UNavigationSystem::DoInitialSetup()
 		UE_LOG(LogNavigation, Error, TEXT("No navigation data types found while setting up required navigation types!"));
 	}
 
+	// spawn abstract nav data separately
+	// it's responsible for direct paths and shouldn't be picked for any agent type as default one
+	UWorld* NavWorld = GetWorld();
+	for (TActorIterator<AAbstractNavData> It(NavWorld); It; ++It)
+	{
+		AAbstractNavData* Nav = (*It);
+		if (Nav && !Nav->IsPendingKill())
+		{
+			AbstractNavData = Nav;
+			break;
+		}
+	}
+
+	if (AbstractNavData == NULL)
+	{
+		FNavDataConfig DummyConfig;
+		AbstractNavData = CreateNavigationDataInstance(AAbstractNavData::StaticClass(), NavWorld, DummyConfig);
+	}
+
 	CreateCrowdManager();
 
 	bInitialSetupHasBeenPerformed = true;
@@ -872,12 +891,10 @@ void UNavigationSystem::SimpleMoveToActor(AController* Controller, const AActor*
 		return;
 	}
 
-	UNavigationComponent* PFindComp = NULL;
 	UPathFollowingComponent* PFollowComp = NULL;
-	
-	Controller->InitNavigationControl(PFindComp, PFollowComp);
+	UNavigationSystem* NavSys = UNavigationSystem::GetCurrent(Controller->GetWorld());
 
-	if (PFindComp == nullptr || PFollowComp == nullptr)
+	if (PFollowComp == nullptr)
 	{
 		FMessageLog("PIE").Warning(FText::Format(
 			LOCTEXT("SimpleMoveErrorNoComp", "SimpleMove failed for {0}: missing components"),
@@ -895,12 +912,16 @@ void UNavigationSystem::SimpleMoveToActor(AController* Controller, const AActor*
 		return;
 	}
 
-	if (!PFollowComp->HasReached(*Goal))
+	if (NavSys && !PFollowComp->HasReached(*Goal))
 	{
-		const bool bPathExists = PFindComp->FindPathToActor(Goal);
-		if (bPathExists)
+		const ANavigationData* NavData = NavSys->GetNavDataForProps(Controller->GetNavAgentProperties());
+		FPathFindingQuery Query(Controller, NavData, Controller->GetNavAgentLocation(), Goal->GetActorLocation());
+		FPathFindingResult Result = NavSys->FindPathSync(Query);
+		if (Result.IsSuccessful())
 		{
-			PFollowComp->RequestMove(PFindComp->GetPath(), Goal);
+			Result.Path->SetGoalActorObservation(*Goal, 100.0f);
+
+			PFollowComp->RequestMove(Result.Path, Goal);
 		}
 	}
 }
@@ -914,12 +935,12 @@ void UNavigationSystem::SimpleMoveToLocation(AController* Controller, const FVec
 		return;
 	}
 
-	UNavigationComponent* PFindComp = NULL;
 	UPathFollowingComponent* PFollowComp = NULL;
+	UNavigationSystem* NavSys = UNavigationSystem::GetCurrent(Controller->GetWorld());
 
-	Controller->InitNavigationControl(PFindComp, PFollowComp);
+	Controller->InitNavigationControl(PFollowComp);
 
-	if (PFindComp == nullptr || PFollowComp == nullptr)
+	if (PFollowComp == nullptr)
 	{
 		FMessageLog("PIE").Warning(FText::Format(
 			LOCTEXT("SimpleMoveErrorNoComp", "SimpleMove failed for {0}: missing components"),
@@ -937,12 +958,14 @@ void UNavigationSystem::SimpleMoveToLocation(AController* Controller, const FVec
 		return;
 	}
 
-	if (!PFollowComp->HasReached(GoalLocation))
+	if (NavSys && !PFollowComp->HasReached(GoalLocation))
 	{
-		const bool bPathExists = PFindComp->FindPathToLocation(GoalLocation);
-		if (bPathExists)
+		const ANavigationData* NavData = NavSys->GetNavDataForProps(Controller->GetNavAgentProperties());
+		FPathFindingQuery Query(Controller, NavData, Controller->GetNavAgentLocation(), GoalLocation);
+		FPathFindingResult Result = NavSys->FindPathSync(Query);
+		if (Result.IsSuccessful())
 		{
-			PFollowComp->RequestMove(PFindComp->GetPath(), NULL);
+			PFollowComp->RequestMove(Result.Path, NULL);
 		}
 	}
 }
diff --git a/Engine/Source/Runtime/Engine/Private/Controller.cpp b/Engine/Source/Runtime/Engine/Private/Controller.cpp
--- a/Engine/Source/Runtime/Engine/Private/Controller.cpp
+++ b/Engine/Source/Runtime/Engine/Private/Controller.cpp
@@ -16,7 +16,6 @@
 
 // @todo this is here only due to circular dependency to AIModule. To be removed
 #include "Navigation/PathFollowingComponent.h"
-#include "Navigation/NavigationComponent.h"
 
 #include "GameFramework/PlayerState.h"
 #include "Components/CapsuleComponent.h"
@@ -557,13 +556,6 @@ bool AController::ShouldPostponePathUpdates() const
 
 void AController::UpdateNavigationComponents()
 {
-	UNavigationComponent* PathFindingComp = FindComponentByClass<UNavigationComponent>();
-	if (PathFindingComp != NULL)
-	{
-		PathFindingComp->OnNavAgentChanged();
-		PathFindingComp->UpdateCachedComponents();
-	}
-
 	UPathFollowingComponent* PathFollowingComp = FindComponentByClass<UPathFollowingComponent>();
 	if (PathFollowingComp != NULL)
 	{
@@ -571,15 +563,8 @@ void AController::UpdateNavigationComponents()
 	}
 }
 
-void AController::InitNavigationControl(UNavigationComponent*& PathFindingComp, UPathFollowingComponent*& PathFollowingComp)
+void AController::InitNavigationControl(UPathFollowingComponent*& PathFollowingComp)
 {
-	PathFindingComp = FindComponentByClass<UNavigationComponent>();
-	if (PathFindingComp == NULL)
-	{
-		PathFindingComp = NewObject<UNavigationComponent>(this);
-		PathFindingComp->RegisterComponentWithWorld(GetWorld());
-	}
-
 	PathFollowingComp = FindComponentByClass<UPathFollowingComponent>();
 	if (PathFollowingComp == NULL)
 	{
