another compilaton fix

diff --git a/Engine/Source/Runtime/AIModule/Private/Navigation/PathFollowingComponent.cpp b/Engine/Source/Runtime/AIModule/Private/Navigation/PathFollowingComponent.cpp
--- a/Engine/Source/Runtime/AIModule/Private/Navigation/PathFollowingComponent.cpp
+++ b/Engine/Source/Runtime/AIModule/Private/Navigation/PathFollowingComponent.cpp
@@ -4,6 +4,7 @@
 #if WITH_RECAST
 #	include "Detour/DetourNavMeshQuery.h"
 #endif
+#include "AI/Navigation/AbstractNavData.h"
 #include "AI/Navigation/RecastNavMesh.h"
 #include "AI/Navigation/NavLinkCustomInterface.h"
 #include "GameFramework/NavMovementComponent.h"
