#UE4 - Added RefreshStimulus() function to AIPerceptionComponent, so it can be overridden by game-specific perception components
#orion - Added OrionAIPerceptionComponent class
  - Fixed several vector function warnings

diff --git a/Engine/Source/Runtime/AIModule/Classes/Perception/AIPerceptionComponent.h b/Engine/Source/Runtime/AIModule/Classes/Perception/AIPerceptionComponent.h
--- a/Engine/Source/Runtime/AIModule/Classes/Perception/AIPerceptionComponent.h
+++ b/Engine/Source/Runtime/AIModule/Classes/Perception/AIPerceptionComponent.h
@@ -251,6 +251,9 @@ protected:
 	virtual void CleanUp();
 
 	void RemoveDeadData();
+
+	/** Updates the stimulus entry in StimulusStore, if NewStimulus is more recent or stronger */
+	virtual void RefreshStimulus(FAIStimulus& StimulusStore, const FAIStimulus& NewStimulus);
 	
 private:
 	FPerceptionListenerID PerceptionListenerId;
diff --git a/Engine/Source/Runtime/AIModule/Private/Perception/AIPerceptionComponent.cpp b/Engine/Source/Runtime/AIModule/Private/Perception/AIPerceptionComponent.cpp
--- a/Engine/Source/Runtime/AIModule/Private/Perception/AIPerceptionComponent.cpp
+++ b/Engine/Source/Runtime/AIModule/Private/Perception/AIPerceptionComponent.cpp
@@ -401,13 +401,7 @@ void UAIPerceptionComponent::ProcessStimuli()
 
 		if (SourcedStimulus->Stimulus.WasSuccessfullySensed())
 		{
-			// if new stimulus is younger or stronger
-			// note that stimulus Age depends on PerceptionSystem::PerceptionAgingRate. It's possible that 
-			// both already stored and the new stimulus have Age of 0, but stored stimulus' acctual age is in [0, PerceptionSystem::PerceptionAgingRate)
-			if (SourcedStimulus->Stimulus.GetAge() <= StimulusStore.GetAge() || StimulusStore.Strength < SourcedStimulus->Stimulus.Strength)
-			{
-				StimulusStore = SourcedStimulus->Stimulus;
-			}
+			RefreshStimulus(StimulusStore, SourcedStimulus->Stimulus);
 		}
 		else
 		{
@@ -430,6 +424,17 @@ void UAIPerceptionComponent::ProcessStimuli()
 	}
 }
 
+void UAIPerceptionComponent::RefreshStimulus(FAIStimulus& StimulusStore, const FAIStimulus& NewStimulus)
+{
+	// if new stimulus is younger or stronger
+	// note that stimulus Age depends on PerceptionSystem::PerceptionAgingRate. It's possible that 
+	// both already stored and the new stimulus have Age of 0, but stored stimulus' acctual age is in [0, PerceptionSystem::PerceptionAgingRate)
+	if (NewStimulus.GetAge() <= StimulusStore.GetAge() || StimulusStore.Strength < NewStimulus.Strength)
+	{
+		StimulusStore = NewStimulus;
+	}
+}
+
 void UAIPerceptionComponent::AgeStimuli(const float ConstPerceptionAgingRate)
 {
 	for (TActorPerceptionContainer::TIterator It(PerceptualData); It; ++It)
