Refactored a bunch of BT functions taking BehaviorTreeComponent pointers to require references instead #UE4

diff --git a/Engine/Source/Developer/AITestSuite/Classes/BehaviorTree/TestBTDecorator_CantExecute.h b/Engine/Source/Developer/AITestSuite/Classes/BehaviorTree/TestBTDecorator_CantExecute.h
--- a/Engine/Source/Developer/AITestSuite/Classes/BehaviorTree/TestBTDecorator_CantExecute.h
+++ b/Engine/Source/Developer/AITestSuite/Classes/BehaviorTree/TestBTDecorator_CantExecute.h
@@ -9,5 +9,5 @@ class UTestBTDecorator_CantExecute : public UBTDecorator
 {
 	GENERATED_UCLASS_BODY()
 
-	virtual bool CalculateRawConditionValue(class UBehaviorTreeComponent* OwnerComp, uint8* NodeMemory) const override;
+	virtual bool CalculateRawConditionValue(UBehaviorTreeComponent& OwnerComp, uint8* NodeMemory) const override;
 };
diff --git a/Engine/Source/Developer/AITestSuite/Classes/BehaviorTree/TestBTDecorator_DelayedAbort.h b/Engine/Source/Developer/AITestSuite/Classes/BehaviorTree/TestBTDecorator_DelayedAbort.h
--- a/Engine/Source/Developer/AITestSuite/Classes/BehaviorTree/TestBTDecorator_DelayedAbort.h
+++ b/Engine/Source/Developer/AITestSuite/Classes/BehaviorTree/TestBTDecorator_DelayedAbort.h
@@ -20,7 +20,7 @@ class UTestBTDecorator_DelayedAbort : public UBTDecorator
 	UPROPERTY()
 	bool bOnlyOnce;
 
-	virtual void OnBecomeRelevant(UBehaviorTreeComponent* OwnerComp, uint8* NodeMemory) override;
-	virtual void TickNode(UBehaviorTreeComponent* OwnerComp, uint8* NodeMemory, float DeltaSeconds) override;
+	virtual void OnBecomeRelevant(UBehaviorTreeComponent& OwnerComp, uint8* NodeMemory) override;
+	virtual void TickNode(UBehaviorTreeComponent& OwnerComp, uint8* NodeMemory, float DeltaSeconds) override;
 	virtual uint16 GetInstanceMemorySize() const override;
 };
diff --git a/Engine/Source/Developer/AITestSuite/Classes/BehaviorTree/TestBTTask_Log.h b/Engine/Source/Developer/AITestSuite/Classes/BehaviorTree/TestBTTask_Log.h
--- a/Engine/Source/Developer/AITestSuite/Classes/BehaviorTree/TestBTTask_Log.h
+++ b/Engine/Source/Developer/AITestSuite/Classes/BehaviorTree/TestBTTask_Log.h
@@ -26,11 +26,11 @@ class UTestBTTask_Log : public UBTTaskNode
 	UPROPERTY()
 	TEnumAsByte<EBTNodeResult::Type> LogResult;
 
-	virtual EBTNodeResult::Type ExecuteTask(class UBehaviorTreeComponent* OwnerComp, uint8* NodeMemory) override;
+	virtual EBTNodeResult::Type ExecuteTask(UBehaviorTreeComponent& OwnerComp, uint8* NodeMemory) override;
 	virtual uint16 GetInstanceMemorySize() const override;
 
-	void LogExecution(class UBehaviorTreeComponent* OwnerComp, int32 LogNumber);
+	void LogExecution(UBehaviorTreeComponent& OwnerComp, int32 LogNumber);
 
 protected:
-	virtual void TickTask(class UBehaviorTreeComponent* OwnerComp, uint8* NodeMemory, float DeltaSeconds) override;
+	virtual void TickTask(UBehaviorTreeComponent& OwnerComp, uint8* NodeMemory, float DeltaSeconds) override;
 };
diff --git a/Engine/Source/Developer/AITestSuite/Classes/BehaviorTree/TestBTTask_SetFlag.h b/Engine/Source/Developer/AITestSuite/Classes/BehaviorTree/TestBTTask_SetFlag.h
--- a/Engine/Source/Developer/AITestSuite/Classes/BehaviorTree/TestBTTask_SetFlag.h
+++ b/Engine/Source/Developer/AITestSuite/Classes/BehaviorTree/TestBTTask_SetFlag.h
@@ -18,5 +18,5 @@ class UTestBTTask_SetFlag : public UBTTaskNode
 	UPROPERTY()
 	TEnumAsByte<EBTNodeResult::Type> TaskResult;
 
-	virtual EBTNodeResult::Type ExecuteTask(class UBehaviorTreeComponent* OwnerComp, uint8* NodeMemory) override;
+	virtual EBTNodeResult::Type ExecuteTask(UBehaviorTreeComponent& OwnerComp, uint8* NodeMemory) override;
 };
diff --git a/Engine/Source/Developer/AITestSuite/Private/BehaviorTree/TestBTDecorator_CantExecute.cpp b/Engine/Source/Developer/AITestSuite/Private/BehaviorTree/TestBTDecorator_CantExecute.cpp
--- a/Engine/Source/Developer/AITestSuite/Private/BehaviorTree/TestBTDecorator_CantExecute.cpp
+++ b/Engine/Source/Developer/AITestSuite/Private/BehaviorTree/TestBTDecorator_CantExecute.cpp
@@ -12,7 +12,7 @@ UTestBTDecorator_CantExecute::UTestBTDecorator_CantExecute(const FObjectInitiali
 	bAllowAbortChildNodes = false;
 }
 
-bool UTestBTDecorator_CantExecute::CalculateRawConditionValue(UBehaviorTreeComponent* OwnerComp, uint8* NodeMemory) const
+bool UTestBTDecorator_CantExecute::CalculateRawConditionValue(UBehaviorTreeComponent& OwnerComp, uint8* NodeMemory) const
 {
 	return false;
 }
diff --git a/Engine/Source/Developer/AITestSuite/Private/BehaviorTree/TestBTDecorator_DelayedAbort.cpp b/Engine/Source/Developer/AITestSuite/Private/BehaviorTree/TestBTDecorator_DelayedAbort.cpp
--- a/Engine/Source/Developer/AITestSuite/Private/BehaviorTree/TestBTDecorator_DelayedAbort.cpp
+++ b/Engine/Source/Developer/AITestSuite/Private/BehaviorTree/TestBTDecorator_DelayedAbort.cpp
@@ -18,19 +18,19 @@ UTestBTDecorator_DelayedAbort::UTestBTDecorator_DelayedAbort(const FObjectInitia
 	FlowAbortMode = EBTFlowAbortMode::Self;
 }
 
-void UTestBTDecorator_DelayedAbort::OnBecomeRelevant(UBehaviorTreeComponent* OwnerComp, uint8* NodeMemory)
+void UTestBTDecorator_DelayedAbort::OnBecomeRelevant(UBehaviorTreeComponent& OwnerComp, uint8* NodeMemory)
 {
 	FBTDelayedAbortMemory* MyMemory = (FBTDelayedAbortMemory*)NodeMemory;
 	MyMemory->EndFrameIdx = FAITestHelpers::FramesCounter() + DelayTicks;
 }
 
-void UTestBTDecorator_DelayedAbort::TickNode(UBehaviorTreeComponent* OwnerComp, uint8* NodeMemory, float DeltaSeconds)
+void UTestBTDecorator_DelayedAbort::TickNode(UBehaviorTreeComponent& OwnerComp, uint8* NodeMemory, float DeltaSeconds)
 {
 	FBTDelayedAbortMemory* MyMemory = (FBTDelayedAbortMemory*)NodeMemory;
 
 	if (FAITestHelpers::FramesCounter() >= MyMemory->EndFrameIdx)
 	{
-		OwnerComp->RequestExecution(this);
+		OwnerComp.RequestExecution(this);
 		MyMemory->EndFrameIdx = bOnlyOnce ? MAX_uint64 : 0;
 	}
 }
diff --git a/Engine/Source/Developer/AITestSuite/Private/BehaviorTree/TestBTTask_Log.cpp b/Engine/Source/Developer/AITestSuite/Private/BehaviorTree/TestBTTask_Log.cpp
--- a/Engine/Source/Developer/AITestSuite/Private/BehaviorTree/TestBTTask_Log.cpp
+++ b/Engine/Source/Developer/AITestSuite/Private/BehaviorTree/TestBTTask_Log.cpp
@@ -15,7 +15,7 @@ UTestBTTask_Log::UTestBTTask_Log(const FObjectInitializer& ObjectInitializer) :
 	bNotifyTick = true;
 }
 
-EBTNodeResult::Type UTestBTTask_Log::ExecuteTask(UBehaviorTreeComponent* OwnerComp, uint8* NodeMemory)
+EBTNodeResult::Type UTestBTTask_Log::ExecuteTask(UBehaviorTreeComponent& OwnerComp, uint8* NodeMemory)
 {
 	FBTLogTaskMemory* MyMemory = (FBTLogTaskMemory*)NodeMemory;
 	MyMemory->EndFrameIdx = ExecutionTicks + FAITestHelpers::FramesCounter();
@@ -29,7 +29,7 @@ EBTNodeResult::Type UTestBTTask_Log::ExecuteTask(UBehaviorTreeComponent* OwnerCo
 	return EBTNodeResult::InProgress;
 }
 
-void UTestBTTask_Log::TickTask(UBehaviorTreeComponent* OwnerComp, uint8* NodeMemory, float DeltaSeconds)
+void UTestBTTask_Log::TickTask(UBehaviorTreeComponent& OwnerComp, uint8* NodeMemory, float DeltaSeconds)
 {
 	FBTLogTaskMemory* MyMemory = (FBTLogTaskMemory*)NodeMemory;
 
@@ -45,7 +45,7 @@ uint16 UTestBTTask_Log::GetInstanceMemorySize() const
 	return sizeof(FBTLogTaskMemory);
 }
 
-void UTestBTTask_Log::LogExecution(UBehaviorTreeComponent* OwnerComp, int32 LogNumber)
+void UTestBTTask_Log::LogExecution(UBehaviorTreeComponent& OwnerComp, int32 LogNumber)
 {
 	if (LogNumber >= 0)
 	{
diff --git a/Engine/Source/Developer/AITestSuite/Private/BehaviorTree/TestBTTask_SetFlag.cpp b/Engine/Source/Developer/AITestSuite/Private/BehaviorTree/TestBTTask_SetFlag.cpp
--- a/Engine/Source/Developer/AITestSuite/Private/BehaviorTree/TestBTTask_SetFlag.cpp
+++ b/Engine/Source/Developer/AITestSuite/Private/BehaviorTree/TestBTTask_SetFlag.cpp
@@ -12,8 +12,8 @@ UTestBTTask_SetFlag::UTestBTTask_SetFlag(const FObjectInitializer& ObjectInitial
 	bValue = true;
 }
 
-EBTNodeResult::Type UTestBTTask_SetFlag::ExecuteTask(UBehaviorTreeComponent* OwnerComp, uint8* NodeMemory)
+EBTNodeResult::Type UTestBTTask_SetFlag::ExecuteTask(UBehaviorTreeComponent& OwnerComp, uint8* NodeMemory)
 {
-	OwnerComp->GetBlackboardComponent()->SetValueAsBool(KeyName, bValue);
+	OwnerComp.GetBlackboardComponent()->SetValueAsBool(KeyName, bValue);
 	return TaskResult;
 }
diff --git a/Engine/Source/Editor/BehaviorTreeEditor/Private/BehaviorTreeDebugger.cpp b/Engine/Source/Editor/BehaviorTreeEditor/Private/BehaviorTreeDebugger.cpp
--- a/Engine/Source/Editor/BehaviorTreeEditor/Private/BehaviorTreeDebugger.cpp
+++ b/Engine/Source/Editor/BehaviorTreeEditor/Private/BehaviorTreeDebugger.cpp
@@ -53,7 +53,7 @@ void FBehaviorTreeDebugger::CacheRootNode()
 	}
 }
 
-void FBehaviorTreeDebugger::Setup(class UBehaviorTree* InTreeAsset, TSharedRef<FBehaviorTreeEditor> InEditorOwner)
+void FBehaviorTreeDebugger::Setup(UBehaviorTree* InTreeAsset, TSharedRef<FBehaviorTreeEditor> InEditorOwner)
 {
 	EditorOwner = InEditorOwner;
 	TreeAsset = InTreeAsset;
@@ -251,19 +251,19 @@ void FBehaviorTreeDebugger::OnAIDebugSelected(const APawn* Pawn)
 	}
 }
 
-void FBehaviorTreeDebugger::OnTreeStarted(const class UBehaviorTreeComponent* OwnerComp, const class UBehaviorTree* InTreeAsset)
+void FBehaviorTreeDebugger::OnTreeStarted(const UBehaviorTreeComponent& OwnerComp, const UBehaviorTree& InTreeAsset)
 {
 	// start debugging if tree asset matches, and no other actor was selected
-	if (!TreeInstance.IsValid() && TreeAsset && TreeAsset == InTreeAsset)
+	if (!TreeInstance.IsValid() && TreeAsset && TreeAsset == &InTreeAsset)
 	{
 		ClearDebuggerState();
-		TreeInstance = OwnerComp;
+		TreeInstance = &OwnerComp;
 
 		UpdateDebuggerViewOnInstanceChange();
 	}
 
 	// update known instances
-	TWeakObjectPtr<UBehaviorTreeComponent> KnownComp = (UBehaviorTreeComponent*)OwnerComp;
+	TWeakObjectPtr<UBehaviorTreeComponent> KnownComp = const_cast<UBehaviorTreeComponent*>(&OwnerComp);
 	KnownInstances.AddUnique(KnownComp);
 }
 
@@ -611,12 +611,12 @@ void FBehaviorTreeDebugger::CollectBreakpointsFromAsset(class UBehaviorTreeGraph
 	}
 }
 
-int32 FBehaviorTreeDebugger::FindMatchingDebuggerStack(class UBehaviorTreeComponent* TestInstance) const
+int32 FBehaviorTreeDebugger::FindMatchingDebuggerStack(UBehaviorTreeComponent& TestInstance) const
 {
 #if USE_BEHAVIORTREE_DEBUGGER
-	if (TestInstance && TestInstance->DebuggerSteps.Num())
+	if (TestInstance.DebuggerSteps.Num())
 	{
-		const FBehaviorTreeExecutionStep& StepInfo = TestInstance->DebuggerSteps.Last();
+		const FBehaviorTreeExecutionStep& StepInfo = TestInstance.DebuggerSteps.Last();
 		for (int32 i = 0; i < StepInfo.InstanceStack.Num(); i++)
 		{
 			if (StepInfo.InstanceStack[i].TreeAsset == TreeAsset)
@@ -630,9 +630,9 @@ int32 FBehaviorTreeDebugger::FindMatchingDebuggerStack(class UBehaviorTreeCompon
 	return INDEX_NONE;
 }
 
-class UBehaviorTreeComponent* FBehaviorTreeDebugger::FindInstanceInActor(AActor* TestActor)
+UBehaviorTreeComponent* FBehaviorTreeDebugger::FindInstanceInActor(AActor* TestActor)
 {
-	class UBehaviorTreeComponent* FoundInstance = NULL;
+	UBehaviorTreeComponent* FoundInstance = NULL;
 	if (TestActor)
 	{
 		APawn* TestPawn = Cast<APawn>(TestActor);
@@ -691,19 +691,22 @@ void FBehaviorTreeDebugger::FindMatchingTreeInstance()
 	for (FActorIterator It(GEditor->PlayWorld); It; ++It)
 	{
 		AActor* TestActor = *It;
-		UBehaviorTreeComponent* TestComp = TestActor ? TestActor->FindComponentByClass<UBehaviorTreeComponent>() : NULL;
+		UBehaviorTreeComponent* TestComp = TestActor ? TestActor->FindComponentByClass<UBehaviorTreeComponent>() : nullptr;
 
-		KnownInstances.Add(TestComp);
-
-		const int32 MatchingIdx = FindMatchingDebuggerStack(TestComp);
-		if (MatchingIdx != INDEX_NONE)
+		if (TestComp)
 		{
-			MatchingComp = TestComp;
+			KnownInstances.Add(TestComp);
 
-			if (TestActor->IsSelected())
+			const int32 MatchingIdx = FindMatchingDebuggerStack(*TestComp);
+			if (MatchingIdx != INDEX_NONE)
 			{
-				TreeInstance = TestComp;
-				return;
+				MatchingComp = TestComp;
+
+				if (TestActor->IsSelected())
+				{
+					TreeInstance = TestComp;
+					return;
+				}
 			}
 		}
 	}
@@ -1075,26 +1078,25 @@ float FBehaviorTreeDebugger::GetTimeStamp(bool bUseCurrentState) const
 
 FString FBehaviorTreeDebugger::GetDebuggedInstanceDesc() const
 {
-	return TreeInstance.IsValid() ? 
-		DescribeInstance(TreeInstance.Get()) :
-		NSLOCTEXT("BlueprintEditor", "DebugActorNothingSelected", "No debug object selected").ToString();
+	UBehaviorTreeComponent* BTComponent = TreeInstance.Get();
+	return BTComponent ? DescribeInstance(*BTComponent) : NSLOCTEXT("BlueprintEditor", "DebugActorNothingSelected", "No debug object selected").ToString();
 }
 
-FString FBehaviorTreeDebugger::DescribeInstance(class UBehaviorTreeComponent* InstanceToDescribe) const
+FString FBehaviorTreeDebugger::DescribeInstance(UBehaviorTreeComponent& InstanceToDescribe) const
 {
 	FString ActorDesc;
-	if (InstanceToDescribe && InstanceToDescribe->GetOwner())
+	if (InstanceToDescribe.GetOwner())
 	{
-		AController* TestController = Cast<AController>(InstanceToDescribe->GetOwner());
+		AController* TestController = Cast<AController>(InstanceToDescribe.GetOwner());
 		ActorDesc = TestController ?
 			TestController->GetName() :
-			InstanceToDescribe->GetOwner()->GetActorLabel();
+			InstanceToDescribe.GetOwner()->GetActorLabel();
 	}
 
 	return ActorDesc;
 }
 
-void FBehaviorTreeDebugger::OnInstanceSelectedInDropdown(class UBehaviorTreeComponent* SelectedInstance)
+void FBehaviorTreeDebugger::OnInstanceSelectedInDropdown(UBehaviorTreeComponent* SelectedInstance)
 {
 	if (SelectedInstance)
 	{
@@ -1126,7 +1128,7 @@ void FBehaviorTreeDebugger::OnInstanceSelectedInDropdown(class UBehaviorTreeComp
 	}
 }
 
-void FBehaviorTreeDebugger::GetMatchingInstances(TArray<class UBehaviorTreeComponent*>& MatchingInstances)
+void FBehaviorTreeDebugger::GetMatchingInstances(TArray<UBehaviorTreeComponent*>& MatchingInstances)
 {
 	for (int32 i = KnownInstances.Num() - 1; i >= 0; i--)
 	{
@@ -1137,7 +1139,7 @@ void FBehaviorTreeDebugger::GetMatchingInstances(TArray<class UBehaviorTreeCompo
 			continue;
 		}
 
-		const int32 StackIdx = FindMatchingDebuggerStack(TestInstance);
+		const int32 StackIdx = FindMatchingDebuggerStack(*TestInstance);
 		if (StackIdx != INDEX_NONE)
 		{
 			MatchingInstances.Add(TestInstance);
@@ -1237,14 +1239,14 @@ void FBehaviorTreeDebugger::UpdateAvailableActions()
 			StepForwardIntoIdx = ActiveStepIndex + 1;
 		}
 
-		class UBehaviorTree* CurTree = CurStepInfo.InstanceStack.IsValidIndex(DebuggerInstanceIndex) ?
+		UBehaviorTree* CurTree = CurStepInfo.InstanceStack.IsValidIndex(DebuggerInstanceIndex) ?
 			CurStepInfo.InstanceStack[DebuggerInstanceIndex].TreeAsset : NULL;
 		const int32 CurStepInstances = DebuggerInstanceIndex + 1;
 
 		for (int32 TestStepIndex = ActiveStepIndex - 1; TestStepIndex >= 0; TestStepIndex--)
 		{
 			const FBehaviorTreeExecutionStep& TestStepInfo = TreeInstance->DebuggerSteps[TestStepIndex];
-			class UBehaviorTree* TestTree = NULL;
+			UBehaviorTree* TestTree = NULL;
 			const int32 TestStepInstances = GetNumActiveInstances(TestStepInfo, TestTree);
 
 			if (TestStepInstances < CurStepInstances ||
@@ -1267,7 +1269,7 @@ void FBehaviorTreeDebugger::UpdateAvailableActions()
 		for (int32 TestStepIndex = ActiveStepIndex + 1; TestStepIndex < TreeInstance->DebuggerSteps.Num(); TestStepIndex++)
 		{
 			const FBehaviorTreeExecutionStep& TestStepInfo = TreeInstance->DebuggerSteps[TestStepIndex];
-			class UBehaviorTree* TestTree = NULL;
+			UBehaviorTree* TestTree = NULL;
 			int32 TestStepInstances = GetNumActiveInstances(TestStepInfo, TestTree);
 
 			if (TestStepInstances < CurStepInstances ||
@@ -1292,7 +1294,7 @@ void FBehaviorTreeDebugger::UpdateAvailableActions()
 			for (int32 TestStepIndex = ActiveStepIndex + 1; TestStepIndex < TreeInstance->DebuggerSteps.Num(); TestStepIndex++)
 			{
 				const FBehaviorTreeExecutionStep& TestStepInfo = TreeInstance->DebuggerSteps[TestStepIndex];
-				class UBehaviorTree* TestTree = NULL;
+				UBehaviorTree* TestTree = NULL;
 				int32 TestStepInstances = GetNumActiveInstances(TestStepInfo, TestTree);
 
 				if (TestStepInstances < CurStepInstances ||
diff --git a/Engine/Source/Editor/BehaviorTreeEditor/Private/BehaviorTreeDebugger.h b/Engine/Source/Editor/BehaviorTreeEditor/Private/BehaviorTreeDebugger.h
--- a/Engine/Source/Editor/BehaviorTreeEditor/Private/BehaviorTreeDebugger.h
+++ b/Engine/Source/Editor/BehaviorTreeEditor/Private/BehaviorTreeDebugger.h
@@ -2,6 +2,9 @@
 
 #pragma once
 
+class UBehaviorTree;
+class UBehaviorTreeComponent;
+
 class FBehaviorTreeDebugger : public FTickableGameObject
 {
 public:
@@ -17,7 +20,7 @@ public:
 	/** Refresh the debugging information we are displaying (only when paused, as Tick() updates when running) */
 	void Refresh();
 
-	void Setup(class UBehaviorTree* InTreeAsset, TSharedRef<class FBehaviorTreeEditor> InEditorOwner);
+	void Setup(UBehaviorTree* InTreeAsset, TSharedRef<class FBehaviorTreeEditor> InEditorOwner);
 
 	/** Store the root node for easy access if we have not already */
 	void CacheRootNode();
@@ -29,12 +32,12 @@ public:
 
 	void OnObjectSelected(UObject* Object);
 	void OnAIDebugSelected(const APawn* Pawn);
-	void OnTreeStarted(const class UBehaviorTreeComponent* OwnerComp, const class UBehaviorTree* InTreeAsset);
+	void OnTreeStarted(const UBehaviorTreeComponent& OwnerComp, const UBehaviorTree& InTreeAsset);
 	void OnBeginPIE(const bool bIsSimulating);
 	void OnEndPIE(const bool bIsSimulating);
 
-	void OnBreakpointAdded(class UBehaviorTreeGraphNode* Node);
-	void OnBreakpointRemoved(class UBehaviorTreeGraphNode* Node);
+	void OnBreakpointAdded(UBehaviorTreeGraphNode* Node);
+	void OnBreakpointRemoved(UBehaviorTreeGraphNode* Node);
 
 	void StepBackInto();
 	void StepBackOver();
@@ -56,9 +59,9 @@ public:
 	static bool IsPIENotSimulating();
 
 	FString GetDebuggedInstanceDesc() const;
-	FString DescribeInstance(class UBehaviorTreeComponent* InstanceToDescribe) const;
-	void OnInstanceSelectedInDropdown(class UBehaviorTreeComponent* SelectedInstance);
-	void GetMatchingInstances(TArray<class UBehaviorTreeComponent*>& MatchingInstances);
+	FString DescribeInstance(UBehaviorTreeComponent& InstanceToDescribe) const;
+	void OnInstanceSelectedInDropdown(UBehaviorTreeComponent* SelectedInstance);
+	void GetMatchingInstances(TArray<UBehaviorTreeComponent*>& MatchingInstances);
 
 	void InitializeFromParent(class FBehaviorTreeDebugger* ParentDebugger);
 	bool HasContinuousNextStep() const;
@@ -85,7 +88,7 @@ private:
 	TWeakPtr<FBehaviorTreeEditor> EditorOwner;
 
 	/** asset for debugging */
-	class UBehaviorTree* TreeAsset;
+	UBehaviorTree* TreeAsset;
 
 	/** root node in asset's graph */
 	TWeakObjectPtr<class UBehaviorTreeGraphNode_Root> RootNode;
@@ -146,10 +149,10 @@ private:
 	void FindMatchingTreeInstance();
 
 	/** find index on execution instance stack of matching tree asset */
-	int32 FindMatchingDebuggerStack(class UBehaviorTreeComponent* TestInstance) const;
+	int32 FindMatchingDebuggerStack(UBehaviorTreeComponent& TestInstance) const;
 
 	/** find BT component in given actor */
-	class UBehaviorTreeComponent* FindInstanceInActor(AActor* TestActor);
+	UBehaviorTreeComponent* FindInstanceInActor(AActor* TestActor);
 
 	/** try to find pawn currently locked by ai debug tool */
 	void FindLockedDebugActor(class UWorld* World);
diff --git a/Engine/Source/Editor/BehaviorTreeEditor/Private/BehaviorTreeEditor.cpp b/Engine/Source/Editor/BehaviorTreeEditor/Private/BehaviorTreeEditor.cpp
--- a/Engine/Source/Editor/BehaviorTreeEditor/Private/BehaviorTreeEditor.cpp
+++ b/Engine/Source/Editor/BehaviorTreeEditor/Private/BehaviorTreeEditor.cpp
@@ -1019,11 +1019,14 @@ TSharedRef<class SWidget> FBehaviorTreeEditor::OnGetDebuggerActorsMenu()
 		// Fill the combo menu with presets of common screen resolutions
 		for (int32 i = 0; i < MatchingInstances.Num(); i++)
 		{
-			const FText ActorDesc = FText::FromString( Debugger->DescribeInstance(MatchingInstances[i]) );
-			TWeakObjectPtr<UBehaviorTreeComponent> InstancePtr = MatchingInstances[i];
+			if (MatchingInstances[i])
+			{
+				const FText ActorDesc = FText::FromString(Debugger->DescribeInstance(*MatchingInstances[i]));
+				TWeakObjectPtr<UBehaviorTreeComponent> InstancePtr = MatchingInstances[i];
 
-			FUIAction ItemAction(FExecuteAction::CreateSP(this, &FBehaviorTreeEditor::OnDebuggerActorSelected, InstancePtr));
-			MenuBuilder.AddMenuEntry(ActorDesc, TAttribute<FText>(), FSlateIcon(), ItemAction);
+				FUIAction ItemAction(FExecuteAction::CreateSP(this, &FBehaviorTreeEditor::OnDebuggerActorSelected, InstancePtr));
+				MenuBuilder.AddMenuEntry(ActorDesc, TAttribute<FText>(), FSlateIcon(), ItemAction);
+			}
 		}
 
 		// Failsafe when no actor match
diff --git a/Engine/Source/Runtime/AIModule/Classes/AIController.h b/Engine/Source/Runtime/AIModule/Classes/AIController.h
--- a/Engine/Source/Runtime/AIModule/Classes/AIController.h
+++ b/Engine/Source/Runtime/AIModule/Classes/AIController.h
@@ -3,6 +3,7 @@
 #pragma once
 
 #include "AITypes.h"
+#include "GameFramework/Pawn.h"
 #include "AI/Navigation/NavigationTypes.h"
 #include "AI/Navigation/NavigationSystem.h"
 #include "Navigation/PathFollowingComponent.h"
@@ -338,4 +339,11 @@ public:
 	UPawnActionsComponent* GetActionsComp() const;
 };
 
-
+namespace FAISystem
+{
+	FORCEINLINE bool IsValidControllerAndHasValidPawn(const AController* Controller)
+	{
+		return Controller != nullptr && Controller->IsPendingKillPending() == false
+			&& Controller->GetPawn() != nullptr && Controller->GetPawn()->IsPendingKillPending() == false;
+	}
+}
diff --git a/Engine/Source/Runtime/AIModule/Classes/BehaviorTree/BTAuxiliaryNode.h b/Engine/Source/Runtime/AIModule/Classes/BehaviorTree/BTAuxiliaryNode.h
--- a/Engine/Source/Runtime/AIModule/Classes/BehaviorTree/BTAuxiliaryNode.h
+++ b/Engine/Source/Runtime/AIModule/Classes/BehaviorTree/BTAuxiliaryNode.h
@@ -33,15 +33,15 @@ class AIMODULE_API UBTAuxiliaryNode : public UBTNode
 	GENERATED_UCLASS_BODY()
 
 	/** wrapper for node instancing: OnBecomeRelevant */
-	void WrappedOnBecomeRelevant(UBehaviorTreeComponent* OwnerComp, uint8* NodeMemory) const;
+	void WrappedOnBecomeRelevant(UBehaviorTreeComponent& OwnerComp, uint8* NodeMemory) const;
 
 	/** wrapper for node instancing: OnCeaseRelevant */
-	void WrappedOnCeaseRelevant(UBehaviorTreeComponent* OwnerComp, uint8* NodeMemory) const;
+	void WrappedOnCeaseRelevant(UBehaviorTreeComponent& OwnerComp, uint8* NodeMemory) const;
 
 	/** wrapper for node instancing: TickNode */
-	void WrappedTickNode(UBehaviorTreeComponent* OwnerComp, uint8* NodeMemory, float DeltaSeconds) const;
+	void WrappedTickNode(UBehaviorTreeComponent& OwnerComp, uint8* NodeMemory, float DeltaSeconds) const;
 
-	virtual void DescribeRuntimeValues(const UBehaviorTreeComponent* OwnerComp, uint8* NodeMemory, EBTDescriptionVerbosity::Type Verbosity, TArray<FString>& Values) const override;
+	virtual void DescribeRuntimeValues(const UBehaviorTreeComponent& OwnerComp, uint8* NodeMemory, EBTDescriptionVerbosity::Type Verbosity, TArray<FString>& Values) const override;
 	virtual uint16 GetSpecialMemorySize() const override;
 
 protected:
@@ -60,15 +60,15 @@ protected:
 
 	/** called when auxiliary node becomes active
 	 * this function should be considered as const (don't modify state of object) if node is not instanced! */
-	virtual void OnBecomeRelevant(UBehaviorTreeComponent* OwnerComp, uint8* NodeMemory);
+	virtual void OnBecomeRelevant(UBehaviorTreeComponent& OwnerComp, uint8* NodeMemory);
 
 	/** called when auxiliary node becomes inactive
 	 * this function should be considered as const (don't modify state of object) if node is not instanced! */
-	virtual void OnCeaseRelevant(UBehaviorTreeComponent* OwnerComp, uint8* NodeMemory);
+	virtual void OnCeaseRelevant(UBehaviorTreeComponent& OwnerComp, uint8* NodeMemory);
 
 	/** tick function
 	 * this function should be considered as const (don't modify state of object) if node is not instanced! */
-	virtual void TickNode(UBehaviorTreeComponent* OwnerComp, uint8* NodeMemory, float DeltaSeconds);
+	virtual void TickNode(UBehaviorTreeComponent& OwnerComp, uint8* NodeMemory, float DeltaSeconds);
 
 	/** sets next tick time */
 	void SetNextTickTime(uint8* NodeMemory, float RemainingTime) const;
diff --git a/Engine/Source/Runtime/AIModule/Classes/BehaviorTree/BTCompositeNode.h b/Engine/Source/Runtime/AIModule/Classes/BehaviorTree/BTCompositeNode.h
--- a/Engine/Source/Runtime/AIModule/Classes/BehaviorTree/BTCompositeNode.h
+++ b/Engine/Source/Runtime/AIModule/Classes/BehaviorTree/BTCompositeNode.h
@@ -105,16 +105,16 @@ class AIMODULE_API UBTCompositeNode : public UBTNode
 	int32 FindChildToExecute(FBehaviorTreeSearchData& SearchData, EBTNodeResult::Type& LastResult) const;
 
 	/** get index of child node (handle subtrees) */
-	int32 GetChildIndex(FBehaviorTreeSearchData& SearchData, const UBTNode* ChildNode) const;
+	int32 GetChildIndex(FBehaviorTreeSearchData& SearchData, const UBTNode& ChildNode) const;
 	/** get index of child node */
-	int32 GetChildIndex(const UBTNode* ChildNode) const;
+	int32 GetChildIndex(const UBTNode& ChildNode) const;
 
 	/** called before passing search to child node */
-	void OnChildActivation(FBehaviorTreeSearchData& SearchData, const UBTNode* ChildNode) const;
+	void OnChildActivation(FBehaviorTreeSearchData& SearchData, const UBTNode& ChildNode) const;
 	void OnChildActivation(FBehaviorTreeSearchData& SearchData, int32 ChildIndex) const;
 
 	/** called after child has finished search */
-	void OnChildDeactivation(FBehaviorTreeSearchData& SearchData, const UBTNode* ChildNode, EBTNodeResult::Type& NodeResult) const;
+	void OnChildDeactivation(FBehaviorTreeSearchData& SearchData, const UBTNode& ChildNode, EBTNodeResult::Type& NodeResult) const;
 	void OnChildDeactivation(FBehaviorTreeSearchData& SearchData, int32 ChildIndex, EBTNodeResult::Type& NodeResult) const;
 
 	/** called when start enters this node */
@@ -127,7 +127,7 @@ class AIMODULE_API UBTCompositeNode : public UBTNode
 	void OnNodeRestart(FBehaviorTreeSearchData& SearchData) const;
 
 	/** notify about task execution start */
-	void ConditionalNotifyChildExecution(UBehaviorTreeComponent* OwnerComp, uint8* NodeMemory, const UBTNode* ChildNode, EBTNodeResult::Type& NodeResult) const;
+	void ConditionalNotifyChildExecution(UBehaviorTreeComponent& OwnerComp, uint8* NodeMemory, const UBTNode& ChildNode, EBTNodeResult::Type& NodeResult) const;
 
 	/** size of instance memory */
 	virtual uint16 GetInstanceMemorySize() const override;
@@ -148,10 +148,10 @@ class AIMODULE_API UBTCompositeNode : public UBTNode
 	void SetChildOverride(FBehaviorTreeSearchData& SearchData, int8 Index) const;
 
 	/** gathers description of all runtime parameters */
-	virtual void DescribeRuntimeValues(const UBehaviorTreeComponent* OwnerComp, uint8* NodeMemory, EBTDescriptionVerbosity::Type Verbosity, TArray<FString>& Values) const override;
+	virtual void DescribeRuntimeValues(const UBehaviorTreeComponent& OwnerComp, uint8* NodeMemory, EBTDescriptionVerbosity::Type Verbosity, TArray<FString>& Values) const override;
 
 	/** check if child node can execute new subtree */
-	virtual bool CanPushSubtree(UBehaviorTreeComponent* OwnerComp, uint8* NodeMemory, int32 ChildIdx) const;
+	virtual bool CanPushSubtree(UBehaviorTreeComponent& OwnerComp, uint8* NodeMemory, int32 ChildIdx) const;
 
 #if WITH_EDITOR
 	/** @return allowed flow abort modes for decorators */
@@ -163,7 +163,7 @@ class AIMODULE_API UBTCompositeNode : public UBTNode
 	int32 GetMatchingChildIndex(int32 ActiveInstanceIdx, FBTNodeIndex& NodeIdx) const;
 
 	/** is child execution allowed by decorators? */
-	bool DoDecoratorsAllowExecution(UBehaviorTreeComponent* OwnerComp, int32 InstanceIdx, int32 ChildIdx) const;
+	bool DoDecoratorsAllowExecution(UBehaviorTreeComponent& OwnerComp, int32 InstanceIdx, int32 ChildIdx) const;
 
 protected:
 
@@ -180,7 +180,7 @@ protected:
 	uint8 bUseNodeDeactivationNotify : 1;
 
 	/** called just after child execution, allows to modify result */
-	virtual void NotifyChildExecution(UBehaviorTreeComponent* OwnerComp, uint8* NodeMemory, int32 ChildIdx, EBTNodeResult::Type& NodeResult) const;
+	virtual void NotifyChildExecution(UBehaviorTreeComponent& OwnerComp, uint8* NodeMemory, int32 ChildIdx, EBTNodeResult::Type& NodeResult) const;
 
 	/** called when start enters this node */
 	virtual void NotifyNodeActivation(FBehaviorTreeSearchData& SearchData) const;
@@ -201,7 +201,7 @@ protected:
 	int32 GetNextChild(FBehaviorTreeSearchData& SearchData, int32 LastChildIdx, EBTNodeResult::Type LastResult) const;
 
 	/** store delayed execution request */
-	void RequestDelayedExecution(UBehaviorTreeComponent* OwnerComp, EBTNodeResult::Type LastResult) const;
+	void RequestDelayedExecution(UBehaviorTreeComponent& OwnerComp, EBTNodeResult::Type LastResult) const;
 };
 
 
@@ -214,7 +214,7 @@ FORCEINLINE UBTNode* UBTCompositeNode::GetChildNode(int32 Index) const
 		(Children[Index].ChildComposite ?
 			(UBTNode*)Children[Index].ChildComposite :
 			(UBTNode*)Children[Index].ChildTask) :
-		NULL;
+		nullptr;
 }
 
 FORCEINLINE int32 UBTCompositeNode::GetChildrenNum() const
diff --git a/Engine/Source/Runtime/AIModule/Classes/BehaviorTree/BTDecorator.h b/Engine/Source/Runtime/AIModule/Classes/BehaviorTree/BTDecorator.h
--- a/Engine/Source/Runtime/AIModule/Classes/BehaviorTree/BTDecorator.h
+++ b/Engine/Source/Runtime/AIModule/Classes/BehaviorTree/BTDecorator.h
@@ -36,7 +36,7 @@ class AIMODULE_API UBTDecorator : public UBTAuxiliaryNode
 	void InitializeDecorator(uint8 InChildIndex);
 
 	/** wrapper for node instancing: CalculateRawConditionValue */
-	bool WrappedCanExecute(UBehaviorTreeComponent* OwnerComp, uint8* NodeMemory) const;
+	bool WrappedCanExecute(UBehaviorTreeComponent& OwnerComp, uint8* NodeMemory) const;
 
 	/** wrapper for node instancing: OnNodeActivation  */
 	void WrappedOnNodeActivation(FBehaviorTreeSearchData& SearchData) const;
@@ -118,7 +118,7 @@ protected:
 	virtual void OnNodeProcessed(FBehaviorTreeSearchData& SearchData, EBTNodeResult::Type& NodeResult);
 
 	/** calculates raw, core value of decorator's condition. Should not include calling IsInversed */
-	virtual bool CalculateRawConditionValue(UBehaviorTreeComponent* OwnerComp, uint8* NodeMemory) const;
+	virtual bool CalculateRawConditionValue(UBehaviorTreeComponent& OwnerComp, uint8* NodeMemory) const;
 
 	friend FBehaviorDecoratorDetails;
 };
diff --git a/Engine/Source/Runtime/AIModule/Classes/BehaviorTree/BTNode.h b/Engine/Source/Runtime/AIModule/Classes/BehaviorTree/BTNode.h
--- a/Engine/Source/Runtime/AIModule/Classes/BehaviorTree/BTNode.h
+++ b/Engine/Source/Runtime/AIModule/Classes/BehaviorTree/BTNode.h
@@ -33,28 +33,28 @@ class AIMODULE_API UBTNode : public UObject
 	virtual void InitializeFromAsset(UBehaviorTree& Asset);
 	
 	/** initialize memory block */
-	virtual void InitializeMemory(UBehaviorTreeComponent* OwnerComp, uint8* NodeMemory, EBTMemoryInit::Type InitType) const;
+	virtual void InitializeMemory(UBehaviorTreeComponent& OwnerComp, uint8* NodeMemory, EBTMemoryInit::Type InitType) const;
 
 	/** cleanup memory block */
-	virtual void CleanupMemory(UBehaviorTreeComponent* OwnerComp, uint8* NodeMemory, EBTMemoryClear::Type CleanupType) const;
+	virtual void CleanupMemory(UBehaviorTreeComponent& OwnerComp, uint8* NodeMemory, EBTMemoryClear::Type CleanupType) const;
 
 	/** gathers description of all runtime parameters */
-	virtual void DescribeRuntimeValues(const UBehaviorTreeComponent* OwnerComp, uint8* NodeMemory, EBTDescriptionVerbosity::Type Verbosity, TArray<FString>& Values) const;
+	virtual void DescribeRuntimeValues(const UBehaviorTreeComponent& OwnerComp, uint8* NodeMemory, EBTDescriptionVerbosity::Type Verbosity, TArray<FString>& Values) const;
 
 	/** size of instance memory */
 	virtual uint16 GetInstanceMemorySize() const;
 
 	/** called when node instance is added to tree */
-	virtual void OnInstanceCreated(UBehaviorTreeComponent* OwnerComp);
+	virtual void OnInstanceCreated(UBehaviorTreeComponent& OwnerComp);
 
 	/** called when node instance is removed from tree */
-	virtual void OnInstanceDestroyed(UBehaviorTreeComponent* OwnerComp);
+	virtual void OnInstanceDestroyed(UBehaviorTreeComponent& OwnerComp);
 
 	/** called on creating subtree to set up memory and instancing */
-	void InitializeInSubtree(UBehaviorTreeComponent* OwnerComp, uint8* NodeMemory, int32& NextInstancedIndex, EBTMemoryInit::Type InitType) const;
+	void InitializeInSubtree(UBehaviorTreeComponent& OwnerComp, uint8* NodeMemory, int32& NextInstancedIndex, EBTMemoryInit::Type InitType) const;
 
 	/** called on removing subtree to cleanup memory */
-	void CleanupInSubtree(UBehaviorTreeComponent* OwnerComp, uint8* NodeMemory, EBTMemoryClear::Type CleanupType) const;
+	void CleanupInSubtree(UBehaviorTreeComponent& OwnerComp, uint8* NodeMemory, EBTMemoryClear::Type CleanupType) const;
 
 	/** size of special, hidden memory block for internal mechanics */
 	virtual uint16 GetSpecialMemorySize() const;
@@ -120,11 +120,11 @@ class AIMODULE_API UBTNode : public UObject
 	UBlackboardData* GetBlackboardAsset() const;
 
 	/** @return node instance if bCreateNodeInstance was set */
-	UBTNode* GetNodeInstance(const UBehaviorTreeComponent* OwnerComp, uint8* NodeMemory) const;
+	UBTNode* GetNodeInstance(const UBehaviorTreeComponent& OwnerComp, uint8* NodeMemory) const;
 	UBTNode* GetNodeInstance(FBehaviorTreeSearchData& SearchData) const;
 
 	/** @return string containing description of this node instance with all relevant runtime values */
-	FString GetRuntimeDescription(const UBehaviorTreeComponent* OwnerComp, uint8* NodeMemory, EBTDescriptionVerbosity::Type Verbosity) const;
+	FString GetRuntimeDescription(const UBehaviorTreeComponent& OwnerComp, uint8* NodeMemory, EBTDescriptionVerbosity::Type Verbosity) const;
 
 	/** @return string containing description of this node with all setup values */
 	virtual FString GetStaticDescription() const;
@@ -253,13 +253,13 @@ FORCEINLINE bool UBTNode::IsInstanced() const
 template<typename T>
 T* UBTNode::GetNodeMemory(FBehaviorTreeSearchData& SearchData) const
 {
-	return GetNodeMemory<T>(SearchData.OwnerComp->InstanceStack[SearchData.OwnerComp->GetActiveInstanceIdx()]);
+	return GetNodeMemory<T>(SearchData.OwnerComp.InstanceStack[SearchData.OwnerComp.GetActiveInstanceIdx()]);
 }
 
 template<typename T>
 const T* UBTNode::GetNodeMemory(const FBehaviorTreeSearchData& SearchData) const
 {
-	return GetNodeMemory<T>(SearchData.OwnerComp->InstanceStack[SearchData.OwnerComp->GetActiveInstanceIdx()]);
+	return GetNodeMemory<T>(SearchData.OwnerComp.InstanceStack[SearchData.OwnerComp.GetActiveInstanceIdx()]);
 }
 
 template<typename T>
diff --git a/Engine/Source/Runtime/AIModule/Classes/BehaviorTree/BTService.h b/Engine/Source/Runtime/AIModule/Classes/BehaviorTree/BTService.h
--- a/Engine/Source/Runtime/AIModule/Classes/BehaviorTree/BTService.h
+++ b/Engine/Source/Runtime/AIModule/Classes/BehaviorTree/BTService.h
@@ -60,7 +60,7 @@ protected:
 
 	/** update next tick interval
 	 * this function should be considered as const (don't modify state of object) if node is not instanced! */
-	virtual void TickNode(UBehaviorTreeComponent* OwnerComp, uint8* NodeMemory, float DeltaSeconds) override;
+	virtual void TickNode(UBehaviorTreeComponent& OwnerComp, uint8* NodeMemory, float DeltaSeconds) override;
 
 	/** called when search enters underlying branch */
 	virtual void OnSearchStart(FBehaviorTreeSearchData& SearchData);
diff --git a/Engine/Source/Runtime/AIModule/Classes/BehaviorTree/BTTaskNode.h b/Engine/Source/Runtime/AIModule/Classes/BehaviorTree/BTTaskNode.h
--- a/Engine/Source/Runtime/AIModule/Classes/BehaviorTree/BTTaskNode.h
+++ b/Engine/Source/Runtime/AIModule/Classes/BehaviorTree/BTTaskNode.h
@@ -31,12 +31,12 @@ class AIMODULE_API UBTTaskNode : public UBTNode
 	/** starts this task, should return Succeeded, Failed or InProgress
 	 *  (use FinishLatentTask() when returning InProgress)
 	 * this function should be considered as const (don't modify state of object) if node is not instanced! */
-	virtual EBTNodeResult::Type ExecuteTask(UBehaviorTreeComponent* OwnerComp, uint8* NodeMemory);
+	virtual EBTNodeResult::Type ExecuteTask(UBehaviorTreeComponent& OwnerComp, uint8* NodeMemory);
 
 	/** aborts this task, should return Aborted or InProgress
 	 *  (use FinishLatentAbort() when returning InProgress)
 	 * this function should be considered as const (don't modify state of object) if node is not instanced! */
-	virtual EBTNodeResult::Type AbortTask(UBehaviorTreeComponent* OwnerComp, uint8* NodeMemory);
+	virtual EBTNodeResult::Type AbortTask(UBehaviorTreeComponent& OwnerComp, uint8* NodeMemory);
 
 #if WITH_EDITOR
 	virtual FName GetNodeIconName() const override;
@@ -46,13 +46,13 @@ class AIMODULE_API UBTTaskNode : public UBTNode
 	void ReceivedMessage(UBrainComponent* BrainComp, const FAIMessage& Message);
 
 	/** wrapper for node instancing: ExecuteTask */
-	EBTNodeResult::Type WrappedExecuteTask(UBehaviorTreeComponent* OwnerComp, uint8* NodeMemory) const;
+	EBTNodeResult::Type WrappedExecuteTask(UBehaviorTreeComponent& OwnerComp, uint8* NodeMemory) const;
 
 	/** wrapper for node instancing: AbortTask */
-	EBTNodeResult::Type WrappedAbortTask(UBehaviorTreeComponent* OwnerComp, uint8* NodeMemory) const;
+	EBTNodeResult::Type WrappedAbortTask(UBehaviorTreeComponent& OwnerComp, uint8* NodeMemory) const;
 
 	/** wrapper for node instancing: TickTask */
-	void WrappedTickTask(UBehaviorTreeComponent* OwnerComp, uint8* NodeMemory, float DeltaSeconds) const;
+	void WrappedTickTask(UBehaviorTreeComponent& OwnerComp, uint8* NodeMemory, float DeltaSeconds) const;
 
 protected:
 
@@ -61,22 +61,22 @@ protected:
 
 	/** ticks this task 
 	 * this function should be considered as const (don't modify state of object) if node is not instanced! */
-	virtual void TickTask(UBehaviorTreeComponent* OwnerComp, uint8* NodeMemory, float DeltaSeconds);
+	virtual void TickTask(UBehaviorTreeComponent& OwnerComp, uint8* NodeMemory, float DeltaSeconds);
 
 	/** message handler, default implementation will finish latent execution/abortion
 	 * this function should be considered as const (don't modify state of object) if node is not instanced! */
-	virtual void OnMessage(UBehaviorTreeComponent* OwnerComp, uint8* NodeMemory, FName Message, int32 RequestID, bool bSuccess);
+	virtual void OnMessage(UBehaviorTreeComponent& OwnerComp, uint8* NodeMemory, FName Message, int32 RequestID, bool bSuccess);
 
 	/** helper function: finish latent executing */
-	void FinishLatentTask(UBehaviorTreeComponent* OwnerComp, EBTNodeResult::Type TaskResult) const;
+	void FinishLatentTask(UBehaviorTreeComponent& OwnerComp, EBTNodeResult::Type TaskResult) const;
 
 	/** helper function: finishes latent aborting */
-	void FinishLatentAbort(UBehaviorTreeComponent* OwnerComp) const;
+	void FinishLatentAbort(UBehaviorTreeComponent& OwnerComp) const;
 
 	/** register message observer */
-	void WaitForMessage(UBehaviorTreeComponent* OwnerComp, FName MessageType) const;
-	void WaitForMessage(UBehaviorTreeComponent* OwnerComp, FName MessageType, int32 RequestID) const;
+	void WaitForMessage(UBehaviorTreeComponent& OwnerComp, FName MessageType) const;
+	void WaitForMessage(UBehaviorTreeComponent& OwnerComp, FName MessageType, int32 RequestID) const;
 	
 	/** unregister message observers */
-	void StopWaitingForMessages(UBehaviorTreeComponent* OwnerComp) const;
+	void StopWaitingForMessages(UBehaviorTreeComponent& OwnerComp) const;
 };
diff --git a/Engine/Source/Runtime/AIModule/Classes/BehaviorTree/BehaviorTreeManager.h b/Engine/Source/Runtime/AIModule/Classes/BehaviorTree/BehaviorTreeManager.h
--- a/Engine/Source/Runtime/AIModule/Classes/BehaviorTree/BehaviorTreeManager.h
+++ b/Engine/Source/Runtime/AIModule/Classes/BehaviorTree/BehaviorTreeManager.h
@@ -49,10 +49,10 @@ class AIMODULE_API UBehaviorTreeManager : public UObject
 	void DumpUsageStats() const;
 
 	/** register new behavior tree component for tracking */
-	void AddActiveComponent(UBehaviorTreeComponent* Component);
+	void AddActiveComponent(UBehaviorTreeComponent& Component);
 
 	/** unregister behavior tree component from tracking */
-	void RemoveActiveComponent(UBehaviorTreeComponent* Component);
+	void RemoveActiveComponent(UBehaviorTreeComponent& Component);
 
 	static UBehaviorTreeManager* GetCurrent(UWorld* World);
 	static UBehaviorTreeManager* GetCurrent(UObject* WorldContextObject);
diff --git a/Engine/Source/Runtime/AIModule/Classes/BehaviorTree/BehaviorTreeTypes.h b/Engine/Source/Runtime/AIModule/Classes/BehaviorTree/BehaviorTreeTypes.h
--- a/Engine/Source/Runtime/AIModule/Classes/BehaviorTree/BehaviorTreeTypes.h
+++ b/Engine/Source/Runtime/AIModule/Classes/BehaviorTree/BehaviorTreeTypes.h
@@ -19,7 +19,7 @@ struct FBehaviorTreeSearchUpdate;
 
 // Visual logging helper
 #define BT_VLOG(Context, Verbosity, Format, ...) UE_VLOG(Context->OwnerComp.IsValid() ? Context->OwnerComp->GetOwner() : NULL, LogBehaviorTree, Verbosity, Format, ##__VA_ARGS__)
-#define BT_SEARCHLOG(SearchData, Verbosity, Format, ...) UE_VLOG(SearchData.OwnerComp->GetOwner(), LogBehaviorTree, Verbosity, Format, ##__VA_ARGS__)
+#define BT_SEARCHLOG(SearchData, Verbosity, Format, ...) UE_VLOG(SearchData.OwnerComp.GetOwner(), LogBehaviorTree, Verbosity, Format, ##__VA_ARGS__)
 
 // Behavior tree debugger in editor
 #define USE_BEHAVIORTREE_DEBUGGER	(1 && WITH_EDITORONLY_DATA)
@@ -297,18 +297,18 @@ struct FBehaviorTreeInstance
 #endif // STATS
 
 	/** initialize memory and create node instances */
-	void Initialize(UBehaviorTreeComponent* OwnerComp, UBTCompositeNode* Node, int32& InstancedIndex, EBTMemoryInit::Type InitType);
+	void Initialize(UBehaviorTreeComponent& OwnerComp, UBTCompositeNode& Node, int32& InstancedIndex, EBTMemoryInit::Type InitType);
 
 	/** update injected nodes */
-	void InjectNodes(UBehaviorTreeComponent* OwnerComp, UBTCompositeNode* Node, int32& InstancedIndex);
+	void InjectNodes(UBehaviorTreeComponent& OwnerComp, UBTCompositeNode& Node, int32& InstancedIndex);
 
 	/** cleanup node instances */
-	void Cleanup(UBehaviorTreeComponent* OwnerComp, EBTMemoryClear::Type CleanupType);
+	void Cleanup(UBehaviorTreeComponent& OwnerComp, EBTMemoryClear::Type CleanupType);
 
 protected:
 
 	/** worker for updating all nodes */
-	void CleanupNodes(UBehaviorTreeComponent* OwnerComp, UBTCompositeNode* Node, EBTMemoryClear::Type CleanupType);
+	void CleanupNodes(UBehaviorTreeComponent& OwnerComp, UBTCompositeNode& Node, EBTMemoryClear::Type CleanupType);
 };
 
 struct FBTNodeIndex
@@ -357,7 +357,7 @@ struct FBehaviorTreeSearchUpdate
 struct FBehaviorTreeSearchData
 {
 	/** BT component */
-	UBehaviorTreeComponent* OwnerComp;
+	UBehaviorTreeComponent& OwnerComp;
 
 	/** requested updates of additional nodes (preconditions, services, parallels)
 	 *  buffered during search to prevent instant add & remove pairs */
@@ -378,6 +378,9 @@ struct FBehaviorTreeSearchData
 	/** assign unique Id number */
 	void AssignSearchId();
 
+	FBehaviorTreeSearchData(UBehaviorTreeComponent& InOwnerComp) : OwnerComp(InOwnerComp)
+	{}
+
 private:
 
 	static int32 NextSearchId;
diff --git a/Engine/Source/Runtime/AIModule/Classes/BehaviorTree/Composites/BTComposite_SimpleParallel.h b/Engine/Source/Runtime/AIModule/Classes/BehaviorTree/Composites/BTComposite_SimpleParallel.h
--- a/Engine/Source/Runtime/AIModule/Classes/BehaviorTree/Composites/BTComposite_SimpleParallel.h
+++ b/Engine/Source/Runtime/AIModule/Classes/BehaviorTree/Composites/BTComposite_SimpleParallel.h
@@ -56,12 +56,12 @@ class AIMODULE_API UBTComposite_SimpleParallel : public UBTCompositeNode
 	/** handle child updates */
 	int32 GetNextChildHandler(FBehaviorTreeSearchData& SearchData, int32 PrevChild, EBTNodeResult::Type LastResult) const;
 
-	virtual void NotifyChildExecution(UBehaviorTreeComponent* OwnerComp, uint8* NodeMemory, int32 ChildIdx, EBTNodeResult::Type& NodeResult) const override;
+	virtual void NotifyChildExecution(UBehaviorTreeComponent& OwnerComp, uint8* NodeMemory, int32 ChildIdx, EBTNodeResult::Type& NodeResult) const override;
 	virtual void NotifyNodeDeactivation(FBehaviorTreeSearchData& SearchData, EBTNodeResult::Type& NodeResult) const override;
-	virtual bool CanPushSubtree(UBehaviorTreeComponent* OwnerComp, uint8* NodeMemory, int32 ChildIdx) const override;
+	virtual bool CanPushSubtree(UBehaviorTreeComponent& OwnerComp, uint8* NodeMemory, int32 ChildIdx) const override;
 	virtual uint16 GetInstanceMemorySize() const override;
 	virtual FString GetStaticDescription() const override;
-	virtual void DescribeRuntimeValues(const UBehaviorTreeComponent* OwnerComp, uint8* NodeMemory, EBTDescriptionVerbosity::Type Verbosity, TArray<FString>& Values) const override;
+	virtual void DescribeRuntimeValues(const UBehaviorTreeComponent& OwnerComp, uint8* NodeMemory, EBTDescriptionVerbosity::Type Verbosity, TArray<FString>& Values) const override;
 
 	/** helper for showing values of EBTParallelMode enum */
 	static FString DescribeFinishMode(EBTParallelMode::Type Mode);
diff --git a/Engine/Source/Runtime/AIModule/Classes/BehaviorTree/Decorators/BTDecorator_Blackboard.h b/Engine/Source/Runtime/AIModule/Classes/BehaviorTree/Decorators/BTDecorator_Blackboard.h
--- a/Engine/Source/Runtime/AIModule/Classes/BehaviorTree/Decorators/BTDecorator_Blackboard.h
+++ b/Engine/Source/Runtime/AIModule/Classes/BehaviorTree/Decorators/BTDecorator_Blackboard.h
@@ -33,9 +33,9 @@ class AIMODULE_API UBTDecorator_Blackboard : public UBTDecorator_BlackboardBase
 {
 	GENERATED_UCLASS_BODY()
 
-	virtual bool CalculateRawConditionValue(UBehaviorTreeComponent* OwnerComp, uint8* NodeMemory) const override;
+	virtual bool CalculateRawConditionValue(UBehaviorTreeComponent& OwnerComp, uint8* NodeMemory) const override;
 	virtual void OnBlackboardChange(const UBlackboardComponent& Blackboard, FBlackboard::FKey ChangedKeyID) override;
-	virtual void DescribeRuntimeValues(const UBehaviorTreeComponent* OwnerComp, uint8* NodeMemory, EBTDescriptionVerbosity::Type Verbosity, TArray<FString>& Values) const override;
+	virtual void DescribeRuntimeValues(const UBehaviorTreeComponent& OwnerComp, uint8* NodeMemory, EBTDescriptionVerbosity::Type Verbosity, TArray<FString>& Values) const override;
 	virtual FString GetStaticDescription() const override;
 
 protected:
diff --git a/Engine/Source/Runtime/AIModule/Classes/BehaviorTree/Decorators/BTDecorator_BlackboardBase.h b/Engine/Source/Runtime/AIModule/Classes/BehaviorTree/Decorators/BTDecorator_BlackboardBase.h
--- a/Engine/Source/Runtime/AIModule/Classes/BehaviorTree/Decorators/BTDecorator_BlackboardBase.h
+++ b/Engine/Source/Runtime/AIModule/Classes/BehaviorTree/Decorators/BTDecorator_BlackboardBase.h
@@ -34,10 +34,10 @@ protected:
 	virtual void PostInitProperties() override;
 
 	/** called when execution flow controller becomes active */
-	virtual void OnBecomeRelevant(UBehaviorTreeComponent* OwnerComp, uint8* NodeMemory) override;
+	virtual void OnBecomeRelevant(UBehaviorTreeComponent& OwnerComp, uint8* NodeMemory) override;
 
 	/** called when execution flow controller becomes inactive */
-	virtual void OnCeaseRelevant(UBehaviorTreeComponent* OwnerComp, uint8* NodeMemory) override;
+	virtual void OnCeaseRelevant(UBehaviorTreeComponent& OwnerComp, uint8* NodeMemory) override;
 };
 
 //////////////////////////////////////////////////////////////////////////
diff --git a/Engine/Source/Runtime/AIModule/Classes/BehaviorTree/Decorators/BTDecorator_BlueprintBase.h b/Engine/Source/Runtime/AIModule/Classes/BehaviorTree/Decorators/BTDecorator_BlueprintBase.h
--- a/Engine/Source/Runtime/AIModule/Classes/BehaviorTree/Decorators/BTDecorator_BlueprintBase.h
+++ b/Engine/Source/Runtime/AIModule/Classes/BehaviorTree/Decorators/BTDecorator_BlueprintBase.h
@@ -38,9 +38,9 @@ class AIMODULE_API UBTDecorator_BlueprintBase : public UBTDecorator
 	void OnBlackboardChange(const UBlackboardComponent& Blackboard, FBlackboard::FKey ChangedKeyID);
 
 	virtual FString GetStaticDescription() const override;
-	virtual void DescribeRuntimeValues(const UBehaviorTreeComponent* OwnerComp, uint8* NodeMemory, EBTDescriptionVerbosity::Type Verbosity, TArray<FString>& Values) const override;
-	virtual bool CalculateRawConditionValue(UBehaviorTreeComponent* OwnerComp, uint8* NodeMemory) const override;
-	virtual void OnInstanceDestroyed(UBehaviorTreeComponent* OwnerComp) override;
+	virtual void DescribeRuntimeValues(const UBehaviorTreeComponent& OwnerComp, uint8* NodeMemory, EBTDescriptionVerbosity::Type Verbosity, TArray<FString>& Values) const override;
+	virtual bool CalculateRawConditionValue(UBehaviorTreeComponent& OwnerComp, uint8* NodeMemory) const override;
+	virtual void OnInstanceDestroyed(UBehaviorTreeComponent& OwnerComp) override;
 
 	virtual void SetOwner(AActor* ActorOwner) override;
 
@@ -91,11 +91,11 @@ protected:
 
 	FOnBlackboardChange BBKeyObserver;
 	
-	virtual void OnBecomeRelevant(UBehaviorTreeComponent* OwnerComp, uint8* NodeMemory) override;
-	virtual void OnCeaseRelevant(UBehaviorTreeComponent* OwnerComp, uint8* NodeMemory) override;
+	virtual void OnBecomeRelevant(UBehaviorTreeComponent& OwnerComp, uint8* NodeMemory) override;
+	virtual void OnCeaseRelevant(UBehaviorTreeComponent& OwnerComp, uint8* NodeMemory) override;
 	virtual void OnNodeActivation(FBehaviorTreeSearchData& SearchData) override;
 	virtual void OnNodeDeactivation(FBehaviorTreeSearchData& SearchData, EBTNodeResult::Type NodeResult) override;
-	virtual void TickNode(UBehaviorTreeComponent* OwnerComp, uint8* NodeMemory, float DeltaSeconds) override;
+	virtual void TickNode(UBehaviorTreeComponent& OwnerComp, uint8* NodeMemory, float DeltaSeconds) override;
 
 	/** tick function
 	 *	@Note that if both generic and AI event versions are implemented only the more
diff --git a/Engine/Source/Runtime/AIModule/Classes/BehaviorTree/Decorators/BTDecorator_CheckGameplayTagsOnActor.h b/Engine/Source/Runtime/AIModule/Classes/BehaviorTree/Decorators/BTDecorator_CheckGameplayTagsOnActor.h
--- a/Engine/Source/Runtime/AIModule/Classes/BehaviorTree/Decorators/BTDecorator_CheckGameplayTagsOnActor.h
+++ b/Engine/Source/Runtime/AIModule/Classes/BehaviorTree/Decorators/BTDecorator_CheckGameplayTagsOnActor.h
@@ -16,8 +16,8 @@ class AIMODULE_API UBTDecorator_CheckGameplayTagsOnActor : public UBTDecorator
 {
 	GENERATED_UCLASS_BODY()
 
-	virtual bool CalculateRawConditionValue(UBehaviorTreeComponent* OwnerComp, uint8* NodeMemory) const override;
-	virtual void DescribeRuntimeValues(const UBehaviorTreeComponent* OwnerComp, uint8* NodeMemory, EBTDescriptionVerbosity::Type Verbosity, TArray<FString>& Values) const override;
+	virtual bool CalculateRawConditionValue(UBehaviorTreeComponent& OwnerComp, uint8* NodeMemory) const override;
+	virtual void DescribeRuntimeValues(const UBehaviorTreeComponent& OwnerComp, uint8* NodeMemory, EBTDescriptionVerbosity::Type Verbosity, TArray<FString>& Values) const override;
 	virtual FString GetStaticDescription() const override;
 
 protected:
diff --git a/Engine/Source/Runtime/AIModule/Classes/BehaviorTree/Decorators/BTDecorator_CompareBBEntries.h b/Engine/Source/Runtime/AIModule/Classes/BehaviorTree/Decorators/BTDecorator_CompareBBEntries.h
--- a/Engine/Source/Runtime/AIModule/Classes/BehaviorTree/Decorators/BTDecorator_CompareBBEntries.h
+++ b/Engine/Source/Runtime/AIModule/Classes/BehaviorTree/Decorators/BTDecorator_CompareBBEntries.h
@@ -46,14 +46,14 @@ public:
 
 	virtual void PostInitProperties() override;
 
-	virtual bool CalculateRawConditionValue(UBehaviorTreeComponent* OwnerComp, uint8* NodeMemory) const override;
+	virtual bool CalculateRawConditionValue(UBehaviorTreeComponent& OwnerComp, uint8* NodeMemory) const override;
 	virtual FString GetStaticDescription() const override;
 
 	virtual void InitializeFromAsset(UBehaviorTree& Asset) override;
 
 	virtual void OnBlackboardChange(const UBlackboardComponent& Blackboard, FBlackboard::FKey ChangedKeyID);
-	virtual void OnBecomeRelevant(UBehaviorTreeComponent* OwnerComp, uint8* NodeMemory) override;
-	virtual void OnCeaseRelevant(UBehaviorTreeComponent* OwnerComp, uint8* NodeMemory) override;
+	virtual void OnBecomeRelevant(UBehaviorTreeComponent& OwnerComp, uint8* NodeMemory) override;
+	virtual void OnCeaseRelevant(UBehaviorTreeComponent& OwnerComp, uint8* NodeMemory) override;
 
 #if WITH_EDITOR
 	virtual FName GetNodeIconName() const override;
diff --git a/Engine/Source/Runtime/AIModule/Classes/BehaviorTree/Decorators/BTDecorator_ConditionalLoop.h b/Engine/Source/Runtime/AIModule/Classes/BehaviorTree/Decorators/BTDecorator_ConditionalLoop.h
--- a/Engine/Source/Runtime/AIModule/Classes/BehaviorTree/Decorators/BTDecorator_ConditionalLoop.h
+++ b/Engine/Source/Runtime/AIModule/Classes/BehaviorTree/Decorators/BTDecorator_ConditionalLoop.h
@@ -13,7 +13,7 @@ class AIMODULE_API UBTDecorator_ConditionalLoop : public UBTDecorator_Blackboard
 {
 	GENERATED_UCLASS_BODY()
 
-	virtual bool CalculateRawConditionValue(UBehaviorTreeComponent* OwnerComp, uint8* NodeMemory) const override;
+	virtual bool CalculateRawConditionValue(UBehaviorTreeComponent& OwnerComp, uint8* NodeMemory) const override;
 	virtual void OnBlackboardChange(const UBlackboardComponent& Blackboard, FBlackboard::FKey ChangedKeyID) override;
 	virtual void OnNodeDeactivation(FBehaviorTreeSearchData& SearchData, EBTNodeResult::Type NodeResult) override;
 
diff --git a/Engine/Source/Runtime/AIModule/Classes/BehaviorTree/Decorators/BTDecorator_ConeCheck.h b/Engine/Source/Runtime/AIModule/Classes/BehaviorTree/Decorators/BTDecorator_ConeCheck.h
--- a/Engine/Source/Runtime/AIModule/Classes/BehaviorTree/Decorators/BTDecorator_ConeCheck.h
+++ b/Engine/Source/Runtime/AIModule/Classes/BehaviorTree/Decorators/BTDecorator_ConeCheck.h
@@ -43,7 +43,7 @@ class AIMODULE_API UBTDecorator_ConeCheck : public UBTDecorator
 
 	virtual void InitializeFromAsset(UBehaviorTree& Asset) override;
 	virtual uint16 GetInstanceMemorySize() const override;
-	virtual void DescribeRuntimeValues(const UBehaviorTreeComponent* OwnerComp, uint8* NodeMemory, EBTDescriptionVerbosity::Type Verbosity, TArray<FString>& Values) const override;
+	virtual void DescribeRuntimeValues(const UBehaviorTreeComponent& OwnerComp, uint8* NodeMemory, EBTDescriptionVerbosity::Type Verbosity, TArray<FString>& Values) const override;
 	virtual FString GetStaticDescription() const override;
 
 #if WITH_EDITOR
@@ -52,14 +52,14 @@ class AIMODULE_API UBTDecorator_ConeCheck : public UBTDecorator
 
 protected:
 
-	virtual bool CalculateRawConditionValue(UBehaviorTreeComponent* OwnerComp, uint8* NodeMemory) const override;
-	virtual void OnBecomeRelevant(UBehaviorTreeComponent* OwnerComp, uint8* NodeMemory) override;
+	virtual bool CalculateRawConditionValue(UBehaviorTreeComponent& OwnerComp, uint8* NodeMemory) const override;
+	virtual void OnBecomeRelevant(UBehaviorTreeComponent& OwnerComp, uint8* NodeMemory) override;
 	void OnBlackboardChange(const UBlackboardComponent& Blackboard, FBlackboard::FKey ChangedKeyID);
 
-	virtual void TickNode(UBehaviorTreeComponent* OwnerComp, uint8* NodeMemory, float DeltaSeconds) override;
+	virtual void TickNode(UBehaviorTreeComponent& OwnerComp, uint8* NodeMemory, float DeltaSeconds) override;
 	
 	bool CalculateDirection(const UBlackboardComponent* BlackboardComp, const FBlackboardKeySelector& Origin, const FBlackboardKeySelector& End, FVector& Direction) const;
 
 private:
-	bool CalcConditionImpl(UBehaviorTreeComponent* OwnerComp, uint8* NodeMemory) const;
+	bool CalcConditionImpl(UBehaviorTreeComponent& OwnerComp, uint8* NodeMemory) const;
 };
diff --git a/Engine/Source/Runtime/AIModule/Classes/BehaviorTree/Decorators/BTDecorator_Cooldown.h b/Engine/Source/Runtime/AIModule/Classes/BehaviorTree/Decorators/BTDecorator_Cooldown.h
--- a/Engine/Source/Runtime/AIModule/Classes/BehaviorTree/Decorators/BTDecorator_Cooldown.h
+++ b/Engine/Source/Runtime/AIModule/Classes/BehaviorTree/Decorators/BTDecorator_Cooldown.h
@@ -27,10 +27,10 @@ class AIMODULE_API UBTDecorator_Cooldown : public UBTDecorator
 	virtual void PostLoad() override;
 	// End UObject Interface
 
-	virtual bool CalculateRawConditionValue(UBehaviorTreeComponent* OwnerComp, uint8* NodeMemory) const override;
-	virtual void InitializeMemory(UBehaviorTreeComponent* OwnerComp, uint8* NodeMemory, EBTMemoryInit::Type InitType) const override;
+	virtual bool CalculateRawConditionValue(UBehaviorTreeComponent& OwnerComp, uint8* NodeMemory) const override;
+	virtual void InitializeMemory(UBehaviorTreeComponent& OwnerComp, uint8* NodeMemory, EBTMemoryInit::Type InitType) const override;
 	virtual uint16 GetInstanceMemorySize() const override;
-	virtual void DescribeRuntimeValues(const UBehaviorTreeComponent* OwnerComp, uint8* NodeMemory, EBTDescriptionVerbosity::Type Verbosity, TArray<FString>& Values) const override;
+	virtual void DescribeRuntimeValues(const UBehaviorTreeComponent& OwnerComp, uint8* NodeMemory, EBTDescriptionVerbosity::Type Verbosity, TArray<FString>& Values) const override;
 	virtual FString GetStaticDescription() const override;
 
 #if WITH_EDITOR
@@ -40,5 +40,5 @@ class AIMODULE_API UBTDecorator_Cooldown : public UBTDecorator
 protected:
 
 	virtual void OnNodeDeactivation(FBehaviorTreeSearchData& SearchData, EBTNodeResult::Type NodeResult) override;
-	virtual void TickNode(UBehaviorTreeComponent* OwnerComp, uint8* NodeMemory, float DeltaSeconds) override;
+	virtual void TickNode(UBehaviorTreeComponent& OwnerComp, uint8* NodeMemory, float DeltaSeconds) override;
 };
diff --git a/Engine/Source/Runtime/AIModule/Classes/BehaviorTree/Decorators/BTDecorator_DoesPathExist.h b/Engine/Source/Runtime/AIModule/Classes/BehaviorTree/Decorators/BTDecorator_DoesPathExist.h
--- a/Engine/Source/Runtime/AIModule/Classes/BehaviorTree/Decorators/BTDecorator_DoesPathExist.h
+++ b/Engine/Source/Runtime/AIModule/Classes/BehaviorTree/Decorators/BTDecorator_DoesPathExist.h
@@ -49,7 +49,7 @@ public:
 	UPROPERTY(Category=Node, EditAnywhere)
 	TSubclassOf<UNavigationQueryFilter> FilterClass;
 
-	virtual bool CalculateRawConditionValue(UBehaviorTreeComponent* OwnerComp, uint8* NodeMemory) const override;
+	virtual bool CalculateRawConditionValue(UBehaviorTreeComponent& OwnerComp, uint8* NodeMemory) const override;
 	virtual FString GetStaticDescription() const override;
 	virtual void InitializeFromAsset(UBehaviorTree& Asset) override;
 
diff --git a/Engine/Source/Runtime/AIModule/Classes/BehaviorTree/Decorators/BTDecorator_KeepInCone.h b/Engine/Source/Runtime/AIModule/Classes/BehaviorTree/Decorators/BTDecorator_KeepInCone.h
--- a/Engine/Source/Runtime/AIModule/Classes/BehaviorTree/Decorators/BTDecorator_KeepInCone.h
+++ b/Engine/Source/Runtime/AIModule/Classes/BehaviorTree/Decorators/BTDecorator_KeepInCone.h
@@ -44,7 +44,7 @@ class AIMODULE_API UBTDecorator_KeepInCone : public UBTDecorator
 
 	virtual void InitializeFromAsset(UBehaviorTree& Asset) override;
 	virtual uint16 GetInstanceMemorySize() const override;
-	virtual void DescribeRuntimeValues(const UBehaviorTreeComponent* OwnerComp, uint8* NodeMemory, EBTDescriptionVerbosity::Type Verbosity, TArray<FString>& Values) const override;
+	virtual void DescribeRuntimeValues(const UBehaviorTreeComponent& OwnerComp, uint8* NodeMemory, EBTDescriptionVerbosity::Type Verbosity, TArray<FString>& Values) const override;
 	virtual FString GetStaticDescription() const override;
 
 #if WITH_EDITOR
@@ -53,8 +53,8 @@ class AIMODULE_API UBTDecorator_KeepInCone : public UBTDecorator
 
 protected:
 
-	virtual void OnBecomeRelevant(UBehaviorTreeComponent* OwnerComp, uint8* NodeMemory) override;
-	virtual void TickNode(UBehaviorTreeComponent* OwnerComp, uint8* NodeMemory, float DeltaSeconds) override;
+	virtual void OnBecomeRelevant(UBehaviorTreeComponent& OwnerComp, uint8* NodeMemory) override;
+	virtual void TickNode(UBehaviorTreeComponent& OwnerComp, uint8* NodeMemory, float DeltaSeconds) override;
 
-	bool CalculateCurrentDirection(const UBehaviorTreeComponent* OwnerComp, FVector& Direction) const;
+	bool CalculateCurrentDirection(const UBehaviorTreeComponent& OwnerComp, FVector& Direction) const;
 };
diff --git a/Engine/Source/Runtime/AIModule/Classes/BehaviorTree/Decorators/BTDecorator_Loop.h b/Engine/Source/Runtime/AIModule/Classes/BehaviorTree/Decorators/BTDecorator_Loop.h
--- a/Engine/Source/Runtime/AIModule/Classes/BehaviorTree/Decorators/BTDecorator_Loop.h
+++ b/Engine/Source/Runtime/AIModule/Classes/BehaviorTree/Decorators/BTDecorator_Loop.h
@@ -27,7 +27,7 @@ class AIMODULE_API UBTDecorator_Loop : public UBTDecorator
 	bool bInfiniteLoop;
 
 	virtual uint16 GetInstanceMemorySize() const override;
-	virtual void DescribeRuntimeValues(const UBehaviorTreeComponent* OwnerComp, uint8* NodeMemory, EBTDescriptionVerbosity::Type Verbosity, TArray<FString>& Values) const override;
+	virtual void DescribeRuntimeValues(const UBehaviorTreeComponent& OwnerComp, uint8* NodeMemory, EBTDescriptionVerbosity::Type Verbosity, TArray<FString>& Values) const override;
 	virtual FString GetStaticDescription() const override;
 
 #if WITH_EDITOR
diff --git a/Engine/Source/Runtime/AIModule/Classes/BehaviorTree/Decorators/BTDecorator_ReachedMoveGoal.h b/Engine/Source/Runtime/AIModule/Classes/BehaviorTree/Decorators/BTDecorator_ReachedMoveGoal.h
--- a/Engine/Source/Runtime/AIModule/Classes/BehaviorTree/Decorators/BTDecorator_ReachedMoveGoal.h
+++ b/Engine/Source/Runtime/AIModule/Classes/BehaviorTree/Decorators/BTDecorator_ReachedMoveGoal.h
@@ -13,7 +13,7 @@ class AIMODULE_API UBTDecorator_ReachedMoveGoal : public UBTDecorator
 {
 	GENERATED_UCLASS_BODY()
 
-	virtual bool CalculateRawConditionValue(UBehaviorTreeComponent* OwnerComp, uint8* NodeMemory) const override;
+	virtual bool CalculateRawConditionValue(UBehaviorTreeComponent& OwnerComp, uint8* NodeMemory) const override;
 
 #if WITH_EDITOR
 	virtual FName GetNodeIconName() const override;
diff --git a/Engine/Source/Runtime/AIModule/Classes/BehaviorTree/Decorators/BTDecorator_TimeLimit.h b/Engine/Source/Runtime/AIModule/Classes/BehaviorTree/Decorators/BTDecorator_TimeLimit.h
--- a/Engine/Source/Runtime/AIModule/Classes/BehaviorTree/Decorators/BTDecorator_TimeLimit.h
+++ b/Engine/Source/Runtime/AIModule/Classes/BehaviorTree/Decorators/BTDecorator_TimeLimit.h
@@ -16,7 +16,7 @@ class AIMODULE_API UBTDecorator_TimeLimit : public UBTDecorator
 	UPROPERTY(Category=Decorator, EditAnywhere)
 	float TimeLimit;
 
-	virtual void DescribeRuntimeValues(const UBehaviorTreeComponent* OwnerComp, uint8* NodeMemory, EBTDescriptionVerbosity::Type Verbosity, TArray<FString>& Values) const override;
+	virtual void DescribeRuntimeValues(const UBehaviorTreeComponent& OwnerComp, uint8* NodeMemory, EBTDescriptionVerbosity::Type Verbosity, TArray<FString>& Values) const override;
 	virtual FString GetStaticDescription() const override;
 
 #if WITH_EDITOR
@@ -25,6 +25,6 @@ class AIMODULE_API UBTDecorator_TimeLimit : public UBTDecorator
 
 protected:
 
-	virtual void OnBecomeRelevant(UBehaviorTreeComponent* OwnerComp, uint8* NodeMemory) override;
-	virtual void TickNode(UBehaviorTreeComponent* OwnerComp, uint8* NodeMemory, float DeltaSeconds) override;
+	virtual void OnBecomeRelevant(UBehaviorTreeComponent& OwnerComp, uint8* NodeMemory) override;
+	virtual void TickNode(UBehaviorTreeComponent& OwnerComp, uint8* NodeMemory, float DeltaSeconds) override;
 };
diff --git a/Engine/Source/Runtime/AIModule/Classes/BehaviorTree/Services/BTService_BlueprintBase.h b/Engine/Source/Runtime/AIModule/Classes/BehaviorTree/Services/BTService_BlueprintBase.h
--- a/Engine/Source/Runtime/AIModule/Classes/BehaviorTree/Services/BTService_BlueprintBase.h
+++ b/Engine/Source/Runtime/AIModule/Classes/BehaviorTree/Services/BTService_BlueprintBase.h
@@ -21,8 +21,8 @@ class AIMODULE_API UBTService_BlueprintBase : public UBTService
 
 	virtual void PostInitProperties() override;
 
-	virtual void DescribeRuntimeValues(const UBehaviorTreeComponent* OwnerComp, uint8* NodeMemory, EBTDescriptionVerbosity::Type Verbosity, TArray<FString>& Values) const override;
-	virtual void OnInstanceDestroyed(UBehaviorTreeComponent* OwnerComp) override;
+	virtual void DescribeRuntimeValues(const UBehaviorTreeComponent& OwnerComp, uint8* NodeMemory, EBTDescriptionVerbosity::Type Verbosity, TArray<FString>& Values) const override;
+	virtual void OnInstanceDestroyed(UBehaviorTreeComponent& OwnerComp) override;
 
 	virtual void SetOwner(AActor* ActorOwner) override;
 
@@ -61,9 +61,9 @@ protected:
 	/** set if ReceiveSearchStart is implemented by blueprint */
 	uint32 ReceiveSearchStartImplementations : 2;
 
-	virtual void OnBecomeRelevant(UBehaviorTreeComponent* OwnerComp, uint8* NodeMemory) override;
-	virtual void OnCeaseRelevant(UBehaviorTreeComponent* OwnerComp, uint8* NodeMemory) override;
-	virtual void TickNode(UBehaviorTreeComponent* OwnerComp, uint8* NodeMemory, float DeltaSeconds) override;
+	virtual void OnBecomeRelevant(UBehaviorTreeComponent& OwnerComp, uint8* NodeMemory) override;
+	virtual void OnCeaseRelevant(UBehaviorTreeComponent& OwnerComp, uint8* NodeMemory) override;
+	virtual void TickNode(UBehaviorTreeComponent& OwnerComp, uint8* NodeMemory, float DeltaSeconds) override;
 	virtual void OnSearchStart(FBehaviorTreeSearchData& SearchData) override;
 
 	/** tick function
diff --git a/Engine/Source/Runtime/AIModule/Classes/BehaviorTree/Services/BTService_DefaultFocus.h b/Engine/Source/Runtime/AIModule/Classes/BehaviorTree/Services/BTService_DefaultFocus.h
--- a/Engine/Source/Runtime/AIModule/Classes/BehaviorTree/Services/BTService_DefaultFocus.h
+++ b/Engine/Source/Runtime/AIModule/Classes/BehaviorTree/Services/BTService_DefaultFocus.h
@@ -26,11 +26,11 @@ class AIMODULE_API UBTService_DefaultFocus : public UBTService_BlackboardBase
 	GENERATED_UCLASS_BODY()
 
 	virtual uint16 GetInstanceMemorySize() const override { return sizeof(FBTFocusMemory); }
-	virtual void OnBecomeRelevant(UBehaviorTreeComponent* OwnerComp, uint8* NodeMemory) override;
-	virtual void OnCeaseRelevant(UBehaviorTreeComponent* OwnerComp, uint8* NodeMemory) override;
+	virtual void OnBecomeRelevant(UBehaviorTreeComponent& OwnerComp, uint8* NodeMemory) override;
+	virtual void OnCeaseRelevant(UBehaviorTreeComponent& OwnerComp, uint8* NodeMemory) override;
 
 	virtual FString GetStaticDescription() const override;
-	virtual void DescribeRuntimeValues(const UBehaviorTreeComponent* OwnerComp, uint8* NodeMemory, EBTDescriptionVerbosity::Type Verbosity, TArray<FString>& Values) const override;
+	virtual void DescribeRuntimeValues(const UBehaviorTreeComponent& OwnerComp, uint8* NodeMemory, EBTDescriptionVerbosity::Type Verbosity, TArray<FString>& Values) const override;
 
 #if WITH_EDITOR
 	virtual FName GetNodeIconName() const override;
diff --git a/Engine/Source/Runtime/AIModule/Classes/BehaviorTree/Tasks/BTTask_BlueprintBase.h b/Engine/Source/Runtime/AIModule/Classes/BehaviorTree/Tasks/BTTask_BlueprintBase.h
--- a/Engine/Source/Runtime/AIModule/Classes/BehaviorTree/Tasks/BTTask_BlueprintBase.h
+++ b/Engine/Source/Runtime/AIModule/Classes/BehaviorTree/Tasks/BTTask_BlueprintBase.h
@@ -23,12 +23,12 @@ class AIMODULE_API UBTTask_BlueprintBase : public UBTTaskNode
 	/** setup node name */
 	virtual void PostInitProperties() override;
 
-	virtual EBTNodeResult::Type ExecuteTask(UBehaviorTreeComponent* OwnerComp, uint8* NodeMemory) override;
-	virtual EBTNodeResult::Type AbortTask(UBehaviorTreeComponent* OwnerComp, uint8* NodeMemory) override;
+	virtual EBTNodeResult::Type ExecuteTask(UBehaviorTreeComponent& OwnerComp, uint8* NodeMemory) override;
+	virtual EBTNodeResult::Type AbortTask(UBehaviorTreeComponent& OwnerComp, uint8* NodeMemory) override;
 
 	virtual FString GetStaticDescription() const override;
-	virtual void DescribeRuntimeValues(const UBehaviorTreeComponent* OwnerComp, uint8* NodeMemory, EBTDescriptionVerbosity::Type Verbosity, TArray<FString>& Values) const override;
-	virtual void OnInstanceDestroyed(UBehaviorTreeComponent* OwnerComp) override;
+	virtual void DescribeRuntimeValues(const UBehaviorTreeComponent& OwnerComp, uint8* NodeMemory, EBTDescriptionVerbosity::Type Verbosity, TArray<FString>& Values) const override;
+	virtual void OnInstanceDestroyed(UBehaviorTreeComponent& OwnerComp) override;
 
 	virtual void SetOwner(AActor* ActorOwner) override;
 
@@ -128,7 +128,7 @@ protected:
 	bool IsTaskExecuting() const;
 	
 	/** ticks this task */
-	virtual void TickTask(UBehaviorTreeComponent* OwnerComp, uint8* NodeMemory, float DeltaSeconds) override;
+	virtual void TickTask(UBehaviorTreeComponent& OwnerComp, uint8* NodeMemory, float DeltaSeconds) override;
 
 	friend FBehaviorBlueprintDetails;
 };
diff --git a/Engine/Source/Runtime/AIModule/Classes/BehaviorTree/Tasks/BTTask_MakeNoise.h b/Engine/Source/Runtime/AIModule/Classes/BehaviorTree/Tasks/BTTask_MakeNoise.h
--- a/Engine/Source/Runtime/AIModule/Classes/BehaviorTree/Tasks/BTTask_MakeNoise.h
+++ b/Engine/Source/Runtime/AIModule/Classes/BehaviorTree/Tasks/BTTask_MakeNoise.h
@@ -16,7 +16,7 @@ class AIMODULE_API UBTTask_MakeNoise : public UBTTaskNode
 	UPROPERTY(Category=Node, EditAnywhere, meta=(ClampMin = "0", ClampMax = "1", UIMin = "0", UIMax = "1"))
 	float Loudnes;
 
-	virtual EBTNodeResult::Type ExecuteTask(UBehaviorTreeComponent* OwnerComp, uint8* NodeMemory) override;
+	virtual EBTNodeResult::Type ExecuteTask(UBehaviorTreeComponent& OwnerComp, uint8* NodeMemory) override;
 
 #if WITH_EDITOR
 	virtual FName GetNodeIconName() const override;
diff --git a/Engine/Source/Runtime/AIModule/Classes/BehaviorTree/Tasks/BTTask_MoveDirectlyToward.h b/Engine/Source/Runtime/AIModule/Classes/BehaviorTree/Tasks/BTTask_MoveDirectlyToward.h
--- a/Engine/Source/Runtime/AIModule/Classes/BehaviorTree/Tasks/BTTask_MoveDirectlyToward.h
+++ b/Engine/Source/Runtime/AIModule/Classes/BehaviorTree/Tasks/BTTask_MoveDirectlyToward.h
@@ -34,10 +34,10 @@ class AIMODULE_API UBTTask_MoveDirectlyToward : public UBTTask_BlackboardBase
 	UPROPERTY(Category=Node, EditAnywhere)
 	uint32 bAllowStrafe : 1;
 
-	virtual EBTNodeResult::Type ExecuteTask(UBehaviorTreeComponent* OwnerComp, uint8* NodeMemory) override;
-	virtual EBTNodeResult::Type AbortTask(UBehaviorTreeComponent* OwnerComp, uint8* NodeMemory) override;
+	virtual EBTNodeResult::Type ExecuteTask(UBehaviorTreeComponent& OwnerComp, uint8* NodeMemory) override;
+	virtual EBTNodeResult::Type AbortTask(UBehaviorTreeComponent& OwnerComp, uint8* NodeMemory) override;
 	virtual uint16 GetInstanceMemorySize() const override;
-	virtual void DescribeRuntimeValues(const UBehaviorTreeComponent* OwnerComp, uint8* NodeMemory, EBTDescriptionVerbosity::Type Verbosity, TArray<FString>& Values) const override;
+	virtual void DescribeRuntimeValues(const UBehaviorTreeComponent& OwnerComp, uint8* NodeMemory, EBTDescriptionVerbosity::Type Verbosity, TArray<FString>& Values) const override;
 	virtual FString GetStaticDescription() const override;
 
 #if WITH_EDITOR
diff --git a/Engine/Source/Runtime/AIModule/Classes/BehaviorTree/Tasks/BTTask_MoveTo.h b/Engine/Source/Runtime/AIModule/Classes/BehaviorTree/Tasks/BTTask_MoveTo.h
--- a/Engine/Source/Runtime/AIModule/Classes/BehaviorTree/Tasks/BTTask_MoveTo.h
+++ b/Engine/Source/Runtime/AIModule/Classes/BehaviorTree/Tasks/BTTask_MoveTo.h
@@ -33,14 +33,14 @@ class AIMODULE_API UBTTask_MoveTo : public UBTTask_BlackboardBase
 	UPROPERTY(Category=Node, EditAnywhere)
 	uint32 bAllowStrafe : 1;
 	
-	virtual EBTNodeResult::Type ExecuteTask(UBehaviorTreeComponent* OwnerComp, uint8* NodeMemory) override;
-	virtual EBTNodeResult::Type AbortTask(UBehaviorTreeComponent* OwnerComp, uint8* NodeMemory) override;
-	virtual void TickTask(UBehaviorTreeComponent* OwnerComp, uint8* NodeMemory, float DeltaSeconds) override;
+	virtual EBTNodeResult::Type ExecuteTask(UBehaviorTreeComponent& OwnerComp, uint8* NodeMemory) override;
+	virtual EBTNodeResult::Type AbortTask(UBehaviorTreeComponent& OwnerComp, uint8* NodeMemory) override;
+	virtual void TickTask(UBehaviorTreeComponent& OwnerComp, uint8* NodeMemory, float DeltaSeconds) override;
 	virtual uint16 GetInstanceMemorySize() const override;
 
-	virtual void OnMessage(UBehaviorTreeComponent* OwnerComp, uint8* NodeMemory, FName Message, int32 RequestID, bool bSuccess) override;
+	virtual void OnMessage(UBehaviorTreeComponent& OwnerComp, uint8* NodeMemory, FName Message, int32 RequestID, bool bSuccess) override;
 
-	virtual void DescribeRuntimeValues(const UBehaviorTreeComponent* OwnerComp, uint8* NodeMemory, EBTDescriptionVerbosity::Type Verbosity, TArray<FString>& Values) const override;
+	virtual void DescribeRuntimeValues(const UBehaviorTreeComponent& OwnerComp, uint8* NodeMemory, EBTDescriptionVerbosity::Type Verbosity, TArray<FString>& Values) const override;
 	virtual FString GetStaticDescription() const override;
 
 #if WITH_EDITOR
@@ -49,5 +49,5 @@ class AIMODULE_API UBTTask_MoveTo : public UBTTask_BlackboardBase
 
 protected:
 
-	EBTNodeResult::Type PerformMoveTask(UBehaviorTreeComponent* OwnerComp, uint8* NodeMemory);
+	EBTNodeResult::Type PerformMoveTask(UBehaviorTreeComponent& OwnerComp, uint8* NodeMemory);
 };
diff --git a/Engine/Source/Runtime/AIModule/Classes/BehaviorTree/Tasks/BTTask_PlaySound.h b/Engine/Source/Runtime/AIModule/Classes/BehaviorTree/Tasks/BTTask_PlaySound.h
--- a/Engine/Source/Runtime/AIModule/Classes/BehaviorTree/Tasks/BTTask_PlaySound.h
+++ b/Engine/Source/Runtime/AIModule/Classes/BehaviorTree/Tasks/BTTask_PlaySound.h
@@ -19,7 +19,7 @@ class AIMODULE_API UBTTask_PlaySound : public UBTTaskNode
 	UPROPERTY(Category=Node, EditAnywhere)
 	USoundCue* SoundToPlay;
 
-	virtual EBTNodeResult::Type ExecuteTask(UBehaviorTreeComponent* OwnerComp, uint8* NodeMemory) override;
+	virtual EBTNodeResult::Type ExecuteTask(UBehaviorTreeComponent& OwnerComp, uint8* NodeMemory) override;
 	virtual FString GetStaticDescription() const override;
 
 #if WITH_EDITOR
diff --git a/Engine/Source/Runtime/AIModule/Classes/BehaviorTree/Tasks/BTTask_RotateToFaceBBEntry.h b/Engine/Source/Runtime/AIModule/Classes/BehaviorTree/Tasks/BTTask_RotateToFaceBBEntry.h
--- a/Engine/Source/Runtime/AIModule/Classes/BehaviorTree/Tasks/BTTask_RotateToFaceBBEntry.h
+++ b/Engine/Source/Runtime/AIModule/Classes/BehaviorTree/Tasks/BTTask_RotateToFaceBBEntry.h
@@ -28,10 +28,10 @@ public:
 	virtual void PostInitProperties() override;
 	virtual void PostLoad() override;
 
-	virtual EBTNodeResult::Type ExecuteTask(UBehaviorTreeComponent* OwnerComp, uint8* NodeMemory) override;
-	virtual void TickTask(UBehaviorTreeComponent* OwnerComp, uint8* NodeMemory, float DeltaSeconds);
-	virtual EBTNodeResult::Type AbortTask(UBehaviorTreeComponent* OwnerComp, uint8* NodeMemory) override;
-	virtual void DescribeRuntimeValues(const UBehaviorTreeComponent* OwnerComp, uint8* NodeMemory, EBTDescriptionVerbosity::Type Verbosity, TArray<FString>& Values) const override;
+	virtual EBTNodeResult::Type ExecuteTask(UBehaviorTreeComponent& OwnerComp, uint8* NodeMemory) override;
+	virtual void TickTask(UBehaviorTreeComponent& OwnerComp, uint8* NodeMemory, float DeltaSeconds);
+	virtual EBTNodeResult::Type AbortTask(UBehaviorTreeComponent& OwnerComp, uint8* NodeMemory) override;
+	virtual void DescribeRuntimeValues(const UBehaviorTreeComponent& OwnerComp, uint8* NodeMemory, EBTDescriptionVerbosity::Type Verbosity, TArray<FString>& Values) const override;
 	virtual FString GetStaticDescription() const override;
 	
 	virtual uint16 GetInstanceMemorySize() const override { return sizeof(FBTFocusMemory); }
diff --git a/Engine/Source/Runtime/AIModule/Classes/BehaviorTree/Tasks/BTTask_RunBehavior.h b/Engine/Source/Runtime/AIModule/Classes/BehaviorTree/Tasks/BTTask_RunBehavior.h
--- a/Engine/Source/Runtime/AIModule/Classes/BehaviorTree/Tasks/BTTask_RunBehavior.h
+++ b/Engine/Source/Runtime/AIModule/Classes/BehaviorTree/Tasks/BTTask_RunBehavior.h
@@ -23,9 +23,9 @@ class AIMODULE_API UBTTask_RunBehavior : public UBTTaskNode
 {
 	GENERATED_UCLASS_BODY()
 
-	virtual EBTNodeResult::Type ExecuteTask(UBehaviorTreeComponent* OwnerComp, uint8* NodeMemory) override;
+	virtual EBTNodeResult::Type ExecuteTask(UBehaviorTreeComponent& OwnerComp, uint8* NodeMemory) override;
 	virtual uint16 GetInstanceMemorySize() const override;
-	virtual void CleanupMemory(UBehaviorTreeComponent* OwnerComp, uint8* NodeMemory, EBTMemoryClear::Type CleanupType) const override;
+	virtual void CleanupMemory(UBehaviorTreeComponent& OwnerComp, uint8* NodeMemory, EBTMemoryClear::Type CleanupType) const override;
 	virtual FString GetStaticDescription() const override;
 
 #if WITH_EDITOR
@@ -33,7 +33,7 @@ class AIMODULE_API UBTTask_RunBehavior : public UBTTaskNode
 #endif // WITH_EDITOR
 
 	/** called on instance startup, prepares root level nodes to use */
-	void InjectNodes(UBehaviorTreeComponent* OwnerComp, uint8* NodeMemory, int32& InstancedIndex) const;
+	void InjectNodes(UBehaviorTreeComponent& OwnerComp, uint8* NodeMemory, int32& InstancedIndex) const;
 
 	/** @returns number of injected nodes */
 	int32 GetInjectedNodesCount() const;
diff --git a/Engine/Source/Runtime/AIModule/Classes/BehaviorTree/Tasks/BTTask_RunEQSQuery.h b/Engine/Source/Runtime/AIModule/Classes/BehaviorTree/Tasks/BTTask_RunEQSQuery.h
--- a/Engine/Source/Runtime/AIModule/Classes/BehaviorTree/Tasks/BTTask_RunEQSQuery.h
+++ b/Engine/Source/Runtime/AIModule/Classes/BehaviorTree/Tasks/BTTask_RunEQSQuery.h
@@ -29,10 +29,10 @@ class AIMODULE_API UBTTask_RunEQSQuery : public UBTTask_BlackboardBase
 	UPROPERTY(Category=Node, EditAnywhere)
 	TArray<FEnvNamedValue> QueryParams;
 
-	virtual EBTNodeResult::Type ExecuteTask(UBehaviorTreeComponent* OwnerComp, uint8* NodeMemory) override;
-	virtual EBTNodeResult::Type AbortTask(UBehaviorTreeComponent* OwnerComp, uint8* NodeMemory) override;
+	virtual EBTNodeResult::Type ExecuteTask(UBehaviorTreeComponent& OwnerComp, uint8* NodeMemory) override;
+	virtual EBTNodeResult::Type AbortTask(UBehaviorTreeComponent& OwnerComp, uint8* NodeMemory) override;
 
-	virtual void DescribeRuntimeValues(const UBehaviorTreeComponent* OwnerComp, uint8* NodeMemory, EBTDescriptionVerbosity::Type Verbosity, TArray<FString>& Values) const override;
+	virtual void DescribeRuntimeValues(const UBehaviorTreeComponent& OwnerComp, uint8* NodeMemory, EBTDescriptionVerbosity::Type Verbosity, TArray<FString>& Values) const override;
 	virtual FString GetStaticDescription() const override;
 	virtual uint16 GetInstanceMemorySize() const override;
 
diff --git a/Engine/Source/Runtime/AIModule/Classes/BehaviorTree/Tasks/BTTask_Wait.h b/Engine/Source/Runtime/AIModule/Classes/BehaviorTree/Tasks/BTTask_Wait.h
--- a/Engine/Source/Runtime/AIModule/Classes/BehaviorTree/Tasks/BTTask_Wait.h
+++ b/Engine/Source/Runtime/AIModule/Classes/BehaviorTree/Tasks/BTTask_Wait.h
@@ -22,9 +22,9 @@ class AIMODULE_API UBTTask_Wait : public UBTTaskNode
 	UPROPERTY(Category=Wait, EditAnywhere)
 	float WaitTime;
 
-	virtual EBTNodeResult::Type ExecuteTask(UBehaviorTreeComponent* OwnerComp, uint8* NodeMemory) override;
+	virtual EBTNodeResult::Type ExecuteTask(UBehaviorTreeComponent& OwnerComp, uint8* NodeMemory) override;
 	virtual uint16 GetInstanceMemorySize() const override;
-	virtual void DescribeRuntimeValues(const UBehaviorTreeComponent* OwnerComp, uint8* NodeMemory, EBTDescriptionVerbosity::Type Verbosity, TArray<FString>& Values) const override;
+	virtual void DescribeRuntimeValues(const UBehaviorTreeComponent& OwnerComp, uint8* NodeMemory, EBTDescriptionVerbosity::Type Verbosity, TArray<FString>& Values) const override;
 	virtual FString GetStaticDescription() const override;
 
 #if WITH_EDITOR
@@ -33,5 +33,5 @@ class AIMODULE_API UBTTask_Wait : public UBTTaskNode
 
 protected:
 
-	virtual void TickTask(UBehaviorTreeComponent* OwnerComp, uint8* NodeMemory, float DeltaSeconds) override;
+	virtual void TickTask(UBehaviorTreeComponent& OwnerComp, uint8* NodeMemory, float DeltaSeconds) override;
 };
diff --git a/Engine/Source/Runtime/AIModule/Classes/BehaviorTree/Tasks/BTTask_WaitBlackboardTime.h b/Engine/Source/Runtime/AIModule/Classes/BehaviorTree/Tasks/BTTask_WaitBlackboardTime.h
--- a/Engine/Source/Runtime/AIModule/Classes/BehaviorTree/Tasks/BTTask_WaitBlackboardTime.h
+++ b/Engine/Source/Runtime/AIModule/Classes/BehaviorTree/Tasks/BTTask_WaitBlackboardTime.h
@@ -14,7 +14,7 @@ class AIMODULE_API UBTTask_WaitBlackboardTime : public UBTTask_Wait
 	GENERATED_UCLASS_BODY()
 
 	virtual void InitializeFromAsset(UBehaviorTree& Asset) override;
-	virtual EBTNodeResult::Type ExecuteTask(UBehaviorTreeComponent* OwnerComp, uint8* NodeMemory) override;
+	virtual EBTNodeResult::Type ExecuteTask(UBehaviorTreeComponent& OwnerComp, uint8* NodeMemory) override;
 	virtual FString GetStaticDescription() const override;
 
 	/** get name of selected blackboard key */
diff --git a/Engine/Source/Runtime/AIModule/Private/BehaviorTree/BTAuxiliaryNode.cpp b/Engine/Source/Runtime/AIModule/Private/BehaviorTree/BTAuxiliaryNode.cpp
--- a/Engine/Source/Runtime/AIModule/Private/BehaviorTree/BTAuxiliaryNode.cpp
+++ b/Engine/Source/Runtime/AIModule/Private/BehaviorTree/BTAuxiliaryNode.cpp
@@ -11,7 +11,7 @@ UBTAuxiliaryNode::UBTAuxiliaryNode(const FObjectInitializer& ObjectInitializer)
 	bTickIntervals = false;
 }
 
-void UBTAuxiliaryNode::WrappedOnBecomeRelevant(UBehaviorTreeComponent* OwnerComp, uint8* NodeMemory) const
+void UBTAuxiliaryNode::WrappedOnBecomeRelevant(UBehaviorTreeComponent& OwnerComp, uint8* NodeMemory) const
 {
 	if (bNotifyBecomeRelevant)
 	{
@@ -23,7 +23,7 @@ void UBTAuxiliaryNode::WrappedOnBecomeRelevant(UBehaviorTreeComponent* OwnerComp
 	}
 }
 
-void UBTAuxiliaryNode::WrappedOnCeaseRelevant(UBehaviorTreeComponent* OwnerComp, uint8* NodeMemory) const
+void UBTAuxiliaryNode::WrappedOnCeaseRelevant(UBehaviorTreeComponent& OwnerComp, uint8* NodeMemory) const
 {
 	if (bNotifyCeaseRelevant)
 	{
@@ -35,7 +35,7 @@ void UBTAuxiliaryNode::WrappedOnCeaseRelevant(UBehaviorTreeComponent* OwnerComp,
 	}
 }
 
-void UBTAuxiliaryNode::WrappedTickNode(UBehaviorTreeComponent* OwnerComp, uint8* NodeMemory, float DeltaSeconds) const
+void UBTAuxiliaryNode::WrappedTickNode(UBehaviorTreeComponent& OwnerComp, uint8* NodeMemory, float DeltaSeconds) const
 {
 	if (bNotifyTick)
 	{
@@ -73,7 +73,7 @@ void UBTAuxiliaryNode::SetNextTickTime(uint8* NodeMemory, float RemainingTime) c
 	}
 }
 
-void UBTAuxiliaryNode::DescribeRuntimeValues(const UBehaviorTreeComponent* OwnerComp, uint8* NodeMemory, EBTDescriptionVerbosity::Type Verbosity, TArray<FString>& Values) const
+void UBTAuxiliaryNode::DescribeRuntimeValues(const UBehaviorTreeComponent& OwnerComp, uint8* NodeMemory, EBTDescriptionVerbosity::Type Verbosity, TArray<FString>& Values) const
 {
 	Super::DescribeRuntimeValues(OwnerComp, NodeMemory, Verbosity, Values);
 
@@ -84,17 +84,17 @@ void UBTAuxiliaryNode::DescribeRuntimeValues(const UBehaviorTreeComponent* Owner
 	}
 }
 
-void UBTAuxiliaryNode::OnBecomeRelevant(UBehaviorTreeComponent* OwnerComp, uint8* NodeMemory)
+void UBTAuxiliaryNode::OnBecomeRelevant(UBehaviorTreeComponent& OwnerComp, uint8* NodeMemory)
 {
 	// empty in base class
 }
 
-void UBTAuxiliaryNode::OnCeaseRelevant(UBehaviorTreeComponent* OwnerComp, uint8* NodeMemory)
+void UBTAuxiliaryNode::OnCeaseRelevant(UBehaviorTreeComponent& OwnerComp, uint8* NodeMemory)
 {
 	// empty in base class
 }
 
-void UBTAuxiliaryNode::TickNode(UBehaviorTreeComponent* OwnerComp, uint8* NodeMemory, float DeltaSeconds)
+void UBTAuxiliaryNode::TickNode(UBehaviorTreeComponent& OwnerComp, uint8* NodeMemory, float DeltaSeconds)
 {
 	// empty in base class
 }
diff --git a/Engine/Source/Runtime/AIModule/Private/BehaviorTree/BTCompositeNode.cpp b/Engine/Source/Runtime/AIModule/Private/BehaviorTree/BTCompositeNode.cpp
--- a/Engine/Source/Runtime/AIModule/Private/BehaviorTree/BTCompositeNode.cpp
+++ b/Engine/Source/Runtime/AIModule/Private/BehaviorTree/BTCompositeNode.cpp
@@ -28,7 +28,7 @@ int32 UBTCompositeNode::FindChildToExecute(FBehaviorTreeSearchData& SearchData,
 			ChildIdx = GetNextChild(SearchData, ChildIdx, LastResult))
 		{
 			// check decorators
-			if (DoDecoratorsAllowExecution(SearchData.OwnerComp, SearchData.OwnerComp->ActiveInstanceIdx, ChildIdx))
+			if (DoDecoratorsAllowExecution(SearchData.OwnerComp, SearchData.OwnerComp.ActiveInstanceIdx, ChildIdx))
 			{
 				OnChildActivation(SearchData, ChildIdx);
 				RetIdx = ChildIdx;
@@ -45,9 +45,9 @@ int32 UBTCompositeNode::FindChildToExecute(FBehaviorTreeSearchData& SearchData,
 	return RetIdx;
 }
 
-int32 UBTCompositeNode::GetChildIndex(FBehaviorTreeSearchData& SearchData, const UBTNode* ChildNode) const
+int32 UBTCompositeNode::GetChildIndex(FBehaviorTreeSearchData& SearchData, const UBTNode& ChildNode) const
 {
-	if (ChildNode->GetParentNode() != this)
+	if (ChildNode.GetParentNode() != this)
 	{
 		FBTCompositeMemory* NodeMemory = GetNodeMemory<FBTCompositeMemory>(SearchData);
 		return NodeMemory->CurrentChild;		
@@ -56,12 +56,12 @@ int32 UBTCompositeNode::GetChildIndex(FBehaviorTreeSearchData& SearchData, const
 	return GetChildIndex(ChildNode);
 }
 
-int32 UBTCompositeNode::GetChildIndex(const UBTNode* ChildNode) const
+int32 UBTCompositeNode::GetChildIndex(const UBTNode& ChildNode) const
 {
 	for (int32 ChildIndex = 0; ChildIndex < Children.Num(); ChildIndex++)
 	{
-		if (Children[ChildIndex].ChildComposite == ChildNode ||
-			Children[ChildIndex].ChildTask == ChildNode)
+		if (Children[ChildIndex].ChildComposite == &ChildNode ||
+			Children[ChildIndex].ChildTask == &ChildNode)
 		{
 			return ChildIndex;
 		}
@@ -70,7 +70,7 @@ int32 UBTCompositeNode::GetChildIndex(const UBTNode* ChildNode) const
 	return BTSpecialChild::ReturnToParent;
 }
 
-void UBTCompositeNode::OnChildActivation(FBehaviorTreeSearchData& SearchData, const UBTNode* ChildNode) const
+void UBTCompositeNode::OnChildActivation(FBehaviorTreeSearchData& SearchData, const UBTNode& ChildNode) const
 {
 	OnChildActivation(SearchData, GetChildIndex(SearchData, ChildNode));
 }
@@ -94,7 +94,7 @@ void UBTCompositeNode::OnChildActivation(FBehaviorTreeSearchData& SearchData, in
 	NodeMemory->CurrentChild = ChildIndex;
 }
 
-void UBTCompositeNode::OnChildDeactivation(FBehaviorTreeSearchData& SearchData, const UBTNode* ChildNode, EBTNodeResult::Type& NodeResult) const
+void UBTCompositeNode::OnChildDeactivation(FBehaviorTreeSearchData& SearchData, const UBTNode& ChildNode, EBTNodeResult::Type& NodeResult) const
 {
 	OnChildDeactivation(SearchData, GetChildIndex(SearchData, ChildNode), NodeResult);
 }
@@ -129,7 +129,7 @@ void UBTCompositeNode::OnNodeActivation(FBehaviorTreeSearchData& SearchData) con
 	for (int32 ServiceIndex = 0; ServiceIndex < Services.Num(); ServiceIndex++)
 	{
 		// add services when execution flow enters this composite
-		SearchData.AddUniqueUpdate(FBehaviorTreeSearchUpdate(Services[ServiceIndex], SearchData.OwnerComp->GetActiveInstanceIdx(), EBTNodeUpdateMode::Add));
+		SearchData.AddUniqueUpdate(FBehaviorTreeSearchUpdate(Services[ServiceIndex], SearchData.OwnerComp.GetActiveInstanceIdx(), EBTNodeUpdateMode::Add));
 
 		// give services chance to perform initial tick before searching further
 		Services[ServiceIndex]->NotifyParentActivation(SearchData);
@@ -146,7 +146,7 @@ void UBTCompositeNode::OnNodeDeactivation(FBehaviorTreeSearchData& SearchData, E
 	// remove all services if execution flow leaves this composite
 	for (int32 ServiceIndex = 0; ServiceIndex < Services.Num(); ServiceIndex++)
 	{
-		SearchData.AddUniqueUpdate(FBehaviorTreeSearchUpdate(Services[ServiceIndex], SearchData.OwnerComp->GetActiveInstanceIdx(), EBTNodeUpdateMode::Remove));
+		SearchData.AddUniqueUpdate(FBehaviorTreeSearchUpdate(Services[ServiceIndex], SearchData.OwnerComp.GetActiveInstanceIdx(), EBTNodeUpdateMode::Remove));
 	}
 }
 
@@ -168,12 +168,12 @@ void UBTCompositeNode::NotifyDecoratorsOnActivation(FBehaviorTreeSearchData& Sea
 		switch (DecoratorOb->GetFlowAbortMode())
 		{
 			case EBTFlowAbortMode::LowerPriority:
-				SearchData.AddUniqueUpdate(FBehaviorTreeSearchUpdate(DecoratorOb, SearchData.OwnerComp->GetActiveInstanceIdx(), EBTNodeUpdateMode::Remove));
+				SearchData.AddUniqueUpdate(FBehaviorTreeSearchUpdate(DecoratorOb, SearchData.OwnerComp.GetActiveInstanceIdx(), EBTNodeUpdateMode::Remove));
 				break;
 
 			case EBTFlowAbortMode::Self:
 			case EBTFlowAbortMode::Both:
-				SearchData.AddUniqueUpdate(FBehaviorTreeSearchUpdate(DecoratorOb, SearchData.OwnerComp->GetActiveInstanceIdx(), EBTNodeUpdateMode::Add));
+				SearchData.AddUniqueUpdate(FBehaviorTreeSearchUpdate(DecoratorOb, SearchData.OwnerComp.GetActiveInstanceIdx(), EBTNodeUpdateMode::Add));
 				break;
 
 			default:
@@ -195,12 +195,12 @@ void UBTCompositeNode::NotifyDecoratorsOnDeactivation(FBehaviorTreeSearchData& S
 		if (DecoratorOb->GetFlowAbortMode() == EBTFlowAbortMode::Self)
 		{
 			// - observers with mode "Self" are now out of scope, remove them
-			SearchData.AddUniqueUpdate(FBehaviorTreeSearchUpdate(DecoratorOb, SearchData.OwnerComp->GetActiveInstanceIdx(), EBTNodeUpdateMode::Remove));
+			SearchData.AddUniqueUpdate(FBehaviorTreeSearchUpdate(DecoratorOb, SearchData.OwnerComp.GetActiveInstanceIdx(), EBTNodeUpdateMode::Remove));
 		}
 		else if (DecoratorOb->GetFlowAbortMode() == EBTFlowAbortMode::LowerPriority)
 		{
 			// - observers with mode "Lower Priority" will try to reactivate themselves ("Both" is not removed on node activation)
-			SearchData.AddUniqueUpdate(FBehaviorTreeSearchUpdate(DecoratorOb, SearchData.OwnerComp->GetActiveInstanceIdx(), EBTNodeUpdateMode::AddForLowerPri));
+			SearchData.AddUniqueUpdate(FBehaviorTreeSearchUpdate(DecoratorOb, SearchData.OwnerComp.GetActiveInstanceIdx(), EBTNodeUpdateMode::AddForLowerPri));
 		}
 	}
 }
@@ -208,7 +208,7 @@ void UBTCompositeNode::NotifyDecoratorsOnDeactivation(FBehaviorTreeSearchData& S
 void UBTCompositeNode::NotifyDecoratorsOnFailedActivation(FBehaviorTreeSearchData& SearchData, int32 ChildIdx, EBTNodeResult::Type& NodeResult) const
 {
 	const FBTCompositeChild& ChildInfo = Children[ChildIdx];
-	const uint16 ActiveInstanceIdx = SearchData.OwnerComp->GetActiveInstanceIdx();
+	const uint16 ActiveInstanceIdx = SearchData.OwnerComp.GetActiveInstanceIdx();
 
 	for (int32 DecoratorIndex = 0; DecoratorIndex < ChildInfo.Decorators.Num(); DecoratorIndex++)
 	{
@@ -222,10 +222,10 @@ void UBTCompositeNode::NotifyDecoratorsOnFailedActivation(FBehaviorTreeSearchDat
 		}
 	}
 
-	SearchData.OwnerComp->StoreDebuggerSearchStep(GetChildNode(ChildIdx), ActiveInstanceIdx, NodeResult);
+	SearchData.OwnerComp.StoreDebuggerSearchStep(GetChildNode(ChildIdx), ActiveInstanceIdx, NodeResult);
 }
 
-void UBTCompositeNode::NotifyChildExecution(UBehaviorTreeComponent* OwnerComp, uint8* NodeMemory, int32 ChildIdx, EBTNodeResult::Type& NodeResult) const
+void UBTCompositeNode::NotifyChildExecution(UBehaviorTreeComponent& OwnerComp, uint8* NodeMemory, int32 ChildIdx, EBTNodeResult::Type& NodeResult) const
 {
 }
 
@@ -237,13 +237,13 @@ void UBTCompositeNode::NotifyNodeDeactivation(FBehaviorTreeSearchData& SearchDat
 {
 }
 
-void UBTCompositeNode::ConditionalNotifyChildExecution(UBehaviorTreeComponent* OwnerComp, uint8* NodeMemory, const UBTNode* ChildNode, EBTNodeResult::Type& NodeResult) const
+void UBTCompositeNode::ConditionalNotifyChildExecution(UBehaviorTreeComponent& OwnerComp, uint8* NodeMemory, const UBTNode& ChildNode, EBTNodeResult::Type& NodeResult) const
 {
 	if (bUseChildExecutionNotify)
 	{
 		for (int32 ChildIndex = 0; ChildIndex < Children.Num(); ChildIndex++)
 		{
-			if (Children[ChildIndex].ChildComposite == ChildNode || Children[ChildIndex].ChildTask == ChildNode)
+			if (Children[ChildIndex].ChildComposite == &ChildNode || Children[ChildIndex].ChildTask == &ChildNode)
 			{
 				NotifyChildExecution(OwnerComp, NodeMemory, ChildIndex, NodeResult);
 				break;
@@ -275,7 +275,7 @@ struct FOperationStackInfo
 		NumLeft(DecoratorOp.Number), Op(DecoratorOp.Operation), bHasForcedResult(0) {};
 };
 
-static bool UpdateOperationStack(const UBehaviorTreeComponent* OwnerComp, FString& Indent,
+static bool UpdateOperationStack(const UBehaviorTreeComponent& OwnerComp, FString& Indent,
 								 TArray<FOperationStackInfo>& Stack, bool bTestResult,
 								 int32& FailedDecoratorIdx, int32& NodeDecoratorIdx, bool& bShouldStoreNodeIndex)
 {
@@ -328,7 +328,7 @@ static bool UpdateOperationStack(const UBehaviorTreeComponent* OwnerComp, FStrin
 
 	if (CurrentOp.NumLeft == 0)
 	{
-		UE_VLOG(OwnerComp->GetOwner(), LogBehaviorTree, Verbose, TEXT("%s%s finished: %s"), *Indent,
+		UE_VLOG(OwnerComp.GetOwner(), LogBehaviorTree, Verbose, TEXT("%s%s finished: %s"), *Indent,
 			*DescribeLogicOp(CurrentOp.Op),
 			bTestResult ? TEXT("allowed") : TEXT("forbidden"));
 		Indent = Indent.LeftChop(2);
@@ -340,7 +340,7 @@ static bool UpdateOperationStack(const UBehaviorTreeComponent* OwnerComp, FStrin
 	return bTestResult;
 }
 
-bool UBTCompositeNode::DoDecoratorsAllowExecution(UBehaviorTreeComponent* OwnerComp, int32 InstanceIdx, int32 ChildIdx) const
+bool UBTCompositeNode::DoDecoratorsAllowExecution(UBehaviorTreeComponent& OwnerComp, int32 InstanceIdx, int32 ChildIdx) const
 {
 	const FBTCompositeChild& ChildInfo = Children[ChildIdx];
 	bool bResult = true;
@@ -350,7 +350,7 @@ bool UBTCompositeNode::DoDecoratorsAllowExecution(UBehaviorTreeComponent* OwnerC
 		return bResult;
 	}
 
-	FBehaviorTreeInstance& MyInstance = OwnerComp->InstanceStack[InstanceIdx];
+	FBehaviorTreeInstance& MyInstance = OwnerComp.InstanceStack[InstanceIdx];
 
 	if (ChildInfo.DecoratorOps.Num() == 0)
 	{
@@ -359,11 +359,11 @@ bool UBTCompositeNode::DoDecoratorsAllowExecution(UBehaviorTreeComponent* OwnerC
 		{
 			const UBTDecorator* TestDecorator = ChildInfo.Decorators[DecoratorIndex];
 			const bool bIsAllowed = TestDecorator ? TestDecorator->WrappedCanExecute(OwnerComp, TestDecorator->GetNodeMemory<uint8>(MyInstance)) : false;
-			OwnerComp->StoreDebuggerSearchStep(TestDecorator, InstanceIdx, bIsAllowed);
+			OwnerComp.StoreDebuggerSearchStep(TestDecorator, InstanceIdx, bIsAllowed);
 
 			if (!bIsAllowed)
 			{
-				UE_VLOG(OwnerComp->GetOwner(), LogBehaviorTree, Verbose, TEXT("Child[%d] execution forbidden by %s"),
+				UE_VLOG(OwnerComp.GetOwner(), LogBehaviorTree, Verbose, TEXT("Child[%d] execution forbidden by %s"),
 					ChildIdx, *UBehaviorTreeTypes::DescribeNodeHelper(TestDecorator));
 
 				bResult = false;
@@ -371,7 +371,7 @@ bool UBTCompositeNode::DoDecoratorsAllowExecution(UBehaviorTreeComponent* OwnerC
 			}
 			else
 			{
-				UE_VLOG(OwnerComp->GetOwner(), LogBehaviorTree, Verbose, TEXT("Child[%d] execution allowed by %s"),
+				UE_VLOG(OwnerComp.GetOwner(), LogBehaviorTree, Verbose, TEXT("Child[%d] execution allowed by %s"),
 					ChildIdx, *UBehaviorTreeTypes::DescribeNodeHelper(TestDecorator));
 			}
 		}
@@ -379,7 +379,7 @@ bool UBTCompositeNode::DoDecoratorsAllowExecution(UBehaviorTreeComponent* OwnerC
 	else
 	{
 		// advanced check: follow decorator logic operations (composite decorator on child link)
-		UE_VLOG(OwnerComp->GetOwner(), LogBehaviorTree, Verbose, TEXT("Child[%d] execution test with logic operations"), ChildIdx);
+		UE_VLOG(OwnerComp.GetOwner(), LogBehaviorTree, Verbose, TEXT("Child[%d] execution test with logic operations"), ChildIdx);
 
 		TArray<FOperationStackInfo> OperationStack;
 		FString Indent;
@@ -400,7 +400,7 @@ bool UBTCompositeNode::DoDecoratorsAllowExecution(UBehaviorTreeComponent* OwnerC
 			{
 				OperationStack.Add(FOperationStackInfo(DecoratorOp));
 				Indent += TEXT("  ");
-				UE_VLOG(OwnerComp->GetOwner(), LogBehaviorTree, Verbose, TEXT("%spushed %s:%d"), *Indent,
+				UE_VLOG(OwnerComp.GetOwner(), LogBehaviorTree, Verbose, TEXT("%spushed %s:%d"), *Indent,
 					*DescribeLogicOp(DecoratorOp.Operation), DecoratorOp.Number);
 			}
 			else if (DecoratorOp.Operation == EBTDecoratorLogic::Test)
@@ -417,7 +417,7 @@ bool UBTCompositeNode::DoDecoratorsAllowExecution(UBehaviorTreeComponent* OwnerC
 
 				UBTDecorator* TestDecorator = ChildInfo.Decorators[DecoratorOp.Number];
 				const bool bIsAllowed = bHasOverride ? bCurrentOverride : TestDecorator->WrappedCanExecute(OwnerComp, TestDecorator->GetNodeMemory<uint8>(MyInstance));
-				UE_VLOG(OwnerComp->GetOwner(), LogBehaviorTree, Verbose, TEXT("%s%s %s: %s"), *Indent,
+				UE_VLOG(OwnerComp.GetOwner(), LogBehaviorTree, Verbose, TEXT("%s%s %s: %s"), *Indent,
 					bHasOverride ? TEXT("skipping") : TEXT("testing"),
 					*UBehaviorTreeTypes::DescribeNodeHelper(TestDecorator),
 					bIsAllowed ? TEXT("allowed") : TEXT("forbidden"));
@@ -425,10 +425,10 @@ bool UBTCompositeNode::DoDecoratorsAllowExecution(UBehaviorTreeComponent* OwnerC
 				bResult = UpdateOperationStack(OwnerComp, Indent, OperationStack, bIsAllowed, FailedDecoratorIdx, NodeDecoratorIdx, bShouldStoreNodeIndex);
 				if (OperationStack.Num() == 0)
 				{
-					UE_VLOG(OwnerComp->GetOwner(), LogBehaviorTree, Verbose, TEXT("finished execution test: %s"),
+					UE_VLOG(OwnerComp.GetOwner(), LogBehaviorTree, Verbose, TEXT("finished execution test: %s"),
 						bResult ? TEXT("allowed") : TEXT("forbidden"));
 
-					OwnerComp->StoreDebuggerSearchStep(ChildInfo.Decorators[FMath::Max(0, FailedDecoratorIdx)], InstanceIdx, bResult);
+					OwnerComp.StoreDebuggerSearchStep(ChildInfo.Decorators[FMath::Max(0, FailedDecoratorIdx)], InstanceIdx, bResult);
 					break;
 				}
 			}
@@ -456,7 +456,7 @@ int32 UBTCompositeNode::GetMatchingChildIndex(int32 ActiveInstanceIdx, FBTNodeIn
 		for (int32 ChildIndex = 0; ChildIndex < Children.Num(); ChildIndex++)
 		{
 			const UBTNode* ChildNode = GetChildNode(ChildIndex);
-			if (ChildNode->GetExecutionIndex() > NodeIdx.ExecutionIndex)
+			if (ChildNode != nullptr && ChildNode->GetExecutionIndex() > NodeIdx.ExecutionIndex)
 			{
 				return ChildIndex ? (ChildIndex - 1) : 0;
 			}
@@ -474,7 +474,7 @@ int32 UBTCompositeNode::GetNextChild(FBehaviorTreeSearchData& SearchData, int32
 {
 	FBTCompositeMemory* NodeMemory = GetNodeMemory<FBTCompositeMemory>(SearchData);
 	int32 NextChildIndex = BTSpecialChild::ReturnToParent;
-	uint16 ActiveInstanceIdx = SearchData.OwnerComp->GetActiveInstanceIdx();
+	uint16 ActiveInstanceIdx = SearchData.OwnerComp.GetActiveInstanceIdx();
 
 	// newly activated node, search range not reached yet: select search branch for decorator test
 	if (LastChildIdx == BTSpecialChild::NotInitialized && SearchData.SearchStart.IsSet() &&
@@ -482,7 +482,7 @@ int32 UBTCompositeNode::GetNextChild(FBehaviorTreeSearchData& SearchData, int32
 	{
 		NextChildIndex = GetMatchingChildIndex(ActiveInstanceIdx, SearchData.SearchStart);
 	}
-	else if (NodeMemory->OverrideChild != BTSpecialChild::NotInitialized && !SearchData.OwnerComp->IsRestartPending())
+	else if (NodeMemory->OverrideChild != BTSpecialChild::NotInitialized && !SearchData.OwnerComp.IsRestartPending())
 	{
 		NextChildIndex = NodeMemory->OverrideChild;
 		NodeMemory->OverrideChild = BTSpecialChild::NotInitialized;
@@ -505,9 +505,9 @@ void UBTCompositeNode::SetChildOverride(FBehaviorTreeSearchData& SearchData, int
 	}
 }
 
-void UBTCompositeNode::RequestDelayedExecution(UBehaviorTreeComponent* OwnerComp, EBTNodeResult::Type LastResult) const
+void UBTCompositeNode::RequestDelayedExecution(UBehaviorTreeComponent& OwnerComp, EBTNodeResult::Type LastResult) const
 {
-	OwnerComp->RequestExecution(LastResult);
+	OwnerComp.RequestExecution(LastResult);
 }
 
 uint16 UBTCompositeNode::GetChildExecutionIndex(int32 Index, EBTChildIndex ChildMode) const
@@ -522,12 +522,12 @@ uint16 UBTCompositeNode::GetChildExecutionIndex(int32 Index, EBTChildIndex Child
 	return (LastExecutionIndex + 1);
 }
 
-bool UBTCompositeNode::CanPushSubtree(UBehaviorTreeComponent* OwnerComp, uint8* NodeMemory, int32 ChildIdx) const
+bool UBTCompositeNode::CanPushSubtree(UBehaviorTreeComponent& OwnerComp, uint8* NodeMemory, int32 ChildIdx) const
 {
 	return true;
 }
 
-void UBTCompositeNode::DescribeRuntimeValues(const UBehaviorTreeComponent* OwnerComp, uint8* NodeMemory, EBTDescriptionVerbosity::Type Verbosity, TArray<FString>& Values) const
+void UBTCompositeNode::DescribeRuntimeValues(const UBehaviorTreeComponent& OwnerComp, uint8* NodeMemory, EBTDescriptionVerbosity::Type Verbosity, TArray<FString>& Values) const
 {
 	Super::DescribeRuntimeValues(OwnerComp, NodeMemory, Verbosity, Values);
 
diff --git a/Engine/Source/Runtime/AIModule/Private/BehaviorTree/BTDecorator.cpp b/Engine/Source/Runtime/AIModule/Private/BehaviorTree/BTDecorator.cpp
--- a/Engine/Source/Runtime/AIModule/Private/BehaviorTree/BTDecorator.cpp
+++ b/Engine/Source/Runtime/AIModule/Private/BehaviorTree/BTDecorator.cpp
@@ -25,7 +25,7 @@ void UBTDecorator::InitializeDecorator(uint8 MyChildIndex)
 	ChildIndex = MyChildIndex;
 }
 
-bool UBTDecorator::CalculateRawConditionValue(UBehaviorTreeComponent* OwnerComp, uint8* NodeMemory) const
+bool UBTDecorator::CalculateRawConditionValue(UBehaviorTreeComponent& OwnerComp, uint8* NodeMemory) const
 {
 	return true;
 }
@@ -47,7 +47,7 @@ void UBTDecorator::OnNodeProcessed(FBehaviorTreeSearchData& SearchData, EBTNodeR
 {
 }
 
-bool UBTDecorator::WrappedCanExecute(UBehaviorTreeComponent* OwnerComp, uint8* NodeMemory) const
+bool UBTDecorator::WrappedCanExecute(UBehaviorTreeComponent& OwnerComp, uint8* NodeMemory) const
 {
 	const UBTDecorator* NodeOb = bCreateNodeInstance ? (const UBTDecorator*)GetNodeInstance(OwnerComp, NodeMemory) : this;
 	return NodeOb ? (IsInversed() != NodeOb->CalculateRawConditionValue(OwnerComp, NodeMemory)) : false;
diff --git a/Engine/Source/Runtime/AIModule/Private/BehaviorTree/BTNode.cpp b/Engine/Source/Runtime/AIModule/Private/BehaviorTree/BTNode.cpp
--- a/Engine/Source/Runtime/AIModule/Private/BehaviorTree/BTNode.cpp
+++ b/Engine/Source/Runtime/AIModule/Private/BehaviorTree/BTNode.cpp
@@ -44,45 +44,45 @@ void UBTNode::InitializeNode(UBTCompositeNode* InParentNode, uint16 InExecutionI
 	TreeDepth = InTreeDepth;
 }
 
-void UBTNode::InitializeMemory(UBehaviorTreeComponent* OwnerComp, uint8* NodeMemory, EBTMemoryInit::Type InitType) const
+void UBTNode::InitializeMemory(UBehaviorTreeComponent& OwnerComp, uint8* NodeMemory, EBTMemoryInit::Type InitType) const
 {
 	// empty in base 
 }
 
-void UBTNode::CleanupMemory(UBehaviorTreeComponent* OwnerComp, uint8* NodeMemory, EBTMemoryClear::Type CleanupType) const
+void UBTNode::CleanupMemory(UBehaviorTreeComponent& OwnerComp, uint8* NodeMemory, EBTMemoryClear::Type CleanupType) const
 {
 	// empty in base 
 }
 
-void UBTNode::OnInstanceCreated(UBehaviorTreeComponent* OwnerComp)
+void UBTNode::OnInstanceCreated(UBehaviorTreeComponent& OwnerComp)
 {
 	// empty in base class
 }
 
-void UBTNode::OnInstanceDestroyed(UBehaviorTreeComponent* OwnerComp)
+void UBTNode::OnInstanceDestroyed(UBehaviorTreeComponent& OwnerComp)
 {
 	// empty in base class
 }
 
-void UBTNode::InitializeInSubtree(UBehaviorTreeComponent* OwnerComp, uint8* NodeMemory, int32& NextInstancedIndex, EBTMemoryInit::Type InitType) const
+void UBTNode::InitializeInSubtree(UBehaviorTreeComponent& OwnerComp, uint8* NodeMemory, int32& NextInstancedIndex, EBTMemoryInit::Type InitType) const
 {
 	if (bCreateNodeInstance)
 	{
 		// composite nodes can't be instanced!
 		check(IsA(UBTCompositeNode::StaticClass()) == false);
 
-		UBTNode* NodeInstance = OwnerComp->NodeInstances.IsValidIndex(NextInstancedIndex) ? OwnerComp->NodeInstances[NextInstancedIndex] : NULL;
+		UBTNode* NodeInstance = OwnerComp.NodeInstances.IsValidIndex(NextInstancedIndex) ? OwnerComp.NodeInstances[NextInstancedIndex] : NULL;
 		if (NodeInstance == NULL)
 		{
-			NodeInstance = ConstructObject<UBTNode>(GetClass(), OwnerComp, GetFName(), RF_NoFlags, (UObject*)(this));
+			NodeInstance = ConstructObject<UBTNode>(GetClass(), &OwnerComp, GetFName(), RF_NoFlags, (UObject*)(this));
 			NodeInstance->InitializeNode(GetParentNode(), GetExecutionIndex(), GetMemoryOffset(), GetTreeDepth());
 			NodeInstance->bIsInstanced = true;
 
-			OwnerComp->NodeInstances.Add(NodeInstance);
+			OwnerComp.NodeInstances.Add(NodeInstance);
 		}
 		check(NodeInstance);
 
-		NodeInstance->SetOwner(OwnerComp->GetOwner());
+		NodeInstance->SetOwner(OwnerComp.GetOwner());
 
 		FBTInstancedNodeMemory* MyMemory = GetSpecialNodeMemory<FBTInstancedNodeMemory>(NodeMemory);
 		MyMemory->NodeIdx = NextInstancedIndex;
@@ -96,7 +96,7 @@ void UBTNode::InitializeInSubtree(UBehaviorTreeComponent* OwnerComp, uint8* Node
 	}
 }
 
-void UBTNode::CleanupInSubtree(UBehaviorTreeComponent* OwnerComp, uint8* NodeMemory, EBTMemoryClear::Type CleanupType) const
+void UBTNode::CleanupInSubtree(UBehaviorTreeComponent& OwnerComp, uint8* NodeMemory, EBTMemoryClear::Type CleanupType) const
 {
 	if (!bCreateNodeInstance && !bIsInjected)
 	{
@@ -131,11 +131,11 @@ uint16 UBTNode::GetSpecialMemorySize() const
 	return bCreateNodeInstance ? sizeof(FBTInstancedNodeMemory) : 0;
 }
 
-UBTNode* UBTNode::GetNodeInstance(const UBehaviorTreeComponent* OwnerComp, uint8* NodeMemory) const
+UBTNode* UBTNode::GetNodeInstance(const UBehaviorTreeComponent& OwnerComp, uint8* NodeMemory) const
 {
 	FBTInstancedNodeMemory* MyMemory = GetSpecialNodeMemory<FBTInstancedNodeMemory>(NodeMemory);
-	return OwnerComp && MyMemory && OwnerComp->NodeInstances.IsValidIndex(MyMemory->NodeIdx) ?
-		OwnerComp->NodeInstances[MyMemory->NodeIdx] : NULL;
+	return MyMemory && OwnerComp.NodeInstances.IsValidIndex(MyMemory->NodeIdx) ?
+		OwnerComp.NodeInstances[MyMemory->NodeIdx] : NULL;
 }
 
 UBTNode* UBTNode::GetNodeInstance(FBehaviorTreeSearchData& SearchData) const
@@ -143,7 +143,7 @@ UBTNode* UBTNode::GetNodeInstance(FBehaviorTreeSearchData& SearchData) const
 	return GetNodeInstance(SearchData.OwnerComp, GetNodeMemory<uint8>(SearchData));
 }
 
-FString UBTNode::GetRuntimeDescription(const UBehaviorTreeComponent* OwnerComp, uint8* NodeMemory, EBTDescriptionVerbosity::Type Verbosity) const
+FString UBTNode::GetRuntimeDescription(const UBehaviorTreeComponent& OwnerComp, uint8* NodeMemory, EBTDescriptionVerbosity::Type Verbosity) const
 {
 	FString Description = NodeName.Len() ? FString::Printf(TEXT("%s [%s]"), *NodeName, *GetStaticDescription()) : GetStaticDescription();
 	TArray<FString> RuntimeValues;
@@ -169,7 +169,7 @@ FString UBTNode::GetStaticDescription() const
 	return UBehaviorTreeTypes::GetShortTypeName(this);
 }
 
-void UBTNode::DescribeRuntimeValues(const UBehaviorTreeComponent* OwnerComp, uint8* NodeMemory, EBTDescriptionVerbosity::Type Verbosity, TArray<FString>& Values) const
+void UBTNode::DescribeRuntimeValues(const UBehaviorTreeComponent& OwnerComp, uint8* NodeMemory, EBTDescriptionVerbosity::Type Verbosity, TArray<FString>& Values) const
 {
 	// nothing stored in memory for base class
 }
diff --git a/Engine/Source/Runtime/AIModule/Private/BehaviorTree/BTService.cpp b/Engine/Source/Runtime/AIModule/Private/BehaviorTree/BTService.cpp
--- a/Engine/Source/Runtime/AIModule/Private/BehaviorTree/BTService.cpp
+++ b/Engine/Source/Runtime/AIModule/Private/BehaviorTree/BTService.cpp
@@ -13,7 +13,7 @@ UBTService::UBTService(const FObjectInitializer& ObjectInitializer) : Super(Obje
 	RandomDeviation = 0.1f;
 }
 
-void UBTService::TickNode(UBehaviorTreeComponent* OwnerComp, uint8* NodeMemory, float DeltaSeconds)
+void UBTService::TickNode(UBehaviorTreeComponent& OwnerComp, uint8* NodeMemory, float DeltaSeconds)
 {
 	const float NextTickTime = FMath::FRandRange(FMath::Max(0.0f, Interval - RandomDeviation), (Interval + RandomDeviation));
 	SetNextTickTime(NodeMemory, NextTickTime);
diff --git a/Engine/Source/Runtime/AIModule/Private/BehaviorTree/BTTaskNode.cpp b/Engine/Source/Runtime/AIModule/Private/BehaviorTree/BTTaskNode.cpp
--- a/Engine/Source/Runtime/AIModule/Private/BehaviorTree/BTTaskNode.cpp
+++ b/Engine/Source/Runtime/AIModule/Private/BehaviorTree/BTTaskNode.cpp
@@ -9,29 +9,29 @@ UBTTaskNode::UBTTaskNode(const FObjectInitializer& ObjectInitializer) : Super(Ob
 	bNotifyTick = false;
 }
 
-EBTNodeResult::Type UBTTaskNode::ExecuteTask(UBehaviorTreeComponent* OwnerComp, uint8* NodeMemory)
+EBTNodeResult::Type UBTTaskNode::ExecuteTask(UBehaviorTreeComponent& OwnerComp, uint8* NodeMemory)
 {
 	return EBTNodeResult::Succeeded;
 }
 
-EBTNodeResult::Type UBTTaskNode::AbortTask(UBehaviorTreeComponent* OwnerComp, uint8* NodeMemory)
+EBTNodeResult::Type UBTTaskNode::AbortTask(UBehaviorTreeComponent& OwnerComp, uint8* NodeMemory)
 {
 	return EBTNodeResult::Aborted;
 }
 
-EBTNodeResult::Type UBTTaskNode::WrappedExecuteTask(UBehaviorTreeComponent* OwnerComp, uint8* NodeMemory) const
+EBTNodeResult::Type UBTTaskNode::WrappedExecuteTask(UBehaviorTreeComponent& OwnerComp, uint8* NodeMemory) const
 {
 	const UBTNode* NodeOb = bCreateNodeInstance ? GetNodeInstance(OwnerComp, NodeMemory) : this;
 	return NodeOb ? ((UBTTaskNode*)NodeOb)->ExecuteTask(OwnerComp, NodeMemory) : EBTNodeResult::Failed;
 }
 
-EBTNodeResult::Type UBTTaskNode::WrappedAbortTask(UBehaviorTreeComponent* OwnerComp, uint8* NodeMemory) const
+EBTNodeResult::Type UBTTaskNode::WrappedAbortTask(UBehaviorTreeComponent& OwnerComp, uint8* NodeMemory) const
 {
 	const UBTNode* NodeOb = bCreateNodeInstance ? GetNodeInstance(OwnerComp, NodeMemory) : this;
 	return NodeOb ? ((UBTTaskNode*)NodeOb)->AbortTask(OwnerComp, NodeMemory) : EBTNodeResult::Aborted;
 }
 
-void UBTTaskNode::WrappedTickTask(UBehaviorTreeComponent* OwnerComp, uint8* NodeMemory, float DeltaSeconds) const
+void UBTTaskNode::WrappedTickTask(UBehaviorTreeComponent& OwnerComp, uint8* NodeMemory, float DeltaSeconds) const
 {
 	if (bNotifyTick)
 	{
@@ -49,17 +49,17 @@ void UBTTaskNode::ReceivedMessage(UBrainComponent* BrainComp, const FAIMessage&
 	uint16 InstanceIdx = OwnerComp->FindInstanceContainingNode(this);
 	uint8* NodeMemory = GetNodeMemory<uint8>(OwnerComp->InstanceStack[InstanceIdx]);
 
-	OnMessage(OwnerComp, NodeMemory, Message.MessageName, Message.RequestID, Message.Status == FAIMessage::Success);
+	OnMessage(*OwnerComp, NodeMemory, Message.MessageName, Message.RequestID, Message.Status == FAIMessage::Success);
 }
 
-void UBTTaskNode::TickTask(UBehaviorTreeComponent* OwnerComp, uint8* NodeMemory, float DeltaSeconds)
+void UBTTaskNode::TickTask(UBehaviorTreeComponent& OwnerComp, uint8* NodeMemory, float DeltaSeconds)
 {
 	// empty in base class
 }
 
-void UBTTaskNode::OnMessage(UBehaviorTreeComponent* OwnerComp, uint8* NodeMemory, FName Message, int32 RequestID, bool bSuccess)
+void UBTTaskNode::OnMessage(UBehaviorTreeComponent& OwnerComp, uint8* NodeMemory, FName Message, int32 RequestID, bool bSuccess)
 {
-	const EBTTaskStatus::Type Status = OwnerComp->GetTaskStatus(this);
+	const EBTTaskStatus::Type Status = OwnerComp.GetTaskStatus(this);
 	if (Status == EBTTaskStatus::Active)
 	{
 		FinishLatentTask(OwnerComp, bSuccess ? EBTNodeResult::Succeeded : EBTNodeResult::Failed);
@@ -70,36 +70,36 @@ void UBTTaskNode::OnMessage(UBehaviorTreeComponent* OwnerComp, uint8* NodeMemory
 	}
 }
 
-void UBTTaskNode::FinishLatentTask(UBehaviorTreeComponent* OwnerComp, EBTNodeResult::Type TaskResult) const
+void UBTTaskNode::FinishLatentTask(UBehaviorTreeComponent& OwnerComp, EBTNodeResult::Type TaskResult) const
 {
 	// OnTaskFinished must receive valid template node
-	UBTTaskNode* TemplateNode = (UBTTaskNode*)OwnerComp->FindTemplateNode(this);
-	OwnerComp->OnTaskFinished(TemplateNode, TaskResult);
+	UBTTaskNode* TemplateNode = (UBTTaskNode*)OwnerComp.FindTemplateNode(this);
+	OwnerComp.OnTaskFinished(TemplateNode, TaskResult);
 }
 
-void UBTTaskNode::FinishLatentAbort(UBehaviorTreeComponent* OwnerComp) const
+void UBTTaskNode::FinishLatentAbort(UBehaviorTreeComponent& OwnerComp) const
 {
 	// OnTaskFinished must receive valid template node
-	UBTTaskNode* TemplateNode = (UBTTaskNode*)OwnerComp->FindTemplateNode(this);
-	OwnerComp->OnTaskFinished(TemplateNode, EBTNodeResult::Aborted);
+	UBTTaskNode* TemplateNode = (UBTTaskNode*)OwnerComp.FindTemplateNode(this);
+	OwnerComp.OnTaskFinished(TemplateNode, EBTNodeResult::Aborted);
 }
 
-void UBTTaskNode::WaitForMessage(UBehaviorTreeComponent* OwnerComp, FName MessageType) const
+void UBTTaskNode::WaitForMessage(UBehaviorTreeComponent& OwnerComp, FName MessageType) const
 {
 	// messages delegates should be called on node instances (if they exists)
-	OwnerComp->RegisterMessageObserver(this, MessageType);
+	OwnerComp.RegisterMessageObserver(this, MessageType);
 }
 
-void UBTTaskNode::WaitForMessage(UBehaviorTreeComponent* OwnerComp, FName MessageType, int32 RequestID) const
+void UBTTaskNode::WaitForMessage(UBehaviorTreeComponent& OwnerComp, FName MessageType, int32 RequestID) const
 {
 	// messages delegates should be called on node instances (if they exists)
-	OwnerComp->RegisterMessageObserver(this, MessageType, RequestID);
+	OwnerComp.RegisterMessageObserver(this, MessageType, RequestID);
 }
 	
-void UBTTaskNode::StopWaitingForMessages(UBehaviorTreeComponent* OwnerComp) const
+void UBTTaskNode::StopWaitingForMessages(UBehaviorTreeComponent& OwnerComp) const
 {
 	// messages delegates should be called on node instances (if they exists)
-	OwnerComp->UnregisterMessageObserversFrom(this);
+	OwnerComp.UnregisterMessageObserversFrom(this);
 }
 
 #if WITH_EDITOR
diff --git a/Engine/Source/Runtime/AIModule/Private/BehaviorTree/BehaviorTreeComponent.cpp b/Engine/Source/Runtime/AIModule/Private/BehaviorTree/BehaviorTreeComponent.cpp
--- a/Engine/Source/Runtime/AIModule/Private/BehaviorTree/BehaviorTreeComponent.cpp
+++ b/Engine/Source/Runtime/AIModule/Private/BehaviorTree/BehaviorTreeComponent.cpp
@@ -13,7 +13,9 @@
 // UBehaviorTreeComponent
 //----------------------------------------------------------------------//
 
-UBehaviorTreeComponent::UBehaviorTreeComponent(const FObjectInitializer& ObjectInitializer) : Super(ObjectInitializer)
+UBehaviorTreeComponent::UBehaviorTreeComponent(const FObjectInitializer& ObjectInitializer) 
+	: Super(ObjectInitializer)
+	, SearchData(*this)
 {
 	ActiveInstanceIdx = 0;
 	bLoopExecution = false;
@@ -23,8 +25,6 @@ UBehaviorTreeComponent::UBehaviorTreeComponent(const FObjectInitializer& ObjectI
 	bWantsInitializeComponent = true; 
 	bIsRunning = false;
 	bIsPaused = false;
-	
-	SearchData.OwnerComp = this;
 }
 
 void UBehaviorTreeComponent::BeginDestroy()
@@ -32,7 +32,7 @@ void UBehaviorTreeComponent::BeginDestroy()
 	UBehaviorTreeManager* BTManager = UBehaviorTreeManager::GetCurrent(GetWorld());
 	if (BTManager)
 	{
-		BTManager->RemoveActiveComponent(this);
+		BTManager->RemoveActiveComponent(*this);
 	}
 
 	RemoveAllInstances();
@@ -134,7 +134,7 @@ bool UBehaviorTreeComponent::StartTree(UBehaviorTree& Asset, EBTExecutionMode::T
 	UBehaviorTreeManager* BTManager = UBehaviorTreeManager::GetCurrent(GetWorld());
 	if (BTManager)
 	{
-		BTManager->AddActiveComponent(this);
+		BTManager->AddActiveComponent(*this);
 	}
 
 	// push new instance
@@ -155,11 +155,11 @@ void UBehaviorTreeComponent::StopTree()
 			const UBTAuxiliaryNode* AuxNode = InstanceInfo.ActiveAuxNodes[AuxIndex];
 			check(AuxNode != NULL);
 			uint8* NodeMemory = AuxNode->GetNodeMemory<uint8>(InstanceInfo);
-			AuxNode->WrappedOnCeaseRelevant(this, NodeMemory);
+			AuxNode->WrappedOnCeaseRelevant(*this, NodeMemory);
 		}
 
 		// clear instance
-		InstanceStack[InstanceIndex].Cleanup(this, EBTMemoryClear::StoreSubtree);
+		InstanceStack[InstanceIndex].Cleanup(*this, EBTMemoryClear::StoreSubtree);
 	}
 
 	InstanceStack.Reset();
@@ -214,7 +214,7 @@ void UBehaviorTreeComponent::OnTaskFinished(const UBTTaskNode* TaskNode, EBTNode
 	uint8* ParentMemory = ParentNode->GetNodeMemory<uint8>(InstanceStack[TaskInstanceIdx]);
 
 	const bool bWasWaitingForAbort = bWaitingForAbortingTasks;
-	ParentNode->ConditionalNotifyChildExecution(this, ParentMemory, TaskNode, TaskResult);
+	ParentNode->ConditionalNotifyChildExecution(*this, ParentMemory, *TaskNode, TaskResult);
 	
 	if (TaskResult != EBTNodeResult::InProgress)
 	{
@@ -329,7 +329,7 @@ bool UBehaviorTreeComponent::IsAuxNodeActive(const UBTAuxiliaryNode* AuxNode) co
 			if (AuxNode->IsInstanced() && TestAuxNode && TestAuxNode->GetExecutionIndex() == AuxExecutionIndex)
 			{
 				const uint8* NodeMemory = TestAuxNode->GetNodeMemory<uint8>(InstanceInfo);
-				UBTNode* NodeInstance = TestAuxNode->GetNodeInstance(this, (uint8*)NodeMemory);
+				UBTNode* NodeInstance = TestAuxNode->GetNodeInstance(*this, (uint8*)NodeMemory);
 
 				if (NodeInstance == AuxNode)
 				{
@@ -563,7 +563,7 @@ void UBehaviorTreeComponent::RequestExecution(UBTCompositeNode* RequestedOn, int
 	{
 		// check if decorators allow execution on requesting link
 		const bool bShouldCheckDecorators = (RequestedByChildIndex >= 0);
-		const bool bCanExecute = !bShouldCheckDecorators || RequestedOn->DoDecoratorsAllowExecution(this, InstanceIdx, RequestedByChildIndex);
+		const bool bCanExecute = !bShouldCheckDecorators || RequestedOn->DoDecoratorsAllowExecution(*this, InstanceIdx, RequestedByChildIndex);
 		if (!bCanExecute)
 		{
 			UE_VLOG(GetOwner(), LogBehaviorTree, Log, TEXT("> skip: decorators are not allowing execution"));
@@ -596,8 +596,8 @@ void UBehaviorTreeComponent::RequestExecution(UBTCompositeNode* RequestedOn, int
 			for (UBTCompositeNode* It = RequestedOn; It && It != CommonParent;)
 			{
 				UBTCompositeNode* ParentNode = It->GetParentNode();
-				const int32 ChildIdx = ParentNode->GetChildIndex(It);
-				const bool bCanExecuteTest = ParentNode->DoDecoratorsAllowExecution(this, CommonInstanceIdx, ChildIdx);
+				const int32 ChildIdx = ParentNode->GetChildIndex(*It);
+				const bool bCanExecuteTest = ParentNode->DoDecoratorsAllowExecution(*this, CommonInstanceIdx, ChildIdx);
 				if (!bCanExecuteTest)
 				{
 					UE_VLOG(GetOwner(), LogBehaviorTree, Log, TEXT("> skip: decorators are not allowing execution"));
@@ -619,7 +619,7 @@ void UBehaviorTreeComponent::RequestExecution(UBTCompositeNode* RequestedOn, int
 			(RequestedOn->Children[RequestedByChildIndex].DecoratorOps.Num() > 0) &&
 			RequestedBy->IsA(UBTDecorator::StaticClass());
 
-		const bool bCanExecute = bShouldCheckDecorators && RequestedOn->DoDecoratorsAllowExecution(this, InstanceIdx, RequestedByChildIndex);
+		const bool bCanExecute = bShouldCheckDecorators && RequestedOn->DoDecoratorsAllowExecution(*this, InstanceIdx, RequestedByChildIndex);
 		if (bCanExecute)
 		{
 			UE_VLOG(GetOwner(), LogBehaviorTree, Log, TEXT("> skip: decorators are still allowing execution"));
@@ -706,12 +706,12 @@ void UBehaviorTreeComponent::ApplySearchUpdates(const TArray<FBehaviorTreeSearch
 			if (UpdateInfo.Mode == EBTNodeUpdateMode::Remove)
 			{
 				UpdateInstance.ActiveAuxNodes.RemoveSingleSwap(UpdateInfo.AuxNode);
-				UpdateInfo.AuxNode->WrappedOnCeaseRelevant(this, NodeMemory);
+				UpdateInfo.AuxNode->WrappedOnCeaseRelevant(*this, NodeMemory);
 			}
 			else
 			{
 				UpdateInstance.ActiveAuxNodes.Add(UpdateInfo.AuxNode);
-				UpdateInfo.AuxNode->WrappedOnBecomeRelevant(this, NodeMemory);
+				UpdateInfo.AuxNode->WrappedOnBecomeRelevant(*this, NodeMemory);
 			}
 		}
 		else if (UpdateInfo.TaskNode)
@@ -722,7 +722,7 @@ void UBehaviorTreeComponent::ApplySearchUpdates(const TArray<FBehaviorTreeSearch
 				UnregisterMessageObserversFrom(UpdateInfo.TaskNode);
 
 				uint8* NodeMemory = (uint8*)UpdateNode->GetNodeMemory<uint8>(UpdateInstance);
-				EBTNodeResult::Type NodeResult = UpdateInfo.TaskNode->AbortTask(this, NodeMemory);
+				EBTNodeResult::Type NodeResult = UpdateInfo.TaskNode->AbortTask(*this, NodeMemory);
 
 				UE_VLOG(GetOwner(), LogBehaviorTree, Log, TEXT("Parallel task aborted: %s (%s)"),
 					*UBehaviorTreeTypes::DescribeNodeHelper(UpdateInfo.TaskNode),
@@ -802,14 +802,14 @@ void UBehaviorTreeComponent::TickComponent(float DeltaTime, enum ELevelTick Tick
 		{
 			const UBTAuxiliaryNode* AuxNode = InstanceInfo.ActiveAuxNodes[AuxIndex];
 			uint8* NodeMemory = AuxNode->GetNodeMemory<uint8>(InstanceInfo);
-			AuxNode->WrappedTickNode(this, NodeMemory, DeltaTime);
+			AuxNode->WrappedTickNode(*this, NodeMemory, DeltaTime);
 		}
 
 		for (int32 TaskIndex = 0; TaskIndex < InstanceInfo.ParallelTasks.Num(); TaskIndex++)
 		{
 			const UBTTaskNode* ParallelTask = InstanceInfo.ParallelTasks[TaskIndex].TaskNode;
 			uint8* NodeMemory = ParallelTask->GetNodeMemory<uint8>(InstanceInfo);
-			ParallelTask->WrappedTickTask(this, NodeMemory, DeltaTime);
+			ParallelTask->WrappedTickTask(*this, NodeMemory, DeltaTime);
 		}
 	}
 
@@ -820,7 +820,7 @@ void UBehaviorTreeComponent::TickComponent(float DeltaTime, enum ELevelTick Tick
 	{
 		UBTTaskNode* ActiveTask = (UBTTaskNode*)ActiveInstance.ActiveNode;
 		uint8* NodeMemory = ActiveTask->GetNodeMemory<uint8>(ActiveInstance);
-		ActiveTask->WrappedTickTask(this, NodeMemory, DeltaTime);
+		ActiveTask->WrappedTickTask(*this, NodeMemory, DeltaTime);
 	}
 }
 
@@ -941,7 +941,7 @@ void UBehaviorTreeComponent::ProcessExecutionRequest()
 
 				if (TestNode)
 				{
-					TestNode->OnChildDeactivation(SearchData, ChildNode, NodeResult);
+					TestNode->OnChildDeactivation(SearchData, *ChildNode, NodeResult);
 				}
 			}
 			else
@@ -1020,7 +1020,7 @@ void UBehaviorTreeComponent::ProcessPendingExecution()
 	{
 		for (int32 InstanceIndex = ActiveInstanceIdx + 1; InstanceIndex < InstanceStack.Num(); InstanceIndex++)
 		{
-			InstanceStack[InstanceIndex].Cleanup(this, EBTMemoryClear::StoreSubtree);
+			InstanceStack[InstanceIndex].Cleanup(*this, EBTMemoryClear::StoreSubtree);
 		}
 
 		InstanceStack.SetNum(ActiveInstanceIdx + 1);
@@ -1054,7 +1054,7 @@ bool UBehaviorTreeComponent::DeactivateUpTo(UBTCompositeNode* Node, uint16 NodeI
 		UBTCompositeNode* NotifyParent = DeactivatedChild->GetParentNode();
 		if (NotifyParent)
 		{
-			NotifyParent->OnChildDeactivation(SearchData, DeactivatedChild, NodeResult);
+			NotifyParent->OnChildDeactivation(SearchData, *DeactivatedChild, NodeResult);
 
 			BT_SEARCHLOG(SearchData, Verbose, TEXT("Deactivate node: %s"), *UBehaviorTreeTypes::DescribeNodeHelper(DeactivatedChild));
 			StoreDebuggerSearchStep(DeactivatedChild, ActiveInstanceIdx, NodeResult);
@@ -1079,7 +1079,7 @@ bool UBehaviorTreeComponent::DeactivateUpTo(UBTCompositeNode* Node, uint16 NodeI
 			if (ActiveInstanceIdx == 0)
 			{
 				BT_SEARCHLOG(SearchData, Error, TEXT("Execution path does NOT contain common parent node, restarting tree! AI:%s"),
-					*GetNameSafe(SearchData.OwnerComp->GetOwner()));
+					*GetNameSafe(SearchData.OwnerComp.GetOwner()));
 
 				RestartTree();
 				return false;
@@ -1131,7 +1131,7 @@ void UBehaviorTreeComponent::ExecuteTask(UBTTaskNode* TaskNode)
 	uint16 InstanceIdx = ActiveInstanceIdx;
 
 	uint8* NodeMemory = (uint8*)(TaskNode->GetNodeMemory<uint8>(ActiveInstance));
-	EBTNodeResult::Type TaskResult = TaskNode->WrappedExecuteTask(this, NodeMemory);
+	EBTNodeResult::Type TaskResult = TaskNode->WrappedExecuteTask(*this, NodeMemory);
 
 	const UBTNode* ActiveNodeAfterExecution = GetActiveNode();
 	if (ActiveNodeAfterExecution == TaskNode)
@@ -1155,7 +1155,7 @@ void UBehaviorTreeComponent::AbortCurrentTask()
 
 	// abort task using current state of tree
 	uint8* NodeMemory = (uint8*)(ActiveTask->GetNodeMemory<uint8>(ActiveInstance));
-	EBTNodeResult::Type TaskResult = ActiveTask->WrappedAbortTask(this, NodeMemory);
+	EBTNodeResult::Type TaskResult = ActiveTask->WrappedAbortTask(*this, NodeMemory);
 
 	const UBTNode* ActiveNodeAfterAbort = GetActiveNode();
 	if (ActiveNodeAfterAbort == ActiveTask)
@@ -1320,9 +1320,9 @@ bool UBehaviorTreeComponent::PushInstance(UBehaviorTree& TreeAsset)
 	if (ActiveParent)
 	{
 		uint8* ParentMemory = GetNodeMemory((UBTNode*)ActiveParent, InstanceStack.Num() - 1);
-		int32 ChildIdx = ActiveParent->GetChildIndex(ActiveNode);
+		int32 ChildIdx = ActiveNode ? ActiveParent->GetChildIndex(*ActiveNode) : INDEX_NONE;
 
-		const bool bIsAllowed = ActiveParent->CanPushSubtree(this, ParentMemory, ChildIdx);
+		const bool bIsAllowed = ActiveParent->CanPushSubtree(*this, ParentMemory, ChildIdx);
 		if (!bIsAllowed)
 		{
 			UE_VLOG(GetOwner(), LogBehaviorTree, Warning, TEXT("Failed to execute tree %s: parent of active node does not allow it! (%s)"),
@@ -1354,8 +1354,8 @@ bool UBehaviorTreeComponent::PushInstance(UBehaviorTree& TreeAsset)
 		}
 
 		NewInstance.InstanceMemory = InstanceInfo.InstanceMemory;
-		NewInstance.Initialize(this, RootNode, NodeInstanceIndex, bFirstTime ? EBTMemoryInit::Initialize : EBTMemoryInit::RestoreSubtree);
-		NewInstance.InjectNodes(this, RootNode, NodeInstanceIndex);
+		NewInstance.Initialize(*this, *RootNode, NodeInstanceIndex, bFirstTime ? EBTMemoryInit::Initialize : EBTMemoryInit::RestoreSubtree);
+		NewInstance.InjectNodes(*this, *RootNode, NodeInstanceIndex);
 
 		INC_DWORD_STAT(STAT_AI_BehaviorTree_NumInstances);
 		InstanceStack.Push(NewInstance);
@@ -1368,10 +1368,10 @@ bool UBehaviorTreeComponent::PushInstance(UBehaviorTree& TreeAsset)
 			uint8* NodeMemory = (uint8*)ServiceNode->GetNodeMemory<uint8>(InstanceStack[ActiveInstanceIdx]);
 
 			InstanceStack[ActiveInstanceIdx].ActiveAuxNodes.Add(ServiceNode);
-			ServiceNode->WrappedOnBecomeRelevant(this, NodeMemory);
+			ServiceNode->WrappedOnBecomeRelevant(*this, NodeMemory);
 		}
 
-		FBehaviorTreeDelegates::OnTreeStarted.Broadcast(this, &TreeAsset);
+		FBehaviorTreeDelegates::OnTreeStarted.Broadcast(*this, TreeAsset);
 
 		// start new task
 		RequestExecution(RootNode, ActiveInstanceIdx, RootNode, 0, EBTNodeResult::InProgress);
@@ -1505,7 +1505,7 @@ void UBehaviorTreeComponent::RemoveAllInstances()
 		DummyInstance.InstanceIdIndex = Idx;
 		DummyInstance.RootNode = Info.RootNode;
 
-		DummyInstance.Cleanup(this, EBTMemoryClear::Destroy);
+		DummyInstance.Cleanup(*this, EBTMemoryClear::Destroy);
 
 		DummyInstance.InstanceMemory.Reset();
 	}
@@ -1551,7 +1551,7 @@ FString UBehaviorTreeComponent::GetDebugInfoString() const
 		while (Node)
 		{
 			uint8* NodeMemory = (uint8*)(Node->GetNodeMemory<uint8>(InstanceData));
-			NodeTrace = FString::Printf(TEXT("  %s\n"), *Node->GetRuntimeDescription(this, NodeMemory, EBTDescriptionVerbosity::Basic)) + NodeTrace;
+			NodeTrace = FString::Printf(TEXT("  %s\n"), *Node->GetRuntimeDescription(*this, NodeMemory, EBTDescriptionVerbosity::Basic)) + NodeTrace;
 			Node = Node->GetParentNode();
 		}
 
@@ -1634,7 +1634,7 @@ void UBehaviorTreeComponent::DescribeSelfToVisLog(FVisualLogEntry* Snapshot) con
 		while (Node)
 		{
 			uint8* NodeMemory = (uint8*)(Node->GetNodeMemory<uint8>(InstanceInfo));
-			Descriptions.Add(Node->GetRuntimeDescription(this, NodeMemory, EBTDescriptionVerbosity::Detailed));
+			Descriptions.Add(Node->GetRuntimeDescription(*this, NodeMemory, EBTDescriptionVerbosity::Detailed));
 		
 			Node = Node->GetParentNode();
 		}
@@ -1847,7 +1847,7 @@ void UBehaviorTreeComponent::StoreDebuggerRuntimeValues(TArray<FString>& Runtime
 		uint8* NodeMemory = (uint8*)Node->GetNodeMemory<uint8>(InstanceInfo);
 
 		RuntimeValues.Reset();
-		Node->DescribeRuntimeValues(this, NodeMemory, EBTDescriptionVerbosity::Basic, RuntimeValues);
+		Node->DescribeRuntimeValues(*this, NodeMemory, EBTDescriptionVerbosity::Basic, RuntimeValues);
 
 		FString ComposedDesc;
 		for (int32 ValueIndex = 0; ValueIndex < RuntimeValues.Num(); ValueIndex++)
@@ -1875,7 +1875,7 @@ void UBehaviorTreeComponent::UpdateDebuggerAfterExecution(const UBTTaskNode* Tas
 	TArray<FString> RuntimeValues;
 	const FBehaviorTreeInstance& InstanceToUpdate = InstanceStack[InstanceIdx];
 	uint8* NodeMemory = (uint8*)TaskNode->GetNodeMemory<uint8>(InstanceToUpdate);
-	TaskNode->DescribeRuntimeValues(this, NodeMemory, EBTDescriptionVerbosity::Basic, RuntimeValues);
+	TaskNode->DescribeRuntimeValues(*this, NodeMemory, EBTDescriptionVerbosity::Basic, RuntimeValues);
 
 	FString ComposedDesc;
 	for (int32 ValueIndex = 0; ValueIndex < RuntimeValues.Num(); ValueIndex++)
diff --git a/Engine/Source/Runtime/AIModule/Private/BehaviorTree/BehaviorTreeManager.cpp b/Engine/Source/Runtime/AIModule/Private/BehaviorTree/BehaviorTreeManager.cpp
--- a/Engine/Source/Runtime/AIModule/Private/BehaviorTree/BehaviorTreeManager.cpp
+++ b/Engine/Source/Runtime/AIModule/Private/BehaviorTree/BehaviorTreeManager.cpp
@@ -363,14 +363,14 @@ void UBehaviorTreeManager::DumpUsageStats() const
 	AllNodesCounter.Print(TEXT(","));
 }
 
-void UBehaviorTreeManager::AddActiveComponent(UBehaviorTreeComponent* Component)
+void UBehaviorTreeManager::AddActiveComponent(UBehaviorTreeComponent& Component)
 {
-	ActiveComponents.AddUnique(Component);
+	ActiveComponents.AddUnique(&Component);
 }
 
-void UBehaviorTreeManager::RemoveActiveComponent(UBehaviorTreeComponent* Component)
+void UBehaviorTreeManager::RemoveActiveComponent(UBehaviorTreeComponent& Component)
 {
-	ActiveComponents.Remove(Component);
+	ActiveComponents.Remove(&Component);
 }
 
 UBehaviorTreeManager* UBehaviorTreeManager::GetCurrent(UWorld* World)
diff --git a/Engine/Source/Runtime/AIModule/Private/BehaviorTree/BehaviorTreeTypes.cpp b/Engine/Source/Runtime/AIModule/Private/BehaviorTree/BehaviorTreeTypes.cpp
--- a/Engine/Source/Runtime/AIModule/Private/BehaviorTree/BehaviorTreeTypes.cpp
+++ b/Engine/Source/Runtime/AIModule/Private/BehaviorTree/BehaviorTreeTypes.cpp
@@ -9,23 +9,18 @@
 //----------------------------------------------------------------------//
 // FBehaviorTreeInstance
 //----------------------------------------------------------------------//
-void FBehaviorTreeInstance::Initialize(UBehaviorTreeComponent* OwnerComp, UBTCompositeNode* Node, int32& InstancedIndex, EBTMemoryInit::Type InitType)
+void FBehaviorTreeInstance::Initialize(UBehaviorTreeComponent& OwnerComp, UBTCompositeNode& Node, int32& InstancedIndex, EBTMemoryInit::Type InitType)
 {
-	if (Node == NULL)
+	for (int32 ServiceIndex = 0; ServiceIndex < Node.Services.Num(); ServiceIndex++)
 	{
-		return;
+		Node.Services[ServiceIndex]->InitializeInSubtree(OwnerComp, Node.Services[ServiceIndex]->GetNodeMemory<uint8>(*this), InstancedIndex, InitType);
 	}
 
-	for (int32 ServiceIndex = 0; ServiceIndex < Node->Services.Num(); ServiceIndex++)
-	{
-		Node->Services[ServiceIndex]->InitializeInSubtree(OwnerComp, Node->Services[ServiceIndex]->GetNodeMemory<uint8>(*this), InstancedIndex, InitType);
-	}
-
-	Node->InitializeInSubtree(OwnerComp, Node->GetNodeMemory<uint8>(*this), InstancedIndex, InitType);
+	Node.InitializeInSubtree(OwnerComp, Node.GetNodeMemory<uint8>(*this), InstancedIndex, InitType);
 
-	for (int32 ChildIndex = 0; ChildIndex < Node->Children.Num(); ChildIndex++)
+	for (int32 ChildIndex = 0; ChildIndex < Node.Children.Num(); ChildIndex++)
 	{
-		FBTCompositeChild& ChildInfo = Node->Children[ChildIndex];
+		FBTCompositeChild& ChildInfo = Node.Children[ChildIndex];
 
 		for (int32 DecoratorIndex = 0; DecoratorIndex < ChildInfo.Decorators.Num(); DecoratorIndex++)
 		{
@@ -34,7 +29,7 @@ void FBehaviorTreeInstance::Initialize(UBehaviorTreeComponent* OwnerComp, UBTCom
 
 		if (ChildInfo.ChildComposite)
 		{
-			Initialize(OwnerComp, ChildInfo.ChildComposite, InstancedIndex, InitType);
+			Initialize(OwnerComp, *(ChildInfo.ChildComposite), InstancedIndex, InitType);
 		}
 		else if (ChildInfo.ChildTask)
 		{
@@ -43,19 +38,14 @@ void FBehaviorTreeInstance::Initialize(UBehaviorTreeComponent* OwnerComp, UBTCom
 	}
 }
 
-void FBehaviorTreeInstance::InjectNodes(UBehaviorTreeComponent* OwnerComp, UBTCompositeNode* Node, int32& InstancedIndex)
+void FBehaviorTreeInstance::InjectNodes(UBehaviorTreeComponent& OwnerComp, UBTCompositeNode& Node, int32& InstancedIndex)
 {
-	if (Node == NULL)
+	for (int32 ChildIndex = 0; ChildIndex < Node.Children.Num(); ChildIndex++)
 	{
-		return;
-	}
-
-	for (int32 ChildIndex = 0; ChildIndex < Node->Children.Num(); ChildIndex++)
-	{
-		FBTCompositeChild& ChildInfo = Node->Children[ChildIndex];
+		FBTCompositeChild& ChildInfo = Node.Children[ChildIndex];
 		if (ChildInfo.ChildComposite)
 		{
-			InjectNodes(OwnerComp, ChildInfo.ChildComposite, InstancedIndex);
+			InjectNodes(OwnerComp, *(ChildInfo.ChildComposite), InstancedIndex);
 		}
 		else
 		{
@@ -69,43 +59,38 @@ void FBehaviorTreeInstance::InjectNodes(UBehaviorTreeComponent* OwnerComp, UBTCo
 	}
 }
 
-void FBehaviorTreeInstance::Cleanup(UBehaviorTreeComponent* OwnerComp, EBTMemoryClear::Type CleanupType)
+void FBehaviorTreeInstance::Cleanup(UBehaviorTreeComponent& OwnerComp, EBTMemoryClear::Type CleanupType)
 {
-	FBehaviorTreeInstanceId& Info = OwnerComp->KnownInstances[InstanceIdIndex];
+	FBehaviorTreeInstanceId& Info = OwnerComp.KnownInstances[InstanceIdIndex];
 	if (Info.FirstNodeInstance >= 0)
 	{
-		const int32 MaxAllowedIdx = OwnerComp->NodeInstances.Num();
-		const int32 LastNodeIdx = OwnerComp->KnownInstances.IsValidIndex(InstanceIdIndex + 1) ?
-			FMath::Min(OwnerComp->KnownInstances[InstanceIdIndex + 1].FirstNodeInstance, MaxAllowedIdx) :
+		const int32 MaxAllowedIdx = OwnerComp.NodeInstances.Num();
+		const int32 LastNodeIdx = OwnerComp.KnownInstances.IsValidIndex(InstanceIdIndex + 1) ?
+			FMath::Min(OwnerComp.KnownInstances[InstanceIdIndex + 1].FirstNodeInstance, MaxAllowedIdx) :
 			MaxAllowedIdx;
 
 		for (int32 Idx = Info.FirstNodeInstance; Idx < LastNodeIdx; Idx++)
 		{
-			OwnerComp->NodeInstances[Idx]->OnInstanceDestroyed(OwnerComp);
+			OwnerComp.NodeInstances[Idx]->OnInstanceDestroyed(OwnerComp);
 		}
 	}
 
-	CleanupNodes(OwnerComp, RootNode, CleanupType);
+	CleanupNodes(OwnerComp, *RootNode, CleanupType);
 	Info.InstanceMemory = InstanceMemory;
 }
 
-void FBehaviorTreeInstance::CleanupNodes(UBehaviorTreeComponent* OwnerComp, UBTCompositeNode* Node, EBTMemoryClear::Type CleanupType)
+void FBehaviorTreeInstance::CleanupNodes(UBehaviorTreeComponent& OwnerComp, UBTCompositeNode& Node, EBTMemoryClear::Type CleanupType)
 {
-	if (Node == NULL)
-	{
-		return;
-	}
-
-	for (int32 ServiceIndex = 0; ServiceIndex < Node->Services.Num(); ServiceIndex++)
+	for (int32 ServiceIndex = 0; ServiceIndex < Node.Services.Num(); ServiceIndex++)
 	{
-		Node->Services[ServiceIndex]->CleanupInSubtree(OwnerComp, Node->Services[ServiceIndex]->GetNodeMemory<uint8>(*this), CleanupType);
+		Node.Services[ServiceIndex]->CleanupInSubtree(OwnerComp, Node.Services[ServiceIndex]->GetNodeMemory<uint8>(*this), CleanupType);
 	}
 
-	Node->CleanupInSubtree(OwnerComp, Node->GetNodeMemory<uint8>(*this), CleanupType);
+	Node.CleanupInSubtree(OwnerComp, Node.GetNodeMemory<uint8>(*this), CleanupType);
 
-	for (int32 ChildIndex = 0; ChildIndex < Node->Children.Num(); ChildIndex++)
+	for (int32 ChildIndex = 0; ChildIndex < Node.Children.Num(); ChildIndex++)
 	{
-		FBTCompositeChild& ChildInfo = Node->Children[ChildIndex];
+		FBTCompositeChild& ChildInfo = Node.Children[ChildIndex];
 
 		for (int32 DecoratorIndex = 0; DecoratorIndex < ChildInfo.Decorators.Num(); DecoratorIndex++)
 		{
@@ -114,7 +99,7 @@ void FBehaviorTreeInstance::CleanupNodes(UBehaviorTreeComponent* OwnerComp, UBTC
 
 		if (ChildInfo.ChildComposite)
 		{
-			CleanupNodes(OwnerComp, ChildInfo.ChildComposite, CleanupType);
+			CleanupNodes(OwnerComp, *(ChildInfo.ChildComposite), CleanupType);
 		}
 		else if (ChildInfo.ChildTask)
 		{
@@ -148,7 +133,7 @@ int32 FBehaviorTreeSearchData::NextSearchId = 1;
 
 void FBehaviorTreeSearchData::AddUniqueUpdate(const FBehaviorTreeSearchUpdate& UpdateInfo)
 {
-	UE_VLOG(OwnerComp->GetOwner(), LogBehaviorTree, Verbose, TEXT("Search node update[%s]: %s"),
+	UE_VLOG(OwnerComp.GetOwner(), LogBehaviorTree, Verbose, TEXT("Search node update[%s]: %s"),
 		*UBehaviorTreeTypes::DescribeNodeUpdateMode(UpdateInfo.Mode),
 		*UBehaviorTreeTypes::DescribeNodeHelper(UpdateInfo.AuxNode ? (UBTNode*)UpdateInfo.AuxNode : (UBTNode*)UpdateInfo.TaskNode));
 
diff --git a/Engine/Source/Runtime/AIModule/Private/BehaviorTree/BlueprintNodeHelpers.cpp b/Engine/Source/Runtime/AIModule/Private/BehaviorTree/BlueprintNodeHelpers.cpp
--- a/Engine/Source/Runtime/AIModule/Private/BehaviorTree/BlueprintNodeHelpers.cpp
+++ b/Engine/Source/Runtime/AIModule/Private/BehaviorTree/BlueprintNodeHelpers.cpp
@@ -212,13 +212,16 @@ namespace BlueprintNodeHelpers
 		return bFound;
 	}
 
-	void AbortLatentActions(UActorComponent* OwnerOb, const UObject* Ob)
+	void AbortLatentActions(UActorComponent& OwnerOb, const UObject& Ob)
 	{
-		if (OwnerOb && !OwnerOb->HasAnyFlags(RF_BeginDestroyed) && OwnerOb->GetOwner())
+		if (!OwnerOb.HasAnyFlags(RF_BeginDestroyed) && OwnerOb.GetOwner())
 		{
-			UWorld* MyWorld = OwnerOb->GetOwner()->GetWorld();
-			MyWorld->GetLatentActionManager().RemoveActionsForObject(Ob);
-			MyWorld->GetTimerManager().ClearAllTimersForObject(Ob);
+			UWorld* MyWorld = OwnerOb.GetOwner()->GetWorld();
+			if (MyWorld)
+			{
+				MyWorld->GetLatentActionManager().RemoveActionsForObject(&Ob);
+				MyWorld->GetTimerManager().ClearAllTimersForObject(&Ob);
+			}
 		}
 	}
 }
diff --git a/Engine/Source/Runtime/AIModule/Private/BehaviorTree/Composites/BTComposite_SimpleParallel.cpp b/Engine/Source/Runtime/AIModule/Private/BehaviorTree/Composites/BTComposite_SimpleParallel.cpp
--- a/Engine/Source/Runtime/AIModule/Private/BehaviorTree/Composites/BTComposite_SimpleParallel.cpp
+++ b/Engine/Source/Runtime/AIModule/Private/BehaviorTree/Composites/BTComposite_SimpleParallel.cpp
@@ -20,7 +20,7 @@ int32 UBTComposite_SimpleParallel::GetNextChildHandler(FBehaviorTreeSearchData&
 	{
 		NextChildIdx = EBTParallelChild::MainTask;
 	}
-	else if ((MyMemory->bMainTaskIsActive || MyMemory->bForceBackgroundTree) && !SearchData.OwnerComp->IsRestartPending())
+	else if ((MyMemory->bMainTaskIsActive || MyMemory->bForceBackgroundTree) && !SearchData.OwnerComp.IsRestartPending())
 	{
 		// if main task is running - always pick background tree
 		// unless search is from abort from background tree - return to parent (and break from search when node gets deactivated)
@@ -31,14 +31,14 @@ int32 UBTComposite_SimpleParallel::GetNextChildHandler(FBehaviorTreeSearchData&
 	return NextChildIdx;
 }
 
-void UBTComposite_SimpleParallel::NotifyChildExecution(UBehaviorTreeComponent* OwnerComp, uint8* NodeMemory, int32 ChildIdx, EBTNodeResult::Type& NodeResult) const
+void UBTComposite_SimpleParallel::NotifyChildExecution(UBehaviorTreeComponent& OwnerComp, uint8* NodeMemory, int32 ChildIdx, EBTNodeResult::Type& NodeResult) const
 {
 	FBTParallelMemory* MyMemory = (FBTParallelMemory*)NodeMemory;
 	if (ChildIdx == EBTParallelChild::MainTask)
 	{
 		if (NodeResult == EBTNodeResult::InProgress)
 		{
-			EBTTaskStatus::Type Status = OwnerComp->GetTaskStatus(Children[EBTParallelChild::MainTask].ChildTask);
+			EBTTaskStatus::Type Status = OwnerComp.GetTaskStatus(Children[EBTParallelChild::MainTask].ChildTask);
 			if (Status == EBTTaskStatus::Active)
 			{
 				// can't register if task is currently being aborted (latent abort returns in progress)
@@ -46,7 +46,7 @@ void UBTComposite_SimpleParallel::NotifyChildExecution(UBehaviorTreeComponent* O
 				MyMemory->bMainTaskIsActive = true;
 				MyMemory->bForceBackgroundTree = false;
 				
-				OwnerComp->RegisterParallelTask(Children[EBTParallelChild::MainTask].ChildTask);
+				OwnerComp.RegisterParallelTask(Children[EBTParallelChild::MainTask].ChildTask);
 				RequestDelayedExecution(OwnerComp, EBTNodeResult::Succeeded);
 			}
 		}
@@ -55,15 +55,15 @@ void UBTComposite_SimpleParallel::NotifyChildExecution(UBehaviorTreeComponent* O
 			MyMemory->MainTaskResult = NodeResult;
 			MyMemory->bMainTaskIsActive = false;
 
-			const int32 MyInstanceIdx = OwnerComp->FindInstanceContainingNode(this);
+			const int32 MyInstanceIdx = OwnerComp.FindInstanceContainingNode(this);
 
-			OwnerComp->UnregisterParallelTask(Children[EBTParallelChild::MainTask].ChildTask, MyInstanceIdx);
+			OwnerComp.UnregisterParallelTask(Children[EBTParallelChild::MainTask].ChildTask, MyInstanceIdx);
 			if (NodeResult != EBTNodeResult::Aborted)
 			{
 				// check if subtree should be aborted when task finished with success/failed result
 				if (FinishMode == EBTParallelMode::AbortBackground)
 				{
-					OwnerComp->RequestExecution((UBTCompositeNode*)this, MyInstanceIdx,
+					OwnerComp.RequestExecution((UBTCompositeNode*)this, MyInstanceIdx,
 						Children[EBTParallelChild::MainTask].ChildTask, EBTParallelChild::MainTask,
 						NodeResult);
 				}
@@ -84,7 +84,7 @@ void UBTComposite_SimpleParallel::NotifyChildExecution(UBehaviorTreeComponent* O
 void UBTComposite_SimpleParallel::NotifyNodeDeactivation(FBehaviorTreeSearchData& SearchData, EBTNodeResult::Type& NodeResult) const
 {
 	FBTParallelMemory* MyMemory = GetNodeMemory<FBTParallelMemory>(SearchData);
-	const uint16 ActiveInstanceIdx = SearchData.OwnerComp->GetActiveInstanceIdx();
+	const uint16 ActiveInstanceIdx = SearchData.OwnerComp.GetActiveInstanceIdx();
 
 	// modify node result if main task has finished
 	if (!MyMemory->bMainTaskIsActive)
@@ -101,7 +101,7 @@ void UBTComposite_SimpleParallel::NotifyNodeDeactivation(FBehaviorTreeSearchData
 	// remove all active nodes from background tree
 	const FBTNodeIndex FirstBackgroundIndex(ActiveInstanceIdx, GetChildExecutionIndex(EBTParallelChild::BackgroundTree, EBTChildIndex::FirstNode));
 	const FBTNodeIndex LastBackgroundIndex(ActiveInstanceIdx, GetLastExecutionIndex());
-	SearchData.OwnerComp->UnregisterAuxNodesUpTo(FirstBackgroundIndex);
+	SearchData.OwnerComp.UnregisterAuxNodesUpTo(FirstBackgroundIndex);
 
 	// remove all pending updates "AddForLowerPri" from background tree 
 	// it doesn't make sense for decorators to reactivate themselves there
@@ -114,7 +114,7 @@ void UBTComposite_SimpleParallel::NotifyNodeDeactivation(FBehaviorTreeSearchData
 			const FBTNodeIndex UpdateIdx(UpdateInfo.InstanceIndex, UpdateNodeIdx);
 			if (FirstBackgroundIndex.TakesPriorityOver(UpdateIdx) && UpdateIdx.TakesPriorityOver(LastBackgroundIndex))
 			{
-				UE_VLOG(SearchData.OwnerComp->GetOwner(), LogBehaviorTree, Verbose, TEXT("Search node update[canceled]: %s"),
+				UE_VLOG(SearchData.OwnerComp.GetOwner(), LogBehaviorTree, Verbose, TEXT("Search node update[canceled]: %s"),
 					*UBehaviorTreeTypes::DescribeNodeHelper(UpdateInfo.AuxNode ? (UBTNode*)UpdateInfo.AuxNode : (UBTNode*)UpdateInfo.TaskNode));
 
 				SearchData.PendingUpdates.RemoveAt(Idx);
@@ -128,7 +128,7 @@ uint16 UBTComposite_SimpleParallel::GetInstanceMemorySize() const
 	return sizeof(FBTParallelMemory);
 }
 
-bool UBTComposite_SimpleParallel::CanPushSubtree(UBehaviorTreeComponent* OwnerComp, uint8* NodeMemory, int32 ChildIdx) const
+bool UBTComposite_SimpleParallel::CanPushSubtree(UBehaviorTreeComponent& OwnerComp, uint8* NodeMemory, int32 ChildIdx) const
 {
 	return (ChildIdx != EBTParallelChild::MainTask);
 }
@@ -147,7 +147,7 @@ FString UBTComposite_SimpleParallel::GetStaticDescription() const
 		FinishMode < 2 ? *FinishDesc[FinishMode] : TEXT(""));
 }
 
-void UBTComposite_SimpleParallel::DescribeRuntimeValues(const UBehaviorTreeComponent* OwnerComp, uint8* NodeMemory, EBTDescriptionVerbosity::Type Verbosity, TArray<FString>& Values) const
+void UBTComposite_SimpleParallel::DescribeRuntimeValues(const UBehaviorTreeComponent& OwnerComp, uint8* NodeMemory, EBTDescriptionVerbosity::Type Verbosity, TArray<FString>& Values) const
 {
 	Super::DescribeRuntimeValues(OwnerComp, NodeMemory, Verbosity, Values);
 
diff --git a/Engine/Source/Runtime/AIModule/Private/BehaviorTree/Decorators/BTDecorator_Blackboard.cpp b/Engine/Source/Runtime/AIModule/Private/BehaviorTree/Decorators/BTDecorator_Blackboard.cpp
--- a/Engine/Source/Runtime/AIModule/Private/BehaviorTree/Decorators/BTDecorator_Blackboard.cpp
+++ b/Engine/Source/Runtime/AIModule/Private/BehaviorTree/Decorators/BTDecorator_Blackboard.cpp
@@ -11,9 +11,9 @@ UBTDecorator_Blackboard::UBTDecorator_Blackboard(const FObjectInitializer& Objec
 	NotifyObserver = EBTBlackboardRestart::ResultChange;
 }
 
-bool UBTDecorator_Blackboard::CalculateRawConditionValue(UBehaviorTreeComponent* OwnerComp, uint8* NodeMemory) const
+bool UBTDecorator_Blackboard::CalculateRawConditionValue(UBehaviorTreeComponent& OwnerComp, uint8* NodeMemory) const
 {
-	const UBlackboardComponent* BlackboardComp = OwnerComp->GetBlackboardComponent();
+	const UBlackboardComponent* BlackboardComp = OwnerComp.GetBlackboardComponent();
 	// note that this may produce unexpected logical results. FALSE is a valid return value here as well
 	// @todo signal it
 	return BlackboardComp && EvaluateOnBlackboard(*BlackboardComp);
@@ -91,11 +91,11 @@ void UBTDecorator_Blackboard::OnBlackboardChange(const UBlackboardComponent& Bla
 	}
 }
 
-void UBTDecorator_Blackboard::DescribeRuntimeValues(const UBehaviorTreeComponent* OwnerComp, uint8* NodeMemory, EBTDescriptionVerbosity::Type Verbosity, TArray<FString>& Values) const
+void UBTDecorator_Blackboard::DescribeRuntimeValues(const UBehaviorTreeComponent& OwnerComp, uint8* NodeMemory, EBTDescriptionVerbosity::Type Verbosity, TArray<FString>& Values) const
 {
 	Super::DescribeRuntimeValues(OwnerComp, NodeMemory, Verbosity, Values);
 
-	const UBlackboardComponent* BlackboardComp = OwnerComp->GetBlackboardComponent();
+	const UBlackboardComponent* BlackboardComp = OwnerComp.GetBlackboardComponent();
 	FString DescKeyValue;
 
 	if (BlackboardComp)
diff --git a/Engine/Source/Runtime/AIModule/Private/BehaviorTree/Decorators/BTDecorator_BlackboardBase.cpp b/Engine/Source/Runtime/AIModule/Private/BehaviorTree/Decorators/BTDecorator_BlackboardBase.cpp
--- a/Engine/Source/Runtime/AIModule/Private/BehaviorTree/Decorators/BTDecorator_BlackboardBase.cpp
+++ b/Engine/Source/Runtime/AIModule/Private/BehaviorTree/Decorators/BTDecorator_BlackboardBase.cpp
@@ -26,18 +26,18 @@ void UBTDecorator_BlackboardBase::InitializeFromAsset(UBehaviorTree& Asset)
 	BlackboardKey.CacheSelectedKey(GetBlackboardAsset());
 }
 
-void UBTDecorator_BlackboardBase::OnBecomeRelevant(UBehaviorTreeComponent* OwnerComp, uint8* NodeMemory)
+void UBTDecorator_BlackboardBase::OnBecomeRelevant(UBehaviorTreeComponent& OwnerComp, uint8* NodeMemory)
 {
-	UBlackboardComponent* BlackboardComp = OwnerComp->GetBlackboardComponent();
+	UBlackboardComponent* BlackboardComp = OwnerComp.GetBlackboardComponent();
 	if (BlackboardComp)
 	{
 		BlackboardComp->RegisterObserver(BlackboardKey.GetSelectedKeyID(), BBKeyObserver);
 	}
 }
 
-void UBTDecorator_BlackboardBase::OnCeaseRelevant(UBehaviorTreeComponent* OwnerComp, uint8* NodeMemory)
+void UBTDecorator_BlackboardBase::OnCeaseRelevant(UBehaviorTreeComponent& OwnerComp, uint8* NodeMemory)
 {
-	UBlackboardComponent* BlackboardComp = OwnerComp->GetBlackboardComponent();
+	UBlackboardComponent* BlackboardComp = OwnerComp.GetBlackboardComponent();
 	if (BlackboardComp)
 	{
 		BlackboardComp->UnregisterObserver(BlackboardKey.GetSelectedKeyID(), BBKeyObserver);
diff --git a/Engine/Source/Runtime/AIModule/Private/BehaviorTree/Decorators/BTDecorator_BlueprintBase.cpp b/Engine/Source/Runtime/AIModule/Private/BehaviorTree/Decorators/BTDecorator_BlueprintBase.cpp
--- a/Engine/Source/Runtime/AIModule/Private/BehaviorTree/Decorators/BTDecorator_BlueprintBase.cpp
+++ b/Engine/Source/Runtime/AIModule/Private/BehaviorTree/Decorators/BTDecorator_BlueprintBase.cpp
@@ -66,7 +66,7 @@ void UBTDecorator_BlueprintBase::SetOwner(AActor* InActorOwner)
 	AIOwner = Cast<AAIController>(InActorOwner);
 }
 
-void UBTDecorator_BlueprintBase::OnBecomeRelevant(UBehaviorTreeComponent* OwnerComp, uint8* NodeMemory)
+void UBTDecorator_BlueprintBase::OnBecomeRelevant(UBehaviorTreeComponent& OwnerComp, uint8* NodeMemory)
 {
 	if (AIOwner != nullptr && ReceiveObserverActivatedImplementations & FBTNodeBPImplementationHelper::AISpecific)
 	{
@@ -77,7 +77,7 @@ void UBTDecorator_BlueprintBase::OnBecomeRelevant(UBehaviorTreeComponent* OwnerC
 		ReceiveObserverActivated(ActorOwner);
 	}
 
-	UBlackboardComponent* BlackboardComp = OwnerComp->GetBlackboardComponent();
+	UBlackboardComponent* BlackboardComp = OwnerComp.GetBlackboardComponent();
 	UBTDecorator_BlueprintBase* CDO = (UBTDecorator_BlueprintBase*)(GetClass()->GetDefaultObject());
 	if (BlackboardComp && CDO)
 	{
@@ -92,9 +92,9 @@ void UBTDecorator_BlueprintBase::OnBecomeRelevant(UBehaviorTreeComponent* OwnerC
 	}
 }
 
-void UBTDecorator_BlueprintBase::OnCeaseRelevant(UBehaviorTreeComponent* OwnerComp, uint8* NodeMemory)
+void UBTDecorator_BlueprintBase::OnCeaseRelevant(UBehaviorTreeComponent& OwnerComp, uint8* NodeMemory)
 {
-	UBlackboardComponent* BlackboardComp = OwnerComp->GetBlackboardComponent();
+	UBlackboardComponent* BlackboardComp = OwnerComp.GetBlackboardComponent();
 	UBTDecorator_BlueprintBase* CDO = (UBTDecorator_BlueprintBase*)(GetClass()->GetDefaultObject());
 	if (BlackboardComp && CDO)
 	{
@@ -142,7 +142,7 @@ void UBTDecorator_BlueprintBase::OnNodeDeactivation(FBehaviorTreeSearchData& Sea
 	}
 }
 
-void UBTDecorator_BlueprintBase::TickNode(UBehaviorTreeComponent* OwnerComp, uint8* NodeMemory, float DeltaSeconds)
+void UBTDecorator_BlueprintBase::TickNode(UBehaviorTreeComponent& OwnerComp, uint8* NodeMemory, float DeltaSeconds)
 {
 	if (AIOwner != nullptr && ReceiveTickImplementations & FBTNodeBPImplementationHelper::AISpecific)
 	{
@@ -154,7 +154,7 @@ void UBTDecorator_BlueprintBase::TickNode(UBehaviorTreeComponent* OwnerComp, uin
 	}
 }
 
-bool UBTDecorator_BlueprintBase::CalculateRawConditionValue(UBehaviorTreeComponent* OwnerComp, uint8* NodeMemory) const
+bool UBTDecorator_BlueprintBase::CalculateRawConditionValue(UBehaviorTreeComponent& OwnerComp, uint8* NodeMemory) const
 {
 	CurrentCallResult = false;
 	if (ReceiveConditionCheckImplementations != 0)
@@ -183,15 +183,13 @@ void UBTDecorator_BlueprintBase::FinishConditionCheck(bool bAllowExecution)
 bool UBTDecorator_BlueprintBase::IsDecoratorExecutionActive() const
 {
 	UBehaviorTreeComponent* OwnerComp = Cast<UBehaviorTreeComponent>(GetOuter());
-	const bool bIsActive = OwnerComp->IsExecutingBranch(GetMyNode(), GetChildIndex());
-	return bIsActive;
+	return OwnerComp && OwnerComp->IsExecutingBranch(GetMyNode(), GetChildIndex());
 }
 
 bool UBTDecorator_BlueprintBase::IsDecoratorObserverActive() const
 {
 	UBehaviorTreeComponent* OwnerComp = Cast<UBehaviorTreeComponent>(GetOuter());
-	const bool bIsActive = OwnerComp->IsAuxNodeActive(this);
-	return bIsActive;
+	return OwnerComp && OwnerComp->IsAuxNodeActive(this);
 }
 
 FString UBTDecorator_BlueprintBase::GetStaticDescription() const
@@ -213,7 +211,7 @@ FString UBTDecorator_BlueprintBase::GetStaticDescription() const
 	return ReturnDesc;
 }
 
-void UBTDecorator_BlueprintBase::DescribeRuntimeValues(const UBehaviorTreeComponent* OwnerComp, uint8* NodeMemory, EBTDescriptionVerbosity::Type Verbosity, TArray<FString>& Values) const
+void UBTDecorator_BlueprintBase::DescribeRuntimeValues(const UBehaviorTreeComponent& OwnerComp, uint8* NodeMemory, EBTDescriptionVerbosity::Type Verbosity, TArray<FString>& Values) const
 {
 	UBTDecorator_BlueprintBase* CDO = (UBTDecorator_BlueprintBase*)(GetClass()->GetDefaultObject());
 	if (CDO && CDO->PropertyData.Num())
@@ -231,10 +229,10 @@ void UBTDecorator_BlueprintBase::OnBlackboardChange(const UBlackboardComponent&
 	}
 }
 
-void UBTDecorator_BlueprintBase::OnInstanceDestroyed(UBehaviorTreeComponent* OwnerComp)
+void UBTDecorator_BlueprintBase::OnInstanceDestroyed(UBehaviorTreeComponent& OwnerComp)
 {
 	// force dropping all pending latent actions associated with this blueprint
-	BlueprintNodeHelpers::AbortLatentActions(OwnerComp, this);
+	BlueprintNodeHelpers::AbortLatentActions(OwnerComp, *this);
 }
 
 #if WITH_EDITOR
diff --git a/Engine/Source/Runtime/AIModule/Private/BehaviorTree/Decorators/BTDecorator_CheckGameplayTagsOnActor.cpp b/Engine/Source/Runtime/AIModule/Private/BehaviorTree/Decorators/BTDecorator_CheckGameplayTagsOnActor.cpp
--- a/Engine/Source/Runtime/AIModule/Private/BehaviorTree/Decorators/BTDecorator_CheckGameplayTagsOnActor.cpp
+++ b/Engine/Source/Runtime/AIModule/Private/BehaviorTree/Decorators/BTDecorator_CheckGameplayTagsOnActor.cpp
@@ -21,10 +21,9 @@ UBTDecorator_CheckGameplayTagsOnActor::UBTDecorator_CheckGameplayTagsOnActor(con
 	bAllowAbortChildNodes = false;
 }
 
-bool UBTDecorator_CheckGameplayTagsOnActor::CalculateRawConditionValue(UBehaviorTreeComponent* OwnerComp, uint8* NodeMemory) const
+bool UBTDecorator_CheckGameplayTagsOnActor::CalculateRawConditionValue(UBehaviorTreeComponent& OwnerComp, uint8* NodeMemory) const
 {
-	check(OwnerComp != NULL);
-	const UBlackboardComponent* BlackboardComp = OwnerComp->GetBlackboardComponent();
+	const UBlackboardComponent* BlackboardComp = OwnerComp.GetBlackboardComponent();
 	if (BlackboardComp == NULL)
 	{
 		return false;
@@ -52,7 +51,7 @@ bool UBTDecorator_CheckGameplayTagsOnActor::CalculateRawConditionValue(UBehavior
 	}
 }
 
-void UBTDecorator_CheckGameplayTagsOnActor::DescribeRuntimeValues(const UBehaviorTreeComponent* OwnerComp, uint8* NodeMemory, EBTDescriptionVerbosity::Type Verbosity, TArray<FString>& Values) const
+void UBTDecorator_CheckGameplayTagsOnActor::DescribeRuntimeValues(const UBehaviorTreeComponent& OwnerComp, uint8* NodeMemory, EBTDescriptionVerbosity::Type Verbosity, TArray<FString>& Values) const
 {
 	Super::DescribeRuntimeValues(OwnerComp, NodeMemory, Verbosity, Values);
 }
diff --git a/Engine/Source/Runtime/AIModule/Private/BehaviorTree/Decorators/BTDecorator_CompareBBEntries.cpp b/Engine/Source/Runtime/AIModule/Private/BehaviorTree/Decorators/BTDecorator_CompareBBEntries.cpp
--- a/Engine/Source/Runtime/AIModule/Private/BehaviorTree/Decorators/BTDecorator_CompareBBEntries.cpp
+++ b/Engine/Source/Runtime/AIModule/Private/BehaviorTree/Decorators/BTDecorator_CompareBBEntries.cpp
@@ -30,7 +30,7 @@ void UBTDecorator_CompareBBEntries::InitializeFromAsset(UBehaviorTree& Asset)
 
 // @note I know it's ugly to have "return" statements in many places inside a function, but the way 
 // around was very awkward here 
-bool UBTDecorator_CompareBBEntries::CalculateRawConditionValue(UBehaviorTreeComponent* OwnerComp, uint8* NodeMemory) const 
+bool UBTDecorator_CompareBBEntries::CalculateRawConditionValue(UBehaviorTreeComponent& OwnerComp, uint8* NodeMemory) const
 {
 	// first of all require same type
 	// @todo this could be checked statically (i.e. in editor, asset creation time)!
@@ -39,7 +39,7 @@ bool UBTDecorator_CompareBBEntries::CalculateRawConditionValue(UBehaviorTreeComp
 		return false;
 	}
 	
-	const UBlackboardComponent* BlackboardComp = OwnerComp->GetBlackboardComponent();
+	const UBlackboardComponent* BlackboardComp = OwnerComp.GetBlackboardComponent();
 	if (BlackboardComp)
 	{
 		//BlackboardComp->GetKeyType
@@ -65,9 +65,9 @@ FString UBTDecorator_CompareBBEntries::GetStaticDescription() const
 		, Operator == EBlackBoardEntryComparison::Equal ? TEXT("EQUAL") : TEXT("NOT EQUAL"));
 }
 
-void UBTDecorator_CompareBBEntries::OnBecomeRelevant(UBehaviorTreeComponent* OwnerComp, uint8* NodeMemory)
+void UBTDecorator_CompareBBEntries::OnBecomeRelevant(UBehaviorTreeComponent& OwnerComp, uint8* NodeMemory)
 {
-	UBlackboardComponent* BlackboardComp = OwnerComp->GetBlackboardComponent();
+	UBlackboardComponent* BlackboardComp = OwnerComp.GetBlackboardComponent();
 	if (BlackboardComp)
 	{
 		BlackboardComp->RegisterObserver(BlackboardKeyA.GetSelectedKeyID(), BBKeyObserver);
@@ -75,9 +75,9 @@ void UBTDecorator_CompareBBEntries::OnBecomeRelevant(UBehaviorTreeComponent* Own
 	}
 }
 
-void UBTDecorator_CompareBBEntries::OnCeaseRelevant(UBehaviorTreeComponent* OwnerComp, uint8* NodeMemory)
+void UBTDecorator_CompareBBEntries::OnCeaseRelevant(UBehaviorTreeComponent& OwnerComp, uint8* NodeMemory)
 {
-	UBlackboardComponent* BlackboardComp = OwnerComp->GetBlackboardComponent();
+	UBlackboardComponent* BlackboardComp = OwnerComp.GetBlackboardComponent();
 	if (BlackboardComp)
 	{
 		BlackboardComp->UnregisterObserver(BlackboardKeyA.GetSelectedKeyID(), BBKeyObserver);
@@ -87,7 +87,7 @@ void UBTDecorator_CompareBBEntries::OnCeaseRelevant(UBehaviorTreeComponent* Owne
 
 void UBTDecorator_CompareBBEntries::OnBlackboardChange(const UBlackboardComponent& Blackboard, FBlackboard::FKey ChangedKeyID)
 {
-	UBehaviorTreeComponent* BehaviorComp = (UBehaviorTreeComponent*)Blackboard.GetBrainComponent();
+	UBehaviorTreeComponent* BehaviorComp = static_cast<UBehaviorTreeComponent*>(Blackboard.GetBrainComponent());
 	if (BehaviorComp && (BlackboardKeyA.GetSelectedKeyID() == ChangedKeyID || BlackboardKeyB.GetSelectedKeyID() == ChangedKeyID))
 	{
 		BehaviorComp->RequestExecution(this);		
diff --git a/Engine/Source/Runtime/AIModule/Private/BehaviorTree/Decorators/BTDecorator_ConditionalLoop.cpp b/Engine/Source/Runtime/AIModule/Private/BehaviorTree/Decorators/BTDecorator_ConditionalLoop.cpp
--- a/Engine/Source/Runtime/AIModule/Private/BehaviorTree/Decorators/BTDecorator_ConditionalLoop.cpp
+++ b/Engine/Source/Runtime/AIModule/Private/BehaviorTree/Decorators/BTDecorator_ConditionalLoop.cpp
@@ -14,7 +14,7 @@ UBTDecorator_ConditionalLoop::UBTDecorator_ConditionalLoop(const FObjectInitiali
 	bAllowAbortChildNodes = false;
 }
 
-bool UBTDecorator_ConditionalLoop::CalculateRawConditionValue(UBehaviorTreeComponent* OwnerComp, uint8* NodeMemory) const
+bool UBTDecorator_ConditionalLoop::CalculateRawConditionValue(UBehaviorTreeComponent& OwnerComp, uint8* NodeMemory) const
 {
 	// always allows execution
 	return true;
@@ -27,11 +27,11 @@ void UBTDecorator_ConditionalLoop::OnBlackboardChange(const UBlackboardComponent
 
 void UBTDecorator_ConditionalLoop::OnNodeDeactivation(FBehaviorTreeSearchData& SearchData, EBTNodeResult::Type NodeResult)
 {
-	if (NodeResult != EBTNodeResult::Aborted && SearchData.OwnerComp)
+	if (NodeResult != EBTNodeResult::Aborted)
 	{
-		const UBlackboardComponent* BlackboardComp = SearchData.OwnerComp->GetBlackboardComponent();
+		const UBlackboardComponent* BlackboardComp = SearchData.OwnerComp.GetBlackboardComponent();
 		const bool bEvalResult = BlackboardComp && EvaluateOnBlackboard(*BlackboardComp);
-		UE_VLOG(SearchData.OwnerComp->GetOwner(), LogBehaviorTree, Verbose, TEXT("Loop condition: %s -> %s"),
+		UE_VLOG(SearchData.OwnerComp.GetOwner(), LogBehaviorTree, Verbose, TEXT("Loop condition: %s -> %s"),
 			bEvalResult ? TEXT("true") : TEXT("false"), (bEvalResult != IsInversed()) ? TEXT("run again!") : TEXT("break"));
 
 		if (bEvalResult != IsInversed())
diff --git a/Engine/Source/Runtime/AIModule/Private/BehaviorTree/Decorators/BTDecorator_ConeCheck.cpp b/Engine/Source/Runtime/AIModule/Private/BehaviorTree/Decorators/BTDecorator_ConeCheck.cpp
--- a/Engine/Source/Runtime/AIModule/Private/BehaviorTree/Decorators/BTDecorator_ConeCheck.cpp
+++ b/Engine/Source/Runtime/AIModule/Private/BehaviorTree/Decorators/BTDecorator_ConeCheck.cpp
@@ -64,9 +64,9 @@ bool UBTDecorator_ConeCheck::CalculateDirection(const UBlackboardComponent* Blac
 	return false;
 }
 
-FORCEINLINE bool UBTDecorator_ConeCheck::CalcConditionImpl(UBehaviorTreeComponent* OwnerComp, uint8* NodeMemory) const
+FORCEINLINE bool UBTDecorator_ConeCheck::CalcConditionImpl(UBehaviorTreeComponent& OwnerComp, uint8* NodeMemory) const
 {
-	const UBlackboardComponent* BBComponent = OwnerComp->GetBlackboardComponent();
+	const UBlackboardComponent* BBComponent = OwnerComp.GetBlackboardComponent();
 
 	FVector ConeDir;
 	FVector DirectionToObserve;
@@ -76,12 +76,12 @@ FORCEINLINE bool UBTDecorator_ConeCheck::CalcConditionImpl(UBehaviorTreeComponen
 		&& ConeDir.CosineAngle2D(DirectionToObserve) > ConeHalfAngleDot;
 }
 
-bool UBTDecorator_ConeCheck::CalculateRawConditionValue(UBehaviorTreeComponent* OwnerComp, uint8* NodeMemory) const
+bool UBTDecorator_ConeCheck::CalculateRawConditionValue(UBehaviorTreeComponent& OwnerComp, uint8* NodeMemory) const
 {
 	return CalcConditionImpl(OwnerComp, NodeMemory);
 }
 
-void UBTDecorator_ConeCheck::OnBecomeRelevant(UBehaviorTreeComponent* OwnerComp, uint8* NodeMemory)
+void UBTDecorator_ConeCheck::OnBecomeRelevant(UBehaviorTreeComponent& OwnerComp, uint8* NodeMemory)
 {
 	TNodeInstanceMemory* DecoratorMemory = (TNodeInstanceMemory*)NodeMemory;
 	DecoratorMemory->bLastRawResult = CalcConditionImpl(OwnerComp, NodeMemory);
@@ -92,13 +92,13 @@ void UBTDecorator_ConeCheck::OnBlackboardChange(const UBlackboardComponent& Blac
 	check(false);
 }
 
-void UBTDecorator_ConeCheck::TickNode(UBehaviorTreeComponent* OwnerComp, uint8* NodeMemory, float DeltaSeconds)
+void UBTDecorator_ConeCheck::TickNode(UBehaviorTreeComponent& OwnerComp, uint8* NodeMemory, float DeltaSeconds)
 {
-	const TNodeInstanceMemory* DecoratorMemory = (TNodeInstanceMemory*)NodeMemory;
+	const TNodeInstanceMemory* DecoratorMemory = reinterpret_cast<TNodeInstanceMemory*>(NodeMemory);
 
 	if (CalcConditionImpl(OwnerComp, NodeMemory) != DecoratorMemory->bLastRawResult)
 	{
-		OwnerComp->RequestExecution(this);
+		OwnerComp.RequestExecution(this);
 	}
 }
 
@@ -109,9 +109,9 @@ FString UBTDecorator_ConeCheck::GetStaticDescription() const
 		, ConeHalfAngle * 2, *ConeOrigin.SelectedKeyName.ToString(), *ConeDirection.SelectedKeyName.ToString());
 }
 
-void UBTDecorator_ConeCheck::DescribeRuntimeValues(const UBehaviorTreeComponent* OwnerComp, uint8* NodeMemory, EBTDescriptionVerbosity::Type Verbosity, TArray<FString>& Values) const
+void UBTDecorator_ConeCheck::DescribeRuntimeValues(const UBehaviorTreeComponent& OwnerComp, uint8* NodeMemory, EBTDescriptionVerbosity::Type Verbosity, TArray<FString>& Values) const
 {
-	const UBlackboardComponent* BBComponent = OwnerComp->GetBlackboardComponent();
+	const UBlackboardComponent* BBComponent = OwnerComp.GetBlackboardComponent();
 			
 	FVector ConeDir;
 	FVector DirectionToObserved;
diff --git a/Engine/Source/Runtime/AIModule/Private/BehaviorTree/Decorators/BTDecorator_Cooldown.cpp b/Engine/Source/Runtime/AIModule/Private/BehaviorTree/Decorators/BTDecorator_Cooldown.cpp
--- a/Engine/Source/Runtime/AIModule/Private/BehaviorTree/Decorators/BTDecorator_Cooldown.cpp
+++ b/Engine/Source/Runtime/AIModule/Private/BehaviorTree/Decorators/BTDecorator_Cooldown.cpp
@@ -21,16 +21,16 @@ void UBTDecorator_Cooldown::PostLoad()
 	bNotifyTick = (FlowAbortMode != EBTFlowAbortMode::None);
 }
 
-bool UBTDecorator_Cooldown::CalculateRawConditionValue(UBehaviorTreeComponent* OwnerComp, uint8* NodeMemory) const 
+bool UBTDecorator_Cooldown::CalculateRawConditionValue(UBehaviorTreeComponent& OwnerComp, uint8* NodeMemory) const
 {
-	FBTCooldownDecoratorMemory* DecoratorMemory = (FBTCooldownDecoratorMemory*)NodeMemory;
-	const float TimePassed = (OwnerComp->GetWorld()->GetTimeSeconds() - DecoratorMemory->LastUseTimestamp);
+	FBTCooldownDecoratorMemory* DecoratorMemory = reinterpret_cast<FBTCooldownDecoratorMemory*>(NodeMemory);
+	const float TimePassed = (OwnerComp.GetWorld()->GetTimeSeconds() - DecoratorMemory->LastUseTimestamp);
 	return TimePassed >= CoolDownTime;
 }
 
-void UBTDecorator_Cooldown::InitializeMemory(UBehaviorTreeComponent* OwnerComp, uint8* NodeMemory, EBTMemoryInit::Type InitType) const
+void UBTDecorator_Cooldown::InitializeMemory(UBehaviorTreeComponent& OwnerComp, uint8* NodeMemory, EBTMemoryInit::Type InitType) const
 {
-	FBTCooldownDecoratorMemory* DecoratorMemory = (FBTCooldownDecoratorMemory*)NodeMemory;
+	FBTCooldownDecoratorMemory* DecoratorMemory = reinterpret_cast<FBTCooldownDecoratorMemory*>(NodeMemory);
 	if (InitType == EBTMemoryInit::Initialize)
 	{
 		DecoratorMemory->LastUseTimestamp = -FLT_MAX;
@@ -42,20 +42,20 @@ void UBTDecorator_Cooldown::InitializeMemory(UBehaviorTreeComponent* OwnerComp,
 void UBTDecorator_Cooldown::OnNodeDeactivation(FBehaviorTreeSearchData& SearchData, EBTNodeResult::Type NodeResult)
 {
 	FBTCooldownDecoratorMemory* DecoratorMemory = GetNodeMemory<FBTCooldownDecoratorMemory>(SearchData);
-	DecoratorMemory->LastUseTimestamp = SearchData.OwnerComp->GetWorld()->GetTimeSeconds();
+	DecoratorMemory->LastUseTimestamp = SearchData.OwnerComp.GetWorld()->GetTimeSeconds();
 	DecoratorMemory->bRequestedRestart = false;
 }
 
-void UBTDecorator_Cooldown::TickNode(UBehaviorTreeComponent* OwnerComp, uint8* NodeMemory, float DeltaSeconds)
+void UBTDecorator_Cooldown::TickNode(UBehaviorTreeComponent& OwnerComp, uint8* NodeMemory, float DeltaSeconds)
 {
-	FBTCooldownDecoratorMemory* DecoratorMemory = (FBTCooldownDecoratorMemory*)NodeMemory;
+	FBTCooldownDecoratorMemory* DecoratorMemory = reinterpret_cast<FBTCooldownDecoratorMemory*>(NodeMemory);
 	if (!DecoratorMemory->bRequestedRestart)
 	{
-		const float TimePassed = (OwnerComp->GetWorld()->GetTimeSeconds() - DecoratorMemory->LastUseTimestamp);
+		const float TimePassed = (OwnerComp.GetWorld()->GetTimeSeconds() - DecoratorMemory->LastUseTimestamp);
 		if (TimePassed >= CoolDownTime)
 		{
 			DecoratorMemory->bRequestedRestart = true;
-			OwnerComp->RequestExecution(this);
+			OwnerComp.RequestExecution(this);
 		}
 	}
 }
@@ -67,12 +67,12 @@ FString UBTDecorator_Cooldown::GetStaticDescription() const
 		CoolDownTime, *UBehaviorTreeTypes::DescribeNodeResult(EBTNodeResult::Failed));
 }
 
-void UBTDecorator_Cooldown::DescribeRuntimeValues(const UBehaviorTreeComponent* OwnerComp, uint8* NodeMemory, EBTDescriptionVerbosity::Type Verbosity, TArray<FString>& Values) const
+void UBTDecorator_Cooldown::DescribeRuntimeValues(const UBehaviorTreeComponent& OwnerComp, uint8* NodeMemory, EBTDescriptionVerbosity::Type Verbosity, TArray<FString>& Values) const
 {
 	Super::DescribeRuntimeValues(OwnerComp, NodeMemory, Verbosity, Values);
 
-	FBTCooldownDecoratorMemory* DecoratorMemory = (FBTCooldownDecoratorMemory*)NodeMemory;
-	const float TimePassed = OwnerComp->GetWorld()->GetTimeSeconds() - DecoratorMemory->LastUseTimestamp;
+	FBTCooldownDecoratorMemory* DecoratorMemory = reinterpret_cast<FBTCooldownDecoratorMemory*>(NodeMemory);
+	const float TimePassed = OwnerComp.GetWorld()->GetTimeSeconds() - DecoratorMemory->LastUseTimestamp;
 	
 	if (TimePassed < CoolDownTime)
 	{
diff --git a/Engine/Source/Runtime/AIModule/Private/BehaviorTree/Decorators/BTDecorator_DoesPathExist.cpp b/Engine/Source/Runtime/AIModule/Private/BehaviorTree/Decorators/BTDecorator_DoesPathExist.cpp
--- a/Engine/Source/Runtime/AIModule/Private/BehaviorTree/Decorators/BTDecorator_DoesPathExist.cpp
+++ b/Engine/Source/Runtime/AIModule/Private/BehaviorTree/Decorators/BTDecorator_DoesPathExist.cpp
@@ -41,9 +41,9 @@ void UBTDecorator_DoesPathExist::InitializeFromAsset(UBehaviorTree& Asset)
 	BlackboardKeyB.CacheSelectedKey(BBAsset);
 }
 
-bool UBTDecorator_DoesPathExist::CalculateRawConditionValue(UBehaviorTreeComponent* OwnerComp, uint8* NodeMemory) const 
+bool UBTDecorator_DoesPathExist::CalculateRawConditionValue(UBehaviorTreeComponent& OwnerComp, uint8* NodeMemory) const
 {
-	const UBlackboardComponent* BlackboardComp = OwnerComp->GetBlackboardComponent();
+	const UBlackboardComponent* BlackboardComp = OwnerComp.GetBlackboardComponent();
 	if (BlackboardComp == NULL)
 	{
 		return false;
@@ -56,10 +56,10 @@ bool UBTDecorator_DoesPathExist::CalculateRawConditionValue(UBehaviorTreeCompone
 	
 	bool bHasPath = false;
 
-	const UNavigationSystem* NavSys = UNavigationSystem::GetCurrent(OwnerComp->GetWorld());
+	const UNavigationSystem* NavSys = UNavigationSystem::GetCurrent(OwnerComp.GetWorld());
 	if (NavSys && bHasPointA && bHasPointB)
 	{
-		const AAIController* AIOwner = OwnerComp->GetAIOwner();
+		const AAIController* AIOwner = OwnerComp.GetAIOwner();
 		const ANavigationData* NavData = AIOwner ? NavSys->GetNavDataForProps(AIOwner->GetNavAgentProperties()) : NULL;
 		TSharedPtr<const FNavigationQueryFilter> QueryFilter = UNavigationQueryFilter::GetQueryFilter(NavData, FilterClass);
 
diff --git a/Engine/Source/Runtime/AIModule/Private/BehaviorTree/Decorators/BTDecorator_KeepInCone.cpp b/Engine/Source/Runtime/AIModule/Private/BehaviorTree/Decorators/BTDecorator_KeepInCone.cpp
--- a/Engine/Source/Runtime/AIModule/Private/BehaviorTree/Decorators/BTDecorator_KeepInCone.cpp
+++ b/Engine/Source/Runtime/AIModule/Private/BehaviorTree/Decorators/BTDecorator_KeepInCone.cpp
@@ -49,9 +49,9 @@ void UBTDecorator_KeepInCone::InitializeFromAsset(UBehaviorTree& Asset)
 	Observed.CacheSelectedKey(BBAsset);
 }
 
-bool UBTDecorator_KeepInCone::CalculateCurrentDirection(const UBehaviorTreeComponent* OwnerComp, FVector& Direction) const
+bool UBTDecorator_KeepInCone::CalculateCurrentDirection(const UBehaviorTreeComponent& OwnerComp, FVector& Direction) const
 {
-	const UBlackboardComponent* BlackboardComp = OwnerComp->GetBlackboardComponent();
+	const UBlackboardComponent* BlackboardComp = OwnerComp.GetBlackboardComponent();
 	if (BlackboardComp == NULL)
 	{
 		return false;
@@ -71,18 +71,18 @@ bool UBTDecorator_KeepInCone::CalculateCurrentDirection(const UBehaviorTreeCompo
 	return false;
 }
 
-void UBTDecorator_KeepInCone::OnBecomeRelevant(UBehaviorTreeComponent* OwnerComp, uint8* NodeMemory)
+void UBTDecorator_KeepInCone::OnBecomeRelevant(UBehaviorTreeComponent& OwnerComp, uint8* NodeMemory)
 {
-	TNodeInstanceMemory* DecoratorMemory = (TNodeInstanceMemory*)NodeMemory;
+	TNodeInstanceMemory* DecoratorMemory = reinterpret_cast<TNodeInstanceMemory*>(NodeMemory);
 	FVector InitialDir(1.0f, 0, 0);
 
 	CalculateCurrentDirection(OwnerComp, InitialDir);
 	DecoratorMemory->InitialDirection = InitialDir;
 }
 
-void UBTDecorator_KeepInCone::TickNode(UBehaviorTreeComponent* OwnerComp, uint8* NodeMemory, float DeltaSeconds)
+void UBTDecorator_KeepInCone::TickNode(UBehaviorTreeComponent& OwnerComp, uint8* NodeMemory, float DeltaSeconds)
 {
-	TNodeInstanceMemory* DecoratorMemory = (TNodeInstanceMemory*)NodeMemory;
+	TNodeInstanceMemory* DecoratorMemory = reinterpret_cast<TNodeInstanceMemory*>(NodeMemory);
 	FVector CurrentDir(1.0f, 0, 0);
 	
 	if (CalculateCurrentDirection(OwnerComp, CurrentDir))
@@ -90,7 +90,7 @@ void UBTDecorator_KeepInCone::TickNode(UBehaviorTreeComponent* OwnerComp, uint8*
 		const float Angle = DecoratorMemory->InitialDirection.CosineAngle2D(CurrentDir);
 		if (Angle < ConeHalfAngleDot || (IsInversed() && Angle > ConeHalfAngleDot))
 		{
-			OwnerComp->RequestExecution(this);
+			OwnerComp.RequestExecution(this);
 		}
 	}
 }
@@ -105,9 +105,9 @@ FString UBTDecorator_KeepInCone::GetStaticDescription() const
 		*Observed.SelectedKeyName.ToString());
 }
 
-void UBTDecorator_KeepInCone::DescribeRuntimeValues(const UBehaviorTreeComponent* OwnerComp, uint8* NodeMemory, EBTDescriptionVerbosity::Type Verbosity, TArray<FString>& Values) const
+void UBTDecorator_KeepInCone::DescribeRuntimeValues(const UBehaviorTreeComponent& OwnerComp, uint8* NodeMemory, EBTDescriptionVerbosity::Type Verbosity, TArray<FString>& Values) const
 {
-	TNodeInstanceMemory* DecoratorMemory = (TNodeInstanceMemory*)NodeMemory;
+	TNodeInstanceMemory* DecoratorMemory = reinterpret_cast<TNodeInstanceMemory*>(NodeMemory);
 	FVector CurrentDir(1.0f, 0, 0);
 	
 	if (CalculateCurrentDirection(OwnerComp, CurrentDir))
diff --git a/Engine/Source/Runtime/AIModule/Private/BehaviorTree/Decorators/BTDecorator_Loop.cpp b/Engine/Source/Runtime/AIModule/Private/BehaviorTree/Decorators/BTDecorator_Loop.cpp
--- a/Engine/Source/Runtime/AIModule/Private/BehaviorTree/Decorators/BTDecorator_Loop.cpp
+++ b/Engine/Source/Runtime/AIModule/Private/BehaviorTree/Decorators/BTDecorator_Loop.cpp
@@ -54,7 +54,7 @@ FString UBTDecorator_Loop::GetStaticDescription() const
 		FString::Printf(TEXT("%s: %d loops"), *Super::GetStaticDescription(), NumLoops);
 }
 
-void UBTDecorator_Loop::DescribeRuntimeValues(const UBehaviorTreeComponent* OwnerComp, uint8* NodeMemory, EBTDescriptionVerbosity::Type Verbosity, TArray<FString>& Values) const
+void UBTDecorator_Loop::DescribeRuntimeValues(const UBehaviorTreeComponent& OwnerComp, uint8* NodeMemory, EBTDescriptionVerbosity::Type Verbosity, TArray<FString>& Values) const
 {
 	Super::DescribeRuntimeValues(OwnerComp, NodeMemory, Verbosity, Values);
 
diff --git a/Engine/Source/Runtime/AIModule/Private/BehaviorTree/Decorators/BTDecorator_ReachedMoveGoal.cpp b/Engine/Source/Runtime/AIModule/Private/BehaviorTree/Decorators/BTDecorator_ReachedMoveGoal.cpp
--- a/Engine/Source/Runtime/AIModule/Private/BehaviorTree/Decorators/BTDecorator_ReachedMoveGoal.cpp
+++ b/Engine/Source/Runtime/AIModule/Private/BehaviorTree/Decorators/BTDecorator_ReachedMoveGoal.cpp
@@ -14,9 +14,9 @@ UBTDecorator_ReachedMoveGoal::UBTDecorator_ReachedMoveGoal(const FObjectInitiali
 	FlowAbortMode = EBTFlowAbortMode::None;
 }
 
-bool UBTDecorator_ReachedMoveGoal::CalculateRawConditionValue(UBehaviorTreeComponent* OwnerComp, uint8* NodeMemory) const 
+bool UBTDecorator_ReachedMoveGoal::CalculateRawConditionValue(UBehaviorTreeComponent& OwnerComp, uint8* NodeMemory) const 
 {
-	AAIController* AIOwner = OwnerComp->GetAIOwner();
+	AAIController* AIOwner = OwnerComp.GetAIOwner();
 	const bool bReachedGoal = AIOwner && AIOwner->GetPathFollowingComponent() && AIOwner->GetPathFollowingComponent()->DidMoveReachGoal();
 	return bReachedGoal;
 }
diff --git a/Engine/Source/Runtime/AIModule/Private/BehaviorTree/Decorators/BTDecorator_TimeLimit.cpp b/Engine/Source/Runtime/AIModule/Private/BehaviorTree/Decorators/BTDecorator_TimeLimit.cpp
--- a/Engine/Source/Runtime/AIModule/Private/BehaviorTree/Decorators/BTDecorator_TimeLimit.cpp
+++ b/Engine/Source/Runtime/AIModule/Private/BehaviorTree/Decorators/BTDecorator_TimeLimit.cpp
@@ -17,16 +17,16 @@ UBTDecorator_TimeLimit::UBTDecorator_TimeLimit(const FObjectInitializer& ObjectI
 	FlowAbortMode = EBTFlowAbortMode::Self;
 }
 
-void UBTDecorator_TimeLimit::OnBecomeRelevant(UBehaviorTreeComponent* OwnerComp, uint8* NodeMemory)
+void UBTDecorator_TimeLimit::OnBecomeRelevant(UBehaviorTreeComponent& OwnerComp, uint8* NodeMemory)
 {
 	FBTAuxiliaryMemory* DecoratorMemory = GetSpecialNodeMemory<FBTAuxiliaryMemory>(NodeMemory);
 	DecoratorMemory->NextTickRemainingTime = TimeLimit;
 	DecoratorMemory->AccumulatedDeltaTime = 0.0f;
 }
 
-void UBTDecorator_TimeLimit::TickNode(UBehaviorTreeComponent* OwnerComp, uint8* NodeMemory, float DeltaSeconds)
+void UBTDecorator_TimeLimit::TickNode(UBehaviorTreeComponent& OwnerComp, uint8* NodeMemory, float DeltaSeconds)
 {
-	OwnerComp->RequestExecution(this);
+	OwnerComp.RequestExecution(this);
 
 	// make sure that next tick won't happen before execution request is processed
 	SetNextTickTime(NodeMemory, 1000000.0f);
@@ -39,12 +39,12 @@ FString UBTDecorator_TimeLimit::GetStaticDescription() const
 		*UBehaviorTreeTypes::DescribeNodeResult(EBTNodeResult::Failed), TimeLimit);
 }
 
-void UBTDecorator_TimeLimit::DescribeRuntimeValues(const UBehaviorTreeComponent* OwnerComp, uint8* NodeMemory, EBTDescriptionVerbosity::Type Verbosity, TArray<FString>& Values) const
+void UBTDecorator_TimeLimit::DescribeRuntimeValues(const UBehaviorTreeComponent& OwnerComp, uint8* NodeMemory, EBTDescriptionVerbosity::Type Verbosity, TArray<FString>& Values) const
 {
 	Super::DescribeRuntimeValues(OwnerComp, NodeMemory, Verbosity, Values);
 
 	FBTAuxiliaryMemory* DecoratorMemory = GetSpecialNodeMemory<FBTAuxiliaryMemory>(NodeMemory);
-	if (OwnerComp && OwnerComp->GetWorld())
+	if (OwnerComp.GetWorld())
 	{
 		const float TimeLeft = DecoratorMemory->NextTickRemainingTime > 0 ? DecoratorMemory->NextTickRemainingTime : 0;
 		Values.Add(FString::Printf(TEXT("%s in %ss"),
diff --git a/Engine/Source/Runtime/AIModule/Private/BehaviorTree/Services/BTService_BlueprintBase.cpp b/Engine/Source/Runtime/AIModule/Private/BehaviorTree/Services/BTService_BlueprintBase.cpp
--- a/Engine/Source/Runtime/AIModule/Private/BehaviorTree/Services/BTService_BlueprintBase.cpp
+++ b/Engine/Source/Runtime/AIModule/Private/BehaviorTree/Services/BTService_BlueprintBase.cpp
@@ -39,7 +39,7 @@ void UBTService_BlueprintBase::SetOwner(AActor* InActorOwner)
 	AIOwner = Cast<AAIController>(InActorOwner);
 }
 
-void UBTService_BlueprintBase::OnBecomeRelevant(UBehaviorTreeComponent* OwnerComp, uint8* NodeMemory)
+void UBTService_BlueprintBase::OnBecomeRelevant(UBehaviorTreeComponent& OwnerComp, uint8* NodeMemory)
 {
 	Super::OnBecomeRelevant(OwnerComp, NodeMemory);
 
@@ -54,15 +54,15 @@ void UBTService_BlueprintBase::OnBecomeRelevant(UBehaviorTreeComponent* OwnerCom
 	}
 }
 
-void UBTService_BlueprintBase::OnCeaseRelevant(UBehaviorTreeComponent* OwnerComp, uint8* NodeMemory)
+void UBTService_BlueprintBase::OnCeaseRelevant(UBehaviorTreeComponent& OwnerComp, uint8* NodeMemory)
 {
 	Super::OnCeaseRelevant(OwnerComp, NodeMemory);
 
-	if (OwnerComp && !OwnerComp->HasAnyFlags(RF_BeginDestroyed) && OwnerComp->GetOwner())
+	if (!OwnerComp.HasAnyFlags(RF_BeginDestroyed) && OwnerComp.GetOwner())
 	{
 		// force dropping all pending latent actions associated with this blueprint
 		// we can't have those resuming activity when node is/was aborted
-		BlueprintNodeHelpers::AbortLatentActions(OwnerComp, this);
+		BlueprintNodeHelpers::AbortLatentActions(OwnerComp, *this);
 
 		if (AIOwner != nullptr && ReceiveActivationImplementations & FBTNodeBPImplementationHelper::AISpecific)
 		{
@@ -76,12 +76,10 @@ void UBTService_BlueprintBase::OnCeaseRelevant(UBehaviorTreeComponent* OwnerComp
 	else
 	{
 		UE_LOG(LogBehaviorTree, Warning,
-			TEXT("OnCeaseRelevant called on Blueprint service %s with invalid owner.  OwnerComponent: %s, OwnerComponent Owner: %s.  %s"),
-			*GetNameSafe(this),
-			*GetNameSafe(OwnerComp),
-			OwnerComp ? *GetNameSafe(OwnerComp->GetOwner()) : TEXT("<None>"),
-			OwnerComp && OwnerComp->HasAnyFlags(RF_BeginDestroyed) ? TEXT("OwnerComponent has BeginDestroyed flag") : TEXT("")
-			  );
+			TEXT("OnCeaseRelevant called on Blueprint service %s with invalid owner.  OwnerComponent: %s, OwnerComponent Owner: %s.  %s")
+			, *GetNameSafe(this), *OwnerComp.GetName(), *GetNameSafe(OwnerComp.GetOwner())
+			, OwnerComp.HasAnyFlags(RF_BeginDestroyed) ? TEXT("OwnerComponent has BeginDestroyed flag") : TEXT("")
+		);
 	}
 }
 
@@ -110,7 +108,7 @@ void UBTService_BlueprintBase::OnSearchStart(FBehaviorTreeSearchData& SearchData
 	}
 }
 
-void UBTService_BlueprintBase::TickNode(UBehaviorTreeComponent* OwnerComp, uint8* NodeMemory, float DeltaSeconds)
+void UBTService_BlueprintBase::TickNode(UBehaviorTreeComponent& OwnerComp, uint8* NodeMemory, float DeltaSeconds)
 {
 	Super::TickNode(OwnerComp, NodeMemory, DeltaSeconds);
 
@@ -159,7 +157,7 @@ FString UBTService_BlueprintBase::GetStaticServiceDescription() const
 	return ReturnDesc;
 }
 
-void UBTService_BlueprintBase::DescribeRuntimeValues(const UBehaviorTreeComponent* OwnerComp, uint8* NodeMemory, EBTDescriptionVerbosity::Type Verbosity, TArray<FString>& Values) const
+void UBTService_BlueprintBase::DescribeRuntimeValues(const UBehaviorTreeComponent& OwnerComp, uint8* NodeMemory, EBTDescriptionVerbosity::Type Verbosity, TArray<FString>& Values) const
 {
 	UBTService_BlueprintBase* CDO = (UBTService_BlueprintBase*)(GetClass()->GetDefaultObject());
 	if (CDO && CDO->PropertyData.Num())
@@ -169,10 +167,10 @@ void UBTService_BlueprintBase::DescribeRuntimeValues(const UBehaviorTreeComponen
 	}
 }
 
-void UBTService_BlueprintBase::OnInstanceDestroyed(UBehaviorTreeComponent* OwnerComp)
+void UBTService_BlueprintBase::OnInstanceDestroyed(UBehaviorTreeComponent& OwnerComp)
 {
 	// force dropping all pending latent actions associated with this blueprint
-	BlueprintNodeHelpers::AbortLatentActions(OwnerComp, this);
+	BlueprintNodeHelpers::AbortLatentActions(OwnerComp, *this);
 }
 
 #if WITH_EDITOR
diff --git a/Engine/Source/Runtime/AIModule/Private/BehaviorTree/Services/BTService_DefaultFocus.cpp b/Engine/Source/Runtime/AIModule/Private/BehaviorTree/Services/BTService_DefaultFocus.cpp
--- a/Engine/Source/Runtime/AIModule/Private/BehaviorTree/Services/BTService_DefaultFocus.cpp
+++ b/Engine/Source/Runtime/AIModule/Private/BehaviorTree/Services/BTService_DefaultFocus.cpp
@@ -19,18 +19,16 @@ UBTService_DefaultFocus::UBTService_DefaultFocus(const FObjectInitializer& Objec
 	BlackboardKey.AddVectorFilter(this);
 }
 
-void UBTService_DefaultFocus::OnBecomeRelevant(UBehaviorTreeComponent* OwnerComp, uint8* NodeMemory) 
+void UBTService_DefaultFocus::OnBecomeRelevant(UBehaviorTreeComponent& OwnerComp, uint8* NodeMemory)
 {
-	check(OwnerComp);
-
 	Super::OnBecomeRelevant(OwnerComp, NodeMemory);
 
 	FBTFocusMemory* MyMemory = (FBTFocusMemory*)NodeMemory;
 	check(MyMemory);
 	MyMemory->Reset();
 
-	AAIController* OwnerController = OwnerComp->GetAIOwner();
-	const UBlackboardComponent* MyBlackboard = OwnerComp->GetBlackboardComponent();
+	AAIController* OwnerController = OwnerComp.GetAIOwner();
+	const UBlackboardComponent* MyBlackboard = OwnerComp.GetBlackboardComponent();
 	
 	if (OwnerController != NULL && MyBlackboard != NULL)
 	{
@@ -54,15 +52,13 @@ void UBTService_DefaultFocus::OnBecomeRelevant(UBehaviorTreeComponent* OwnerComp
 	}
 }
 
-void UBTService_DefaultFocus::OnCeaseRelevant(UBehaviorTreeComponent* OwnerComp, uint8* NodeMemory)
+void UBTService_DefaultFocus::OnCeaseRelevant(UBehaviorTreeComponent& OwnerComp, uint8* NodeMemory)
 {
-	check(OwnerComp);
-
 	Super::OnCeaseRelevant(OwnerComp, NodeMemory);
 
 	FBTFocusMemory* MyMemory = (FBTFocusMemory*)NodeMemory;
 	check(MyMemory);
-	AAIController* OwnerController = OwnerComp->GetAIOwner();
+	AAIController* OwnerController = OwnerComp.GetAIOwner();
 	if (OwnerController != NULL)
 	{
 		bool bClearFocus = false;
@@ -94,7 +90,7 @@ FString UBTService_DefaultFocus::GetStaticDescription() const
 	return FString::Printf(TEXT("Set default focus to %s"), *KeyDesc);
 }
 
-void UBTService_DefaultFocus::DescribeRuntimeValues(const UBehaviorTreeComponent* OwnerComp, uint8* NodeMemory, EBTDescriptionVerbosity::Type Verbosity, TArray<FString>& Values) const
+void UBTService_DefaultFocus::DescribeRuntimeValues(const UBehaviorTreeComponent& OwnerComp, uint8* NodeMemory, EBTDescriptionVerbosity::Type Verbosity, TArray<FString>& Values) const
 {
 	Super::DescribeRuntimeValues(OwnerComp, NodeMemory, Verbosity, Values);
 }
diff --git a/Engine/Source/Runtime/AIModule/Private/BehaviorTree/Tasks/BTTask_BlueprintBase.cpp b/Engine/Source/Runtime/AIModule/Private/BehaviorTree/Tasks/BTTask_BlueprintBase.cpp
--- a/Engine/Source/Runtime/AIModule/Private/BehaviorTree/Tasks/BTTask_BlueprintBase.cpp
+++ b/Engine/Source/Runtime/AIModule/Private/BehaviorTree/Tasks/BTTask_BlueprintBase.cpp
@@ -35,7 +35,7 @@ void UBTTask_BlueprintBase::SetOwner(AActor* InActorOwner)
 	AIOwner = Cast<AAIController>(InActorOwner);
 }
 
-EBTNodeResult::Type UBTTask_BlueprintBase::ExecuteTask(UBehaviorTreeComponent* OwnerComp, uint8* NodeMemory)
+EBTNodeResult::Type UBTTask_BlueprintBase::ExecuteTask(UBehaviorTreeComponent& OwnerComp, uint8* NodeMemory)
 {
 	// fail when task doesn't react to execution (start or tick)
 	CurrentCallResult = (ReceiveExecuteImplementations != 0 || ReceiveTickImplementations != 0) ? EBTNodeResult::InProgress : EBTNodeResult::Failed;
@@ -59,11 +59,11 @@ EBTNodeResult::Type UBTTask_BlueprintBase::ExecuteTask(UBehaviorTreeComponent* O
 	return CurrentCallResult;
 }
 
-EBTNodeResult::Type UBTTask_BlueprintBase::AbortTask(UBehaviorTreeComponent* OwnerComp, uint8* NodeMemory)
+EBTNodeResult::Type UBTTask_BlueprintBase::AbortTask(UBehaviorTreeComponent& OwnerComp, uint8* NodeMemory)
 {
 	// force dropping all pending latent actions associated with this blueprint
 	// we can't have those resuming activity when node is/was aborted
-	BlueprintNodeHelpers::AbortLatentActions(OwnerComp, this);
+	BlueprintNodeHelpers::AbortLatentActions(OwnerComp, *this);
 
 	CurrentCallResult = ReceiveAbortImplementations != 0 ? EBTNodeResult::InProgress : EBTNodeResult::Aborted;
 
@@ -86,7 +86,7 @@ EBTNodeResult::Type UBTTask_BlueprintBase::AbortTask(UBehaviorTreeComponent* Own
 	return CurrentCallResult;
 }
 
-void UBTTask_BlueprintBase::TickTask(UBehaviorTreeComponent* OwnerComp, uint8* NodeMemory, float DeltaSeconds) 
+void UBTTask_BlueprintBase::TickTask(UBehaviorTreeComponent& OwnerComp, uint8* NodeMemory, float DeltaSeconds)
 {
 	if (AIOwner != nullptr && ReceiveTickImplementations & FBTNodeBPImplementationHelper::AISpecific)
 	{
@@ -109,7 +109,7 @@ void UBTTask_BlueprintBase::FinishExecute(bool bSuccess)
 	}
 	else if (OwnerComp)
 	{
-		FinishLatentTask(OwnerComp, NodeResult);
+		FinishLatentTask(*OwnerComp, NodeResult);
 	}
 }
 
@@ -124,7 +124,7 @@ void UBTTask_BlueprintBase::FinishAbort()
 	}
 	else if (OwnerComp)
 	{
-		FinishLatentAbort(OwnerComp);
+		FinishLatentAbort(*OwnerComp);
 	}
 }
 
@@ -173,7 +173,7 @@ FString UBTTask_BlueprintBase::GetStaticDescription() const
 	return ReturnDesc;
 }
 
-void UBTTask_BlueprintBase::DescribeRuntimeValues(const UBehaviorTreeComponent* OwnerComp, uint8* NodeMemory, EBTDescriptionVerbosity::Type Verbosity, TArray<FString>& Values) const
+void UBTTask_BlueprintBase::DescribeRuntimeValues(const UBehaviorTreeComponent& OwnerComp, uint8* NodeMemory, EBTDescriptionVerbosity::Type Verbosity, TArray<FString>& Values) const
 {
 	UBTTask_BlueprintBase* CDO = (UBTTask_BlueprintBase*)(GetClass()->GetDefaultObject());
 	if (CDO && CDO->PropertyData.Num())
@@ -182,10 +182,10 @@ void UBTTask_BlueprintBase::DescribeRuntimeValues(const UBehaviorTreeComponent*
 	}
 }
 
-void UBTTask_BlueprintBase::OnInstanceDestroyed(UBehaviorTreeComponent* OwnerComp)
+void UBTTask_BlueprintBase::OnInstanceDestroyed(UBehaviorTreeComponent& OwnerComp)
 {
 	// force dropping all pending latent actions associated with this blueprint
-	BlueprintNodeHelpers::AbortLatentActions(OwnerComp, this);
+	BlueprintNodeHelpers::AbortLatentActions(OwnerComp, *this);
 }
 
 #if WITH_EDITOR
diff --git a/Engine/Source/Runtime/AIModule/Private/BehaviorTree/Tasks/BTTask_MakeNoise.cpp b/Engine/Source/Runtime/AIModule/Private/BehaviorTree/Tasks/BTTask_MakeNoise.cpp
--- a/Engine/Source/Runtime/AIModule/Private/BehaviorTree/Tasks/BTTask_MakeNoise.cpp
+++ b/Engine/Source/Runtime/AIModule/Private/BehaviorTree/Tasks/BTTask_MakeNoise.cpp
@@ -10,11 +10,9 @@ UBTTask_MakeNoise::UBTTask_MakeNoise(const FObjectInitializer& ObjectInitializer
 	NodeName = "MakeNoise";
 }
 
-EBTNodeResult::Type UBTTask_MakeNoise::ExecuteTask(UBehaviorTreeComponent* OwnerComp, uint8* NodeMemory)
+EBTNodeResult::Type UBTTask_MakeNoise::ExecuteTask(UBehaviorTreeComponent& OwnerComp, uint8* NodeMemory)
 {
-	check(OwnerComp);
-
-	const AController* MyController = OwnerComp ? Cast<AController>(OwnerComp->GetOwner()) : NULL;
+	const AController* MyController = Cast<AController>(OwnerComp.GetOwner());
 	APawn* MyPawn = MyController ? MyController->GetPawn() : NULL;
 
 	if (MyPawn)
diff --git a/Engine/Source/Runtime/AIModule/Private/BehaviorTree/Tasks/BTTask_MoveDirectlyToward.cpp b/Engine/Source/Runtime/AIModule/Private/BehaviorTree/Tasks/BTTask_MoveDirectlyToward.cpp
--- a/Engine/Source/Runtime/AIModule/Private/BehaviorTree/Tasks/BTTask_MoveDirectlyToward.cpp
+++ b/Engine/Source/Runtime/AIModule/Private/BehaviorTree/Tasks/BTTask_MoveDirectlyToward.cpp
@@ -19,11 +19,11 @@ UBTTask_MoveDirectlyToward::UBTTask_MoveDirectlyToward(const FObjectInitializer&
 	BlackboardKey.AddVectorFilter(this);
 }
 
-EBTNodeResult::Type UBTTask_MoveDirectlyToward::ExecuteTask(UBehaviorTreeComponent* OwnerComp, uint8* NodeMemory)
+EBTNodeResult::Type UBTTask_MoveDirectlyToward::ExecuteTask(UBehaviorTreeComponent& OwnerComp, uint8* NodeMemory)
 {
-	const UBlackboardComponent* MyBlackboard = OwnerComp->GetBlackboardComponent();
-	FBTMoveDirectlyTowardMemory* MyMemory = (FBTMoveDirectlyTowardMemory*)NodeMemory;
-	AAIController* MyController = OwnerComp ? Cast<AAIController>(OwnerComp->GetOwner()) : NULL;
+	const UBlackboardComponent* MyBlackboard = OwnerComp.GetBlackboardComponent();
+	FBTMoveDirectlyTowardMemory* MyMemory = reinterpret_cast<FBTMoveDirectlyTowardMemory*>(NodeMemory);
+	AAIController* MyController = OwnerComp.GetAIOwner();
 	EBTNodeResult::Type NodeResult = EBTNodeResult::Failed;
 
 	if (MyController && MyBlackboard)
@@ -65,10 +65,10 @@ EBTNodeResult::Type UBTTask_MoveDirectlyToward::ExecuteTask(UBehaviorTreeCompone
 	return NodeResult;
 }
 
-EBTNodeResult::Type UBTTask_MoveDirectlyToward::AbortTask(UBehaviorTreeComponent* OwnerComp, uint8* NodeMemory)
+EBTNodeResult::Type UBTTask_MoveDirectlyToward::AbortTask(UBehaviorTreeComponent& OwnerComp, uint8* NodeMemory)
 {
-	FBTMoveDirectlyTowardMemory* MyMemory = (FBTMoveDirectlyTowardMemory*)NodeMemory;
-	AAIController* MyController = OwnerComp ? Cast<AAIController>(OwnerComp->GetOwner()) : NULL;
+	FBTMoveDirectlyTowardMemory* MyMemory = reinterpret_cast<FBTMoveDirectlyTowardMemory*>(NodeMemory);
+	AAIController* MyController = OwnerComp.GetAIOwner();
 
 	if (MyController && MyController->GetPathFollowingComponent())
 	{
@@ -90,12 +90,12 @@ FString UBTTask_MoveDirectlyToward::GetStaticDescription() const
 	return FString::Printf(TEXT("%s: %s"), *Super::GetStaticDescription(), *KeyDesc);
 }
 
-void UBTTask_MoveDirectlyToward::DescribeRuntimeValues(const UBehaviorTreeComponent* OwnerComp, uint8* NodeMemory, EBTDescriptionVerbosity::Type Verbosity, TArray<FString>& Values) const
+void UBTTask_MoveDirectlyToward::DescribeRuntimeValues(const UBehaviorTreeComponent& OwnerComp, uint8* NodeMemory, EBTDescriptionVerbosity::Type Verbosity, TArray<FString>& Values) const
 {
 	Super::DescribeRuntimeValues(OwnerComp, NodeMemory, Verbosity, Values);
 
-	const UBlackboardComponent* BlackboardComp = OwnerComp->GetBlackboardComponent();
-	FBTMoveDirectlyTowardMemory* MyMemory = (FBTMoveDirectlyTowardMemory*)NodeMemory;
+	const UBlackboardComponent* BlackboardComp = OwnerComp.GetBlackboardComponent();
+	FBTMoveDirectlyTowardMemory* MyMemory = reinterpret_cast<FBTMoveDirectlyTowardMemory*>(NodeMemory);
 
 	if (MyMemory->MoveRequestID && BlackboardComp)
 	{
diff --git a/Engine/Source/Runtime/AIModule/Private/BehaviorTree/Tasks/BTTask_MoveTo.cpp b/Engine/Source/Runtime/AIModule/Private/BehaviorTree/Tasks/BTTask_MoveTo.cpp
--- a/Engine/Source/Runtime/AIModule/Private/BehaviorTree/Tasks/BTTask_MoveTo.cpp
+++ b/Engine/Source/Runtime/AIModule/Private/BehaviorTree/Tasks/BTTask_MoveTo.cpp
@@ -20,12 +20,12 @@ UBTTask_MoveTo::UBTTask_MoveTo(const FObjectInitializer& ObjectInitializer)
 	BlackboardKey.AddVectorFilter(this);
 }
 
-EBTNodeResult::Type UBTTask_MoveTo::ExecuteTask(UBehaviorTreeComponent* OwnerComp, uint8* NodeMemory)
+EBTNodeResult::Type UBTTask_MoveTo::ExecuteTask(UBehaviorTreeComponent& OwnerComp, uint8* NodeMemory)
 {
 	EBTNodeResult::Type NodeResult = EBTNodeResult::InProgress;
 
 	FBTMoveToTaskMemory* MyMemory = (FBTMoveToTaskMemory*)NodeMemory;
-	AAIController* MyController = OwnerComp ? Cast<AAIController>(OwnerComp->GetOwner()) : NULL;
+	AAIController* MyController = OwnerComp.GetAIOwner();
 
 	MyMemory->bWaitingForPath = MyController->ShouldPostponePathUpdates();
 	if (!MyMemory->bWaitingForPath)
@@ -40,11 +40,11 @@ EBTNodeResult::Type UBTTask_MoveTo::ExecuteTask(UBehaviorTreeComponent* OwnerCom
 	return NodeResult;
 }
 
-EBTNodeResult::Type UBTTask_MoveTo::PerformMoveTask(UBehaviorTreeComponent* OwnerComp, uint8* NodeMemory)
+EBTNodeResult::Type UBTTask_MoveTo::PerformMoveTask(UBehaviorTreeComponent& OwnerComp, uint8* NodeMemory)
 {
-	const UBlackboardComponent* MyBlackboard = OwnerComp->GetBlackboardComponent();
+	const UBlackboardComponent* MyBlackboard = OwnerComp.GetBlackboardComponent();
 	FBTMoveToTaskMemory* MyMemory = (FBTMoveToTaskMemory*)NodeMemory;
-	AAIController* MyController = OwnerComp ? Cast<AAIController>(OwnerComp->GetOwner()) : NULL;
+	AAIController* MyController = OwnerComp.GetAIOwner();
 
 	EBTNodeResult::Type NodeResult = EBTNodeResult::Failed;
 
@@ -90,12 +90,12 @@ EBTNodeResult::Type UBTTask_MoveTo::PerformMoveTask(UBehaviorTreeComponent* Owne
 	return NodeResult;
 }
 
-EBTNodeResult::Type UBTTask_MoveTo::AbortTask(UBehaviorTreeComponent* OwnerComp, uint8* NodeMemory)
+EBTNodeResult::Type UBTTask_MoveTo::AbortTask(UBehaviorTreeComponent& OwnerComp, uint8* NodeMemory)
 {
 	FBTMoveToTaskMemory* MyMemory = (FBTMoveToTaskMemory*)NodeMemory;
 	if (!MyMemory->bWaitingForPath)
 	{
-		AAIController* MyController = OwnerComp ? Cast<AAIController>(OwnerComp->GetOwner()) : NULL;
+		AAIController* MyController = OwnerComp.GetAIOwner();
 
 		if (MyController && MyController->GetPathFollowingComponent())
 		{
@@ -106,12 +106,12 @@ EBTNodeResult::Type UBTTask_MoveTo::AbortTask(UBehaviorTreeComponent* OwnerComp,
 	return Super::AbortTask(OwnerComp, NodeMemory);
 }
 
-void UBTTask_MoveTo::TickTask(UBehaviorTreeComponent* OwnerComp, uint8* NodeMemory, float DeltaSeconds)
+void UBTTask_MoveTo::TickTask(UBehaviorTreeComponent& OwnerComp, uint8* NodeMemory, float DeltaSeconds)
 {
 	FBTMoveToTaskMemory* MyMemory = (FBTMoveToTaskMemory*)NodeMemory;
-	if (MyMemory->bWaitingForPath && !OwnerComp->IsPaused())
+	if (MyMemory->bWaitingForPath && !OwnerComp.IsPaused())
 	{
-		AAIController* MyController = OwnerComp ? Cast<AAIController>(OwnerComp->GetOwner()) : NULL;
+		AAIController* MyController = OwnerComp.GetAIOwner();
 		if (MyController && !MyController->ShouldPostponePathUpdates())
 		{
 			UE_VLOG(MyController, LogBehaviorTree, Log, TEXT("Pathfinding requests are unlocked!"));
@@ -126,7 +126,7 @@ void UBTTask_MoveTo::TickTask(UBehaviorTreeComponent* OwnerComp, uint8* NodeMemo
 	}
 }
 
-void UBTTask_MoveTo::OnMessage(UBehaviorTreeComponent* OwnerComp, uint8* NodeMemory, FName Message, int32 SenderID, bool bSuccess)
+void UBTTask_MoveTo::OnMessage(UBehaviorTreeComponent& OwnerComp, uint8* NodeMemory, FName Message, int32 SenderID, bool bSuccess)
 {
 	// AIMessage_RepathFailed means task has failed
 	bSuccess &= (Message != UBrainComponent::AIMessage_RepathFailed);
@@ -145,11 +145,11 @@ FString UBTTask_MoveTo::GetStaticDescription() const
 	return FString::Printf(TEXT("%s: %s"), *Super::GetStaticDescription(), *KeyDesc);
 }
 
-void UBTTask_MoveTo::DescribeRuntimeValues(const UBehaviorTreeComponent* OwnerComp, uint8* NodeMemory, EBTDescriptionVerbosity::Type Verbosity, TArray<FString>& Values) const
+void UBTTask_MoveTo::DescribeRuntimeValues(const UBehaviorTreeComponent& OwnerComp, uint8* NodeMemory, EBTDescriptionVerbosity::Type Verbosity, TArray<FString>& Values) const
 {
 	Super::DescribeRuntimeValues(OwnerComp, NodeMemory, Verbosity, Values);
 
-	const UBlackboardComponent* BlackboardComp = OwnerComp->GetBlackboardComponent();
+	const UBlackboardComponent* BlackboardComp = OwnerComp.GetBlackboardComponent();
 	FBTMoveToTaskMemory* MyMemory = (FBTMoveToTaskMemory*)NodeMemory;
 
 	if (MyMemory->MoveRequestID && BlackboardComp)
diff --git a/Engine/Source/Runtime/AIModule/Private/BehaviorTree/Tasks/BTTask_PlaySound.cpp b/Engine/Source/Runtime/AIModule/Private/BehaviorTree/Tasks/BTTask_PlaySound.cpp
--- a/Engine/Source/Runtime/AIModule/Private/BehaviorTree/Tasks/BTTask_PlaySound.cpp
+++ b/Engine/Source/Runtime/AIModule/Private/BehaviorTree/Tasks/BTTask_PlaySound.cpp
@@ -11,9 +11,9 @@ UBTTask_PlaySound::UBTTask_PlaySound(const FObjectInitializer& ObjectInitializer
 	NodeName = "PlaySound";
 }
 
-EBTNodeResult::Type UBTTask_PlaySound::ExecuteTask(UBehaviorTreeComponent* OwnerComp, uint8* NodeMemory)
+EBTNodeResult::Type UBTTask_PlaySound::ExecuteTask(UBehaviorTreeComponent& OwnerComp, uint8* NodeMemory)
 {
-	const AAIController* MyController = OwnerComp ? Cast<AAIController>(OwnerComp->GetOwner()) : NULL;
+	const AAIController* MyController = OwnerComp.GetAIOwner();
 
 	UAudioComponent* AC = NULL;
 	if (SoundToPlay && MyController)
diff --git a/Engine/Source/Runtime/AIModule/Private/BehaviorTree/Tasks/BTTask_RotateToFaceBBEntry.cpp b/Engine/Source/Runtime/AIModule/Private/BehaviorTree/Tasks/BTTask_RotateToFaceBBEntry.cpp
--- a/Engine/Source/Runtime/AIModule/Private/BehaviorTree/Tasks/BTTask_RotateToFaceBBEntry.cpp
+++ b/Engine/Source/Runtime/AIModule/Private/BehaviorTree/Tasks/BTTask_RotateToFaceBBEntry.cpp
@@ -39,10 +39,9 @@ namespace
 	}
 }
 
-EBTNodeResult::Type UBTTask_RotateToFaceBBEntry::ExecuteTask(UBehaviorTreeComponent* OwnerComp, uint8* NodeMemory)
+EBTNodeResult::Type UBTTask_RotateToFaceBBEntry::ExecuteTask(UBehaviorTreeComponent& OwnerComp, uint8* NodeMemory)
 {
-	check(OwnerComp);
-	AAIController* AIController = OwnerComp->GetAIOwner();
+	AAIController* AIController = OwnerComp.GetAIOwner();
 
 	if (AIController == NULL || AIController->GetPawn() == NULL)
 	{
@@ -57,7 +56,7 @@ EBTNodeResult::Type UBTTask_RotateToFaceBBEntry::ExecuteTask(UBehaviorTreeCompon
 
 	APawn* Pawn = AIController->GetPawn();
 	const FVector PawnLocation = Pawn->GetActorLocation();
-	const UBlackboardComponent* MyBlackboard = OwnerComp->GetBlackboardComponent();
+	const UBlackboardComponent* MyBlackboard = OwnerComp.GetBlackboardComponent();
 
 	if (BlackboardKey.SelectedKeyType == UBlackboardKeyType_Object::StaticClass())
 	{
@@ -130,10 +129,9 @@ EBTNodeResult::Type UBTTask_RotateToFaceBBEntry::ExecuteTask(UBehaviorTreeCompon
 	return Result;
 }
 
-void UBTTask_RotateToFaceBBEntry::TickTask(UBehaviorTreeComponent* OwnerComp, uint8* NodeMemory, float DeltaSeconds)
+void UBTTask_RotateToFaceBBEntry::TickTask(UBehaviorTreeComponent& OwnerComp, uint8* NodeMemory, float DeltaSeconds)
 {
-	check(OwnerComp);
-	AAIController* AIController = OwnerComp->GetAIOwner();
+	AAIController* AIController = OwnerComp.GetAIOwner();
 
 	if (AIController == NULL || AIController->GetPawn() == NULL)
 	{
@@ -173,9 +171,9 @@ void UBTTask_RotateToFaceBBEntry::CleanUp(AAIController& AIController, uint8* No
 	}
 }
 
-EBTNodeResult::Type UBTTask_RotateToFaceBBEntry::AbortTask(UBehaviorTreeComponent* OwnerComp, uint8* NodeMemory)
+EBTNodeResult::Type UBTTask_RotateToFaceBBEntry::AbortTask(UBehaviorTreeComponent& OwnerComp, uint8* NodeMemory)
 {
-	AAIController* AIController = OwnerComp->GetAIOwner();
+	AAIController* AIController = OwnerComp.GetAIOwner();
 
 	if (AIController != NULL)
 	{
@@ -185,12 +183,12 @@ EBTNodeResult::Type UBTTask_RotateToFaceBBEntry::AbortTask(UBehaviorTreeComponen
 	return EBTNodeResult::Succeeded;
 }
 
-void UBTTask_RotateToFaceBBEntry::DescribeRuntimeValues(const UBehaviorTreeComponent* OwnerComp, uint8* NodeMemory, EBTDescriptionVerbosity::Type Verbosity, TArray<FString>& Values) const
+void UBTTask_RotateToFaceBBEntry::DescribeRuntimeValues(const UBehaviorTreeComponent& OwnerComp, uint8* NodeMemory, EBTDescriptionVerbosity::Type Verbosity, TArray<FString>& Values) const
 {
 	FString KeyDesc = BlackboardKey.SelectedKeyName.ToString();
 	Values.Add(FString::Printf(TEXT("%s: %s"), *Super::GetStaticDescription(), *KeyDesc));
 
-	AAIController* AIController = OwnerComp->GetAIOwner();
+	AAIController* AIController = OwnerComp.GetAIOwner();
 
 	if (AIController != NULL && AIController->GetPawn() != NULL)
 	{
diff --git a/Engine/Source/Runtime/AIModule/Private/BehaviorTree/Tasks/BTTask_RunBehavior.cpp b/Engine/Source/Runtime/AIModule/Private/BehaviorTree/Tasks/BTTask_RunBehavior.cpp
--- a/Engine/Source/Runtime/AIModule/Private/BehaviorTree/Tasks/BTTask_RunBehavior.cpp
+++ b/Engine/Source/Runtime/AIModule/Private/BehaviorTree/Tasks/BTTask_RunBehavior.cpp
@@ -9,9 +9,9 @@ UBTTask_RunBehavior::UBTTask_RunBehavior(const FObjectInitializer& ObjectInitial
 	NodeName = "Run Behavior";
 }
 
-EBTNodeResult::Type UBTTask_RunBehavior::ExecuteTask(UBehaviorTreeComponent* OwnerComp, uint8* NodeMemory)
+EBTNodeResult::Type UBTTask_RunBehavior::ExecuteTask(UBehaviorTreeComponent& OwnerComp, uint8* NodeMemory)
 {
-	const bool bPushed = OwnerComp->PushInstance(*BehaviorAsset);
+	const bool bPushed = OwnerComp.PushInstance(*BehaviorAsset);
 	return bPushed ? EBTNodeResult::InProgress : EBTNodeResult::Failed;
 }
 
@@ -39,7 +39,7 @@ uint16 UBTTask_RunBehavior::GetInstanceMemorySize() const
 	return MemorySize;
 }
 
-void UBTTask_RunBehavior::InjectNodes(UBehaviorTreeComponent* OwnerComp, uint8* NodeMemory, int32& InstancedIndex) const
+void UBTTask_RunBehavior::InjectNodes(UBehaviorTreeComponent& OwnerComp, uint8* NodeMemory, int32& InstancedIndex) const
 {
 	if (BehaviorAsset == NULL || BehaviorAsset->RootDecorators.Num() == 0)
 	{
@@ -53,7 +53,7 @@ void UBTTask_RunBehavior::InjectNodes(UBehaviorTreeComponent* OwnerComp, uint8*
 	uint8* InjectedMemoryBase = NodeMemory + sizeof(FBTInstancedNodeMemory);
 
 	// initialize on first access
-	if (!OwnerComp->NodeInstances.IsValidIndex(InstancedIndex))
+	if (!OwnerComp.NodeInstances.IsValidIndex(InstancedIndex))
 	{
 		TArray<uint16> MemoryOffsets;
 		int32 MemorySize = 0;
@@ -88,7 +88,7 @@ void UBTTask_RunBehavior::InjectNodes(UBehaviorTreeComponent* OwnerComp, uint8*
 				DecoratorOb->InitializeInSubtree(OwnerComp, RawDummyMemory, InstancedIndex, InitType);
 				DecoratorOb->ForceInstancing(false);
 
-				UBTDecorator* InstancedOb = Cast<UBTDecorator>(OwnerComp->NodeInstances.Last());
+				UBTDecorator* InstancedOb = Cast<UBTDecorator>(OwnerComp.NodeInstances.Last());
 				InstancedOb->InitializeMemory(OwnerComp, InjectedMemoryBase + MemoryOffsets[Idx], InitType);
 			}
 		}
@@ -97,7 +97,7 @@ void UBTTask_RunBehavior::InjectNodes(UBehaviorTreeComponent* OwnerComp, uint8*
 		uint16 NewExecutionIdx = GetExecutionIndex() - GetInjectedNodesCount();
 		for (int32 Idx = 0; Idx < NumInjectedDecorators; Idx++)
 		{
-			UBTDecorator* InstancedOb = Cast<UBTDecorator>(OwnerComp->NodeInstances[FirstNodeIdx + Idx]);
+			UBTDecorator* InstancedOb = Cast<UBTDecorator>(OwnerComp.NodeInstances[FirstNodeIdx + Idx]);
 			InstancedOb->InitializeNode(GetParentNode(), NewExecutionIdx, GetMemoryOffset() + MemoryOffsets[Idx], GetTreeDepth() - 1);
 			InstancedOb->MarkInjectedNode();
 
@@ -113,7 +113,7 @@ void UBTTask_RunBehavior::InjectNodes(UBehaviorTreeComponent* OwnerComp, uint8*
 		for (int32 Idx = 0; Idx < NumInjectedDecorators; Idx++)
 		{
 			UBTDecorator* DecoratorOb = BehaviorAsset->RootDecorators[Idx];
-			UBTDecorator* InstancedOb = Cast<UBTDecorator>(OwnerComp->NodeInstances[FirstNodeIdx + Idx]);
+			UBTDecorator* InstancedOb = Cast<UBTDecorator>(OwnerComp.NodeInstances[FirstNodeIdx + Idx]);
 
 			const bool bUsesInstancing = DecoratorOb->HasInstance();
 			if (bUsesInstancing)
@@ -133,7 +133,7 @@ void UBTTask_RunBehavior::InjectNodes(UBehaviorTreeComponent* OwnerComp, uint8*
 	// inject nodes
 	if (GetParentNode())
 	{
-		const int32 ChildIdx = GetParentNode()->GetChildIndex(this);
+		const int32 ChildIdx = GetParentNode()->GetChildIndex(*this);
 		if (ChildIdx != INDEX_NONE)
 		{
 			FBTCompositeChild& LinkData = GetParentNode()->Children[ChildIdx];
@@ -153,7 +153,7 @@ void UBTTask_RunBehavior::InjectNodes(UBehaviorTreeComponent* OwnerComp, uint8*
 			const int32 NumOriginalDecorators = LinkData.Decorators.Num();
 			for (int32 Idx = 0; Idx < NumInjectedDecorators; Idx++)
 			{
-				UBTDecorator* InstancedOb = Cast<UBTDecorator>(OwnerComp->NodeInstances[FirstNodeIdx + Idx]);
+				UBTDecorator* InstancedOb = Cast<UBTDecorator>(OwnerComp.NodeInstances[FirstNodeIdx + Idx]);
 				InstancedOb->InitializeFromAsset(*BehaviorAsset);
 				InstancedOb->InitializeDecorator(ChildIdx);
 
@@ -223,12 +223,12 @@ void UBTTask_RunBehavior::InjectNodes(UBehaviorTreeComponent* OwnerComp, uint8*
 
 				if (NodeIt)
 				{
-					NodeIt->InitializeExecutionOrder(OwnerComp->NodeInstances[FirstNodeIdx]);
+					NodeIt->InitializeExecutionOrder(OwnerComp.NodeInstances[FirstNodeIdx]);
 					NodeIt = NodeIt->GetNextNode();
 
 					for (int32 Idx = 1; Idx < NumInjectedDecorators; Idx++)
 					{
-						NodeIt->InitializeExecutionOrder(OwnerComp->NodeInstances[FirstNodeIdx + Idx]);
+						NodeIt->InitializeExecutionOrder(OwnerComp.NodeInstances[FirstNodeIdx + Idx]);
 						NodeIt = NodeIt->GetNextNode();
 					}
 
@@ -240,11 +240,11 @@ void UBTTask_RunBehavior::InjectNodes(UBehaviorTreeComponent* OwnerComp, uint8*
 	}
 }
 
-void UBTTask_RunBehavior::CleanupMemory(UBehaviorTreeComponent* OwnerComp, uint8* NodeMemory, EBTMemoryClear::Type CleanupType) const
+void UBTTask_RunBehavior::CleanupMemory(UBehaviorTreeComponent& OwnerComp, uint8* NodeMemory, EBTMemoryClear::Type CleanupType) const
 {
 	Super::CleanupMemory(OwnerComp, NodeMemory, CleanupType);
 
-	if (GetParentNode() && OwnerComp && BehaviorAsset)
+	if (GetParentNode() && BehaviorAsset)
 	{
 		FBTInstancedNodeMemory* NodeMemoryHeader = (FBTInstancedNodeMemory*)NodeMemory;
 		uint8* InjectedMemoryBase = NodeMemory + sizeof(FBTInstancedNodeMemory);
@@ -253,8 +253,8 @@ void UBTTask_RunBehavior::CleanupMemory(UBehaviorTreeComponent* OwnerComp, uint8
 		for (int32 Idx = 0; Idx < NumInjectedDecorators; Idx++)
 		{
 			const int32 InstancedNodeIdx = NodeMemoryHeader->NodeIdx + Idx;
-			UBTDecorator* InstancedOb = OwnerComp->NodeInstances.IsValidIndex(InstancedNodeIdx) ?
-				Cast<UBTDecorator>(OwnerComp->NodeInstances[InstancedNodeIdx]) : NULL;
+			UBTDecorator* InstancedOb = OwnerComp.NodeInstances.IsValidIndex(InstancedNodeIdx) ?
+				Cast<UBTDecorator>(OwnerComp.NodeInstances[InstancedNodeIdx]) : NULL;
 
 			if (InstancedOb)
 			{
diff --git a/Engine/Source/Runtime/AIModule/Private/BehaviorTree/Tasks/BTTask_RunEQSQuery.cpp b/Engine/Source/Runtime/AIModule/Private/BehaviorTree/Tasks/BTTask_RunEQSQuery.cpp
--- a/Engine/Source/Runtime/AIModule/Private/BehaviorTree/Tasks/BTTask_RunEQSQuery.cpp
+++ b/Engine/Source/Runtime/AIModule/Private/BehaviorTree/Tasks/BTTask_RunEQSQuery.cpp
@@ -18,10 +18,11 @@ UBTTask_RunEQSQuery::UBTTask_RunEQSQuery(const FObjectInitializer& ObjectInitial
 	}
 }
 
-EBTNodeResult::Type UBTTask_RunEQSQuery::ExecuteTask(UBehaviorTreeComponent* OwnerComp, uint8* NodeMemory)
+EBTNodeResult::Type UBTTask_RunEQSQuery::ExecuteTask(UBehaviorTreeComponent& OwnerComp, uint8* NodeMemory)
 {
-	AActor* QueryOwner = OwnerComp->GetOwner();
-	if (AController* ControllerOwner = Cast<AController>(QueryOwner))
+	AActor* QueryOwner = OwnerComp.GetOwner();
+	AController* ControllerOwner = Cast<AController>(QueryOwner);
+	if (ControllerOwner)
 	{
 		QueryOwner = ControllerOwner->GetPawn();
 	}
@@ -42,9 +43,9 @@ EBTNodeResult::Type UBTTask_RunEQSQuery::ExecuteTask(UBehaviorTreeComponent* Own
 	return EBTNodeResult::Failed;
 }
 
-EBTNodeResult::Type UBTTask_RunEQSQuery::AbortTask(UBehaviorTreeComponent* OwnerComp, uint8* NodeMemory)
+EBTNodeResult::Type UBTTask_RunEQSQuery::AbortTask(UBehaviorTreeComponent& OwnerComp, uint8* NodeMemory)
 {
-	UWorld* MyWorld = OwnerComp->GetWorld();
+	UWorld* MyWorld = OwnerComp.GetWorld();
 	UEnvQueryManager* QueryManager = UEnvQueryManager::GetCurrent(MyWorld);
 	
 	if (QueryManager)
@@ -62,7 +63,7 @@ FString UBTTask_RunEQSQuery::GetStaticDescription() const
 		*GetNameSafe(QueryTemplate), *BlackboardKey.SelectedKeyName.ToString());
 }
 
-void UBTTask_RunEQSQuery::DescribeRuntimeValues(const UBehaviorTreeComponent* OwnerComp, uint8* NodeMemory, EBTDescriptionVerbosity::Type Verbosity, TArray<FString>& Values) const
+void UBTTask_RunEQSQuery::DescribeRuntimeValues(const UBehaviorTreeComponent& OwnerComp, uint8* NodeMemory, EBTDescriptionVerbosity::Type Verbosity, TArray<FString>& Values) const
 {
 	Super::DescribeRuntimeValues(OwnerComp, NodeMemory, Verbosity, Values);
 
diff --git a/Engine/Source/Runtime/AIModule/Private/BehaviorTree/Tasks/BTTask_Wait.cpp b/Engine/Source/Runtime/AIModule/Private/BehaviorTree/Tasks/BTTask_Wait.cpp
--- a/Engine/Source/Runtime/AIModule/Private/BehaviorTree/Tasks/BTTask_Wait.cpp
+++ b/Engine/Source/Runtime/AIModule/Private/BehaviorTree/Tasks/BTTask_Wait.cpp
@@ -10,7 +10,7 @@ UBTTask_Wait::UBTTask_Wait(const FObjectInitializer& ObjectInitializer) : Super(
 	bNotifyTick = true;
 }
 
-EBTNodeResult::Type UBTTask_Wait::ExecuteTask(UBehaviorTreeComponent* OwnerComp, uint8* NodeMemory)
+EBTNodeResult::Type UBTTask_Wait::ExecuteTask(UBehaviorTreeComponent& OwnerComp, uint8* NodeMemory)
 {
 	FBTWaitTaskMemory* MyMemory = (FBTWaitTaskMemory*)NodeMemory;
 	MyMemory->RemainingWaitTime = WaitTime;
@@ -18,7 +18,7 @@ EBTNodeResult::Type UBTTask_Wait::ExecuteTask(UBehaviorTreeComponent* OwnerComp,
 	return EBTNodeResult::InProgress;
 }
 
-void UBTTask_Wait::TickTask(UBehaviorTreeComponent* OwnerComp, uint8* NodeMemory, float DeltaSeconds)
+void UBTTask_Wait::TickTask(UBehaviorTreeComponent& OwnerComp, uint8* NodeMemory, float DeltaSeconds)
 {
 	FBTWaitTaskMemory* MyMemory = (FBTWaitTaskMemory*)NodeMemory;
 	MyMemory->RemainingWaitTime -= DeltaSeconds;
@@ -35,7 +35,7 @@ FString UBTTask_Wait::GetStaticDescription() const
 	return FString::Printf(TEXT("%s: %.1fs"), *Super::GetStaticDescription(), WaitTime);
 }
 
-void UBTTask_Wait::DescribeRuntimeValues(const UBehaviorTreeComponent* OwnerComp, uint8* NodeMemory, EBTDescriptionVerbosity::Type Verbosity, TArray<FString>& Values) const
+void UBTTask_Wait::DescribeRuntimeValues(const UBehaviorTreeComponent& OwnerComp, uint8* NodeMemory, EBTDescriptionVerbosity::Type Verbosity, TArray<FString>& Values) const
 {
 	Super::DescribeRuntimeValues(OwnerComp, NodeMemory, Verbosity, Values);
 
diff --git a/Engine/Source/Runtime/AIModule/Private/BehaviorTree/Tasks/BTTask_WaitBlackboardTime.cpp b/Engine/Source/Runtime/AIModule/Private/BehaviorTree/Tasks/BTTask_WaitBlackboardTime.cpp
--- a/Engine/Source/Runtime/AIModule/Private/BehaviorTree/Tasks/BTTask_WaitBlackboardTime.cpp
+++ b/Engine/Source/Runtime/AIModule/Private/BehaviorTree/Tasks/BTTask_WaitBlackboardTime.cpp
@@ -18,10 +18,10 @@ void UBTTask_WaitBlackboardTime::InitializeFromAsset(UBehaviorTree& Asset)
 	BlackboardKey.CacheSelectedKey(GetBlackboardAsset());
 }
 
-EBTNodeResult::Type UBTTask_WaitBlackboardTime::ExecuteTask(UBehaviorTreeComponent* OwnerComp, uint8* NodeMemory)
+EBTNodeResult::Type UBTTask_WaitBlackboardTime::ExecuteTask(UBehaviorTreeComponent& OwnerComp, uint8* NodeMemory)
 {
 	// Update wait time based on current blackboard key value
-	const UBlackboardComponent* MyBlackboard = OwnerComp->GetBlackboardComponent();
+	const UBlackboardComponent* MyBlackboard = OwnerComp.GetBlackboardComponent();
 	if (MyBlackboard && BlackboardKey.SelectedKeyType == UBlackboardKeyType_Float::StaticClass())
 	{
 		WaitTime = MyBlackboard->GetValueAsFloat(BlackboardKey.GetSelectedKeyID());
diff --git a/Engine/Source/Runtime/AIModule/Public/BehaviorTreeDelegates.h b/Engine/Source/Runtime/AIModule/Public/BehaviorTreeDelegates.h
--- a/Engine/Source/Runtime/AIModule/Public/BehaviorTreeDelegates.h
+++ b/Engine/Source/Runtime/AIModule/Public/BehaviorTreeDelegates.h
@@ -9,7 +9,7 @@ class APawn;
 struct AIMODULE_API FBehaviorTreeDelegates
 {
 	/** delegate type for tree execution events (Params: const UBehaviorTreeComponent* OwnerComp, const UBehaviorTree* TreeAsset) */
-	DECLARE_MULTICAST_DELEGATE_TwoParams(FOnTreeStarted, const UBehaviorTreeComponent*, const UBehaviorTree* );
+	DECLARE_MULTICAST_DELEGATE_TwoParams(FOnTreeStarted, const UBehaviorTreeComponent&, const UBehaviorTree& );
 
 	/** delegate type for locking AI debugging tool on pawn */
 	DECLARE_MULTICAST_DELEGATE_OneParam(FOnTreeDebugTool, const APawn*);
diff --git a/Engine/Source/Runtime/AIModule/Public/BlueprintNodeHelpers.h b/Engine/Source/Runtime/AIModule/Public/BlueprintNodeHelpers.h
--- a/Engine/Source/Runtime/AIModule/Public/BlueprintNodeHelpers.h
+++ b/Engine/Source/Runtime/AIModule/Public/BlueprintNodeHelpers.h
@@ -20,7 +20,7 @@ namespace BlueprintNodeHelpers
 
 	bool FindNodeOwner(AActor* OwningActor, UBTNode* Node, UBehaviorTreeComponent*& OwningComp, int32& OwningInstanceIdx);
 
-	void AbortLatentActions(UActorComponent* OwnerOb, const UObject* Ob);
+	void AbortLatentActions(UActorComponent& OwnerOb, const UObject& Ob);
 
 	FORCEINLINE bool HasBlueprintFunction(FName FuncName, const UObject* Ob, const UClass* StopAtClass)
 	{
diff --git a/Engine/Source/Runtime/Engine/Public/AI/BlueprintNodeHelpers.h b/Engine/Source/Runtime/Engine/Public/AI/BlueprintNodeHelpers.h
--- a/Engine/Source/Runtime/Engine/Public/AI/BlueprintNodeHelpers.h
+++ b/Engine/Source/Runtime/Engine/Public/AI/BlueprintNodeHelpers.h
@@ -20,7 +20,7 @@ namespace BlueprintNodeHelpers
 
 	bool FindNodeOwner(AActor* OwningActor, UBTNode* Node, UBehaviorTreeComponent*& OwningComp, int32& OwningInstanceIdx);
 
-	void AbortLatentActions(UActorComponent* OwnerOb, const UObject* Ob);
+	void AbortLatentActions(UActorComponent& OwnerOb, const UObject& Ob);
 
 	FORCEINLINE bool HasBlueprintFunction(FName FuncName, const UObject* Ob, const UClass* StopAtClass)
 	{
