Added a BP-implementable event called when AI controller starts using new Blackboard asset #UE4

- to be used for BB values initialization
- unvirtualized UseBlackboard and introduced virtual InitializeBlackboard for better readibility

diff --git a/Engine/Source/Developer/AITestSuite/Private/MockAI/MockAI_BT.cpp b/Engine/Source/Developer/AITestSuite/Private/MockAI/MockAI_BT.cpp
--- a/Engine/Source/Developer/AITestSuite/Private/MockAI/MockAI_BT.cpp
+++ b/Engine/Source/Developer/AITestSuite/Private/MockAI/MockAI_BT.cpp
@@ -28,7 +28,7 @@ bool UMockAI_BT::IsRunning() const
 bool UMockAI_BT::RunBT(UBehaviorTree& BTAsset, EBTExecutionMode::Type RunType)
 {
 	BBComp->InitializeBlackboard(BTAsset.BlackboardAsset);
-	BBComp->CacheBrainComponent(BTComp);
+	BBComp->CacheBrainComponent(*BTComp);
 	BTComp->CacheBlackboardComponent(BBComp);
 
 	UWorld* World = FAITestHelpers::GetWorld();
diff --git a/Engine/Source/Runtime/AIModule/Classes/AIController.h b/Engine/Source/Runtime/AIModule/Classes/AIController.h
--- a/Engine/Source/Runtime/AIModule/Classes/AIController.h
+++ b/Engine/Source/Runtime/AIModule/Classes/AIController.h
@@ -73,10 +73,10 @@ struct FFocusKnowledge
  * @see https://docs.unrealengine.com/latest/INT/Gameplay/Framework/Controller/
  */
 
-UCLASS(BlueprintType, Blueprintable)
+UCLASS(ClassGroup = AI, BlueprintType, Blueprintable)
 class AIMODULE_API AAIController : public AController, public IAIPerceptionListenerInterface
 {
-	GENERATED_UCLASS_BODY()
+	GENERATED_BODY()
 
 protected:
 	FFocusKnowledge	FocusInformation;
@@ -120,9 +120,12 @@ private_subobject:
 	UPawnActionsComponent* ActionsComp;
 
 public:
+	
+	AAIController(const FObjectInitializer& ObjectInitializer);
+
 	/** Event called when PossessedPawn is possesed by this controller. */
-	UFUNCTION(BlueprintImplementableEvent, Category = "AI")
-	void OnPossess(APawn* PossessedPawn);
+	UFUNCTION(BlueprintImplementableEvent)
+	virtual void OnPossess(APawn* PossessedPawn);
 
 	/** Makes AI go toward specified Goal actor (destination will be continuously updated)
 	 *  @param AcceptanceRadius - finish move if pawn gets close enough
@@ -196,7 +199,7 @@ public:
 	/** Blueprint notification that we've completed the current movement request */
 	UPROPERTY(BlueprintAssignable, meta=(DisplayName="MoveCompleted"))
 	FAIMoveCompletedSignature ReceiveMoveCompleted;
-
+	
 	/** Returns status of path following */
 	UFUNCTION(BlueprintCallable, Category="AI|Navigation")
 	EPathFollowingStatus::Type GetMoveStatus() const;
@@ -222,8 +225,15 @@ public:
 
 	/** makes AI use specified BB asset */
 	UFUNCTION(BlueprintCallable, Category = "AI")
-	virtual bool UseBlackboard(UBlackboardData* BlackboardAsset);
+	bool UseBlackboard(UBlackboardData* BlackboardAsset);
 
+protected:
+	UFUNCTION(BlueprintImplementableEvent)
+	virtual void OnUsingBlackBoard(UBlackboardComponent* BlackboardComp, UBlackboardData* BlackboardAsset);
+
+	virtual bool InitializeBlackboard(UBlackboardComponent& BlackboardComp, UBlackboardData& BlackboardAsset);
+
+public:
 	/** Retrieve the final position that controller should be looking at. */
 	UFUNCTION(BlueprintCallable, Category="AI")
 	virtual FVector GetFocalPoint() const;
diff --git a/Engine/Source/Runtime/AIModule/Classes/BehaviorTree/BlackboardComponent.h b/Engine/Source/Runtime/AIModule/Classes/BehaviorTree/BlackboardComponent.h
--- a/Engine/Source/Runtime/AIModule/Classes/BehaviorTree/BlackboardComponent.h
+++ b/Engine/Source/Runtime/AIModule/Classes/BehaviorTree/BlackboardComponent.h
@@ -75,10 +75,10 @@ class AIMODULE_API UBlackboardComponent : public UActorComponent
 	UBlackboardData* GetBlackboardAsset() const;
 
 	/** caches UBrainComponent pointer to be used in communication */
-	void CacheBrainComponent(UBrainComponent* BrainComponent);
+	void CacheBrainComponent(UBrainComponent& BrainComponent);
 
 	/** setup component for using given blackboard asset */
-	void InitializeBlackboard(UBlackboardData* NewAsset);
+	bool InitializeBlackboard(UBlackboardData* NewAsset);
 	
 	/** @return true if component can be used with specified blackboard asset */
 	bool IsCompatibleWith(UBlackboardData* TestAsset) const;
diff --git a/Engine/Source/Runtime/AIModule/Private/AIController.cpp b/Engine/Source/Runtime/AIModule/Private/AIController.cpp
--- a/Engine/Source/Runtime/AIModule/Private/AIController.cpp
+++ b/Engine/Source/Runtime/AIModule/Private/AIController.cpp
@@ -767,6 +767,17 @@ bool AAIController::RunBehaviorTree(UBehaviorTree* BTAsset)
 	return bSuccess;
 }
 
+bool AAIController::InitializeBlackboard(UBlackboardComponent& BlackboardComp, UBlackboardData& BlackboardAsset)
+{
+	check(BlackboardComp.GetOwner() == this);
+	if (BlackboardComp.InitializeBlackboard(&BlackboardAsset))
+	{
+		OnUsingBlackBoard(&BlackboardComp, &BlackboardAsset);
+		return true;
+	}
+	return false;
+}
+
 bool AAIController::UseBlackboard(UBlackboardData* BlackboardAsset)
 {
 	if (BlackboardAsset == NULL)
@@ -783,20 +794,20 @@ bool AAIController::UseBlackboard(UBlackboardData* BlackboardAsset)
 		BlackboardComp = ConstructObject<UBlackboardComponent>(UBlackboardComponent::StaticClass(), this, TEXT("BlackboardComponent"));
 		if (BlackboardComp != NULL)
 		{
-			BlackboardComp->InitializeBlackboard(BlackboardAsset);
+			InitializeBlackboard(*BlackboardComp, *BlackboardAsset);
 			BlackboardComp->RegisterComponent();
 		}
 
 	}
 	else if (BlackboardComp->GetBlackboardAsset() == NULL)
 	{
-		BlackboardComp->InitializeBlackboard(BlackboardAsset);
+		InitializeBlackboard(*BlackboardComp, *BlackboardAsset);
 	}
 	else if (BlackboardComp->GetBlackboardAsset() != BlackboardAsset)
 	{
 		UE_VLOG(this, LogBehaviorTree, Log, TEXT("UseBlackboard: requested blackboard %s while already has %s instantiated. Forcing new BB.")
 			, *GetNameSafe(BlackboardAsset), *GetNameSafe(BlackboardComp->GetBlackboardAsset()));
-		BlackboardComp->InitializeBlackboard(BlackboardAsset);
+		InitializeBlackboard(*BlackboardComp, *BlackboardAsset);
 	}
 
 	return bSuccess;
diff --git a/Engine/Source/Runtime/AIModule/Private/BehaviorTree/BlackboardComponent.cpp b/Engine/Source/Runtime/AIModule/Private/BehaviorTree/BlackboardComponent.cpp
--- a/Engine/Source/Runtime/AIModule/Private/BehaviorTree/BlackboardComponent.cpp
+++ b/Engine/Source/Runtime/AIModule/Private/BehaviorTree/BlackboardComponent.cpp
@@ -31,11 +31,11 @@ void UBlackboardComponent::InitializeComponent()
 	InitializeBlackboard(BlackboardAsset);
 }
 
-void UBlackboardComponent::CacheBrainComponent(UBrainComponent* BrainComponent)
+void UBlackboardComponent::CacheBrainComponent(UBrainComponent& BrainComponent)
 {
-	if (BrainComponent != BrainComp)
+	if (&BrainComponent != BrainComp)
 	{
-		BrainComp = BrainComponent;
+		BrainComp = &BrainComponent;
 	}
 }
 
@@ -68,19 +68,21 @@ void UBlackboardComponent::InitializeParentChain(UBlackboardData* NewAsset)
 	}
 }
 
-void UBlackboardComponent::InitializeBlackboard(UBlackboardData* NewAsset)
+bool UBlackboardComponent::InitializeBlackboard(UBlackboardData* NewAsset)
 {
 	// if we re-initialize with the same asset then there's no point
 	// in reseting, since we'd loose all the accumulated knowledge
 	if (NewAsset == BlackboardAsset)
 	{
-		return;
+		return false;
 	}
 
 	BlackboardAsset = NewAsset;
 	ValueMemory.Reset();
 	ValueOffsets.Reset();
 
+	bool bSuccess = true;
+	
 	if (BlackboardAsset)
 	{
 		if (BlackboardAsset->IsValid())
@@ -127,9 +129,12 @@ void UBlackboardComponent::InitializeBlackboard(UBlackboardData* NewAsset)
 		}
 		else
 		{
+			bSuccess = false;
 			UE_LOG(LogBehaviorTree, Error, TEXT("Blackboard asset (%s) has errors and can't be used!"), *GetNameSafe(BlackboardAsset));
 		}
 	}
+
+	return bSuccess;
 }
 
 UBrainComponent* UBlackboardComponent::GetBrainComponent() const
diff --git a/Engine/Source/Runtime/AIModule/Private/BrainComponent.cpp b/Engine/Source/Runtime/AIModule/Private/BrainComponent.cpp
--- a/Engine/Source/Runtime/AIModule/Private/BrainComponent.cpp
+++ b/Engine/Source/Runtime/AIModule/Private/BrainComponent.cpp
@@ -258,7 +258,7 @@ void UBrainComponent::InitializeComponent()
 		BlackboardComp = GetOwner()->FindComponentByClass<UBlackboardComponent>();
 		if (BlackboardComp)
 		{
-			BlackboardComp->CacheBrainComponent(this);
+			BlackboardComp->CacheBrainComponent(*this);
 		}
 	}
 }
