Fixed BT-implemented BT decorators ignoring FloatAbortMode #UE4

- also, fixed BB-key observing broken for BP-decorators
- added a flag to BTDecorator_BlueprintBase to control whether CalculateRawConditionValue should get called every tick or just on observed BB keys changes
- Note: will result in any BP-implemented BT decorator needing an update since ReceiveConditionCheckAI changed it's signature that will require users to redo it with different UI

diff --git a/Engine/Source/Runtime/AIModule/Classes/BehaviorTree/BTNode.h b/Engine/Source/Runtime/AIModule/Classes/BehaviorTree/BTNode.h
--- a/Engine/Source/Runtime/AIModule/Classes/BehaviorTree/BTNode.h
+++ b/Engine/Source/Runtime/AIModule/Classes/BehaviorTree/BTNode.h
@@ -143,7 +143,7 @@ class AIMODULE_API UBTNode : public UObject
 protected:
 
 	/** node name */
-	UPROPERTY(Category=Description, EditInstanceOnly)
+	UPROPERTY(Category=Description, EditAnywhere)
 	FString NodeName;
 	
 private:
diff --git a/Engine/Source/Runtime/AIModule/Classes/BehaviorTree/Decorators/BTDecorator_BlueprintBase.h b/Engine/Source/Runtime/AIModule/Classes/BehaviorTree/Decorators/BTDecorator_BlueprintBase.h
--- a/Engine/Source/Runtime/AIModule/Classes/BehaviorTree/Decorators/BTDecorator_BlueprintBase.h
+++ b/Engine/Source/Runtime/AIModule/Classes/BehaviorTree/Decorators/BTDecorator_BlueprintBase.h
@@ -26,22 +26,29 @@ struct FBehaviorTreeSearchData;
 UCLASS(Abstract, Blueprintable)
 class AIMODULE_API UBTDecorator_BlueprintBase : public UBTDecorator
 {
-	GENERATED_UCLASS_BODY()
+	GENERATED_BODY()
+
+public:
+	UBTDecorator_BlueprintBase(const FObjectInitializer& ObjectInitializer);
 
 	/** initialize data about blueprint defined properties */
 	void InitializeProperties();
 
 	/** setup node name */
 	virtual void PostInitProperties() override;
+	virtual void PostLoad() override;
 
 	/** notify about changes in blackboard */
 	void OnBlackboardChange(const UBlackboardComponent& Blackboard, FBlackboard::FKey ChangedKeyID);
 
 	virtual FString GetStaticDescription() const override;
 	virtual void DescribeRuntimeValues(const UBehaviorTreeComponent& OwnerComp, uint8* NodeMemory, EBTDescriptionVerbosity::Type Verbosity, TArray<FString>& Values) const override;
-	virtual bool CalculateRawConditionValue(UBehaviorTreeComponent& OwnerComp, uint8* NodeMemory) const override;
+	virtual bool CalculateRawConditionValue(UBehaviorTreeComponent& OwnerComp, uint8* NodeMemory) const override final;
 	virtual void OnInstanceDestroyed(UBehaviorTreeComponent& OwnerComp) override;
 
+	/** return if this decorator should abort in current circumstances */
+	bool GetShouldAbort(UBehaviorTreeComponent& OwnerComp) const;
+
 	virtual void SetOwner(AActor* ActorOwner) override;
 
 #if WITH_EDITOR
@@ -59,6 +66,7 @@ protected:
 	AActor* ActorOwner;
 
 	/** blackboard key names that should be observed */
+	UPROPERTY()
 	TArray<FName> ObservedKeyNames;
 
 	/** properties with runtime values, stored only in class default object */
@@ -68,9 +76,16 @@ protected:
 	UPROPERTY(EditInstanceOnly, Category=Description)
 	uint32 bShowPropertyDetails : 1;
 
-	/** temporary variable for ReceiveConditionCheck-FinishConditionCheck chain */
-	mutable uint32 CurrentCallResult : 1;
+	/** Applies only if Decorator has any FBlackboardKeySelector property and if decorator is 
+	 *	set to abort BT flow. Is set to true ReceiveConditionCheck will be called only on changes 
+	  *	to observed BB keys. If false or no BB keys observed ReceiveConditionCheck will be called every tick */
+	UPROPERTY(EditDefaultsOnly, Category = "FlowControl", AdvancedDisplay)
+	uint32 bCheckConditionOnlyBlackBoardChanges : 1;
 
+	/** gets set to true if decorator declared BB keys it can potentially observe */
+	UPROPERTY()
+	uint32 bIsObservingBB : 1;
+	
 	/** set if ReceiveTick is implemented by blueprint */
 	uint32 ReceiveTickImplementations : 2;
 
@@ -90,6 +105,8 @@ protected:
 	uint32 ReceiveConditionCheckImplementations : 2;
 
 	FOnBlackboardChange BBKeyObserver;
+
+	bool CalculateRawConditionValueImpl(UBehaviorTreeComponent& OwnerComp) const;
 	
 	virtual void OnBecomeRelevant(UBehaviorTreeComponent& OwnerComp, uint8* NodeMemory) override;
 	virtual void OnCeaseRelevant(UBehaviorTreeComponent& OwnerComp, uint8* NodeMemory) override;
@@ -131,7 +148,7 @@ protected:
 	 *	@Note that if both generic and AI event versions are implemented only the more
 	 *	suitable one will be called, meaning the AI version if called for AI, generic one otherwise */
 	UFUNCTION(BlueprintImplementableEvent)
-	virtual void ReceiveConditionCheck(AActor* OwnerActor);
+	virtual bool ReceiveConditionCheck(AActor* OwnerActor);
 
 	/** Alternative AI version of ReceiveTick
 	 *	@see ReceiveTick for more details
@@ -173,11 +190,7 @@ protected:
 	 *	@Note that if both generic and AI event versions are implemented only the more
 	 *	suitable one will be called, meaning the AI version if called for AI, generic one otherwise */
 	UFUNCTION(BlueprintImplementableEvent, Category = AI)
-	virtual void ReceiveConditionCheckAI(AAIController* OwnerController, APawn* ControlledPawn);
-
-	/** finishes condition check */
-	UFUNCTION(BlueprintCallable, Category="AI|BehaviorTree")
-	void FinishConditionCheck(bool bAllowExecution);
+	virtual bool ReceiveConditionCheckAI(AAIController* OwnerController, APawn* ControlledPawn);
 		
 	/** check if decorator is part of currently active branch */
 	UFUNCTION(BlueprintCallable, Category="AI|BehaviorTree")
@@ -188,4 +201,6 @@ protected:
 	bool IsDecoratorObserverActive() const;
 
 	friend FBehaviorBlueprintDetails;
+
+	FORCEINLINE bool GetNeedsTickForConditionChecking() const { return ReceiveConditionCheckImplementations != 0 && (bIsObservingBB == false || bCheckConditionOnlyBlackBoardChanges == false); }
 };
diff --git a/Engine/Source/Runtime/AIModule/Private/BehaviorTree/BTAuxiliaryNode.cpp b/Engine/Source/Runtime/AIModule/Private/BehaviorTree/BTAuxiliaryNode.cpp
--- a/Engine/Source/Runtime/AIModule/Private/BehaviorTree/BTAuxiliaryNode.cpp
+++ b/Engine/Source/Runtime/AIModule/Private/BehaviorTree/BTAuxiliaryNode.cpp
@@ -15,7 +15,7 @@ void UBTAuxiliaryNode::WrappedOnBecomeRelevant(UBehaviorTreeComponent& OwnerComp
 {
 	if (bNotifyBecomeRelevant)
 	{
-		const UBTNode* NodeOb = bCreateNodeInstance ? GetNodeInstance(OwnerComp, NodeMemory) : this;
+		const UBTNode* NodeOb = HasInstance() ? GetNodeInstance(OwnerComp, NodeMemory) : this;
 		if (NodeOb)
 		{
 			((UBTAuxiliaryNode*)NodeOb)->OnBecomeRelevant(OwnerComp, NodeMemory);
@@ -27,7 +27,7 @@ void UBTAuxiliaryNode::WrappedOnCeaseRelevant(UBehaviorTreeComponent& OwnerComp,
 {
 	if (bNotifyCeaseRelevant)
 	{
-		const UBTNode* NodeOb = bCreateNodeInstance ? GetNodeInstance(OwnerComp, NodeMemory) : this;
+		const UBTNode* NodeOb = HasInstance() ? GetNodeInstance(OwnerComp, NodeMemory) : this;
 		if (NodeOb)
 		{
 			((UBTAuxiliaryNode*)NodeOb)->OnCeaseRelevant(OwnerComp, NodeMemory);
@@ -37,29 +37,32 @@ void UBTAuxiliaryNode::WrappedOnCeaseRelevant(UBehaviorTreeComponent& OwnerComp,
 
 void UBTAuxiliaryNode::WrappedTickNode(UBehaviorTreeComponent& OwnerComp, uint8* NodeMemory, float DeltaSeconds) const
 {
-	if (bNotifyTick)
+	if (bNotifyTick || HasInstance())
 	{
-		float UseDeltaTime = DeltaSeconds;
+		const UBTAuxiliaryNode* NodeOb = HasInstance() ? static_cast<UBTAuxiliaryNode*>(GetNodeInstance(OwnerComp, NodeMemory)) : this;
+		
+		ensure(NodeOb);
 
-		if (bTickIntervals)
+		if (NodeOb != nullptr && NodeOb->bNotifyTick)
 		{
-			FBTAuxiliaryMemory* AuxMemory = GetSpecialNodeMemory<FBTAuxiliaryMemory>(NodeMemory);
-			AuxMemory->NextTickRemainingTime -= DeltaSeconds;
-			AuxMemory->AccumulatedDeltaTime += DeltaSeconds;
-			
-			if (AuxMemory->NextTickRemainingTime > 0.0f)
+			float UseDeltaTime = DeltaSeconds;
+
+			if (NodeOb->bTickIntervals)
 			{
-				return;
-			}
+				FBTAuxiliaryMemory* AuxMemory = GetSpecialNodeMemory<FBTAuxiliaryMemory>(NodeMemory);
+				AuxMemory->NextTickRemainingTime -= DeltaSeconds;
+				AuxMemory->AccumulatedDeltaTime += DeltaSeconds;
 
-			UseDeltaTime = AuxMemory->AccumulatedDeltaTime;
-			AuxMemory->AccumulatedDeltaTime = 0.0f;
-		}
+				if (AuxMemory->NextTickRemainingTime > 0.0f)
+				{
+					return;
+				}
 
-		const UBTNode* NodeOb = bCreateNodeInstance ? GetNodeInstance(OwnerComp, NodeMemory) : this;
-		if (NodeOb)
-		{
-			((UBTAuxiliaryNode*)NodeOb)->TickNode(OwnerComp, NodeMemory, UseDeltaTime);
+				UseDeltaTime = AuxMemory->AccumulatedDeltaTime;
+				AuxMemory->AccumulatedDeltaTime = 0.0f;
+			}
+
+			const_cast<UBTAuxiliaryNode*>(NodeOb)->TickNode(OwnerComp, NodeMemory, UseDeltaTime);
 		}
 	}
 }
diff --git a/Engine/Source/Runtime/AIModule/Private/BehaviorTree/BlueprintNodeHelpers.cpp b/Engine/Source/Runtime/AIModule/Private/BehaviorTree/BlueprintNodeHelpers.cpp
--- a/Engine/Source/Runtime/AIModule/Private/BehaviorTree/BlueprintNodeHelpers.cpp
+++ b/Engine/Source/Runtime/AIModule/Private/BehaviorTree/BlueprintNodeHelpers.cpp
@@ -129,6 +129,35 @@ namespace BlueprintNodeHelpers
 		}
 	}
 
+	bool HasAnyBlackboardSelectors(const UObject* Ob, const UClass* StopAtClass)
+	{
+		bool bResult = false;
+
+		for (UProperty* TestProperty = Ob->GetClass()->PropertyLink; TestProperty; TestProperty = TestProperty->PropertyLinkNext)
+		{
+			// stop when reaching base class
+			if (TestProperty->GetOuter() == StopAtClass)
+			{
+				break;
+			}
+
+			// skip properties without any setup data	
+			if (TestProperty->HasAnyPropertyFlags(CPF_Transient) ||
+				TestProperty->HasAnyPropertyFlags(CPF_DisableEditOnInstance))
+			{
+				continue;
+			}
+
+			const UStructProperty* StructProp = Cast<const UStructProperty>(TestProperty);
+			if (StructProp && StructProp->GetCPPType(NULL, CPPF_None).Contains(GET_STRUCT_NAME_CHECKED(FBlackboardKeySelector)))
+			{
+				bResult = true;
+			}
+		}
+
+		return bResult;
+	}
+
 #undef GET_STRUCT_NAME_CHECKED
 
 	FString CollectPropertyDescription(const UObject* Ob, const UClass* StopAtClass, const TArray<UProperty*>& PropertyData)
diff --git a/Engine/Source/Runtime/AIModule/Private/BehaviorTree/Decorators/BTDecorator_BlueprintBase.cpp b/Engine/Source/Runtime/AIModule/Private/BehaviorTree/Decorators/BTDecorator_BlueprintBase.cpp
--- a/Engine/Source/Runtime/AIModule/Private/BehaviorTree/Decorators/BTDecorator_BlueprintBase.cpp
+++ b/Engine/Source/Runtime/AIModule/Private/BehaviorTree/Decorators/BTDecorator_BlueprintBase.cpp
@@ -6,7 +6,8 @@
 #include "BehaviorTree/BlackboardComponent.h"
 #include "BehaviorTree/BlackboardData.h"
 
-UBTDecorator_BlueprintBase::UBTDecorator_BlueprintBase(const FObjectInitializer& ObjectInitializer) : Super(ObjectInitializer)
+UBTDecorator_BlueprintBase::UBTDecorator_BlueprintBase(const FObjectInitializer& ObjectInitializer) 
+	: Super(ObjectInitializer)
 {
 	UClass* StopAtClass = UBTDecorator_BlueprintBase::StaticClass();
 	ReceiveTickImplementations = FBTNodeBPImplementationHelper::CheckEventImplementationVersion(TEXT("ReceiveTick"), TEXT("ReceiveTickAI"), this, StopAtClass);
@@ -22,11 +23,11 @@ UBTDecorator_BlueprintBase::UBTDecorator_BlueprintBase(const FObjectInitializer&
 	bNotifyActivation = ReceiveExecutionStartImplementations != 0;
 	bNotifyDeactivation = ReceiveExecutionFinishImplementations != 0;
 	bShowPropertyDetails = true;
+	bCheckConditionOnlyBlackBoardChanges = false;
+	bIsObservingBB = false;
 
 	// all blueprint based nodes must create instances
 	bCreateNodeInstance = true;
-
-	InitializeProperties();
 }
 
 void UBTDecorator_BlueprintBase::InitializeProperties()
@@ -36,28 +37,41 @@ void UBTDecorator_BlueprintBase::InitializeProperties()
 		UClass* StopAtClass = UBTDecorator_BlueprintBase::StaticClass();
 		BlueprintNodeHelpers::CollectPropertyData(this, StopAtClass, PropertyData);
 
-		if (GetFlowAbortMode() != EBTFlowAbortMode::None)
-		{
-			// find all blackboard key selectors and store their names
-			BlueprintNodeHelpers::CollectBlackboardSelectors(this, StopAtClass, ObservedKeyNames);
-		}
-	}
-	else
-	{
-		UBTDecorator_BlueprintBase* CDO = (UBTDecorator_BlueprintBase*)(GetClass()->GetDefaultObject());
-		if (CDO && CDO->ObservedKeyNames.Num())
-		{
-			bNotifyBecomeRelevant = true;
-			bNotifyCeaseRelevant = true;
-		}
+		bIsObservingBB = BlueprintNodeHelpers::HasAnyBlackboardSelectors(this, StopAtClass);
+
+		NodeName = BlueprintNodeHelpers::GetNodeName(this);
 	}
 }
 
 void UBTDecorator_BlueprintBase::PostInitProperties()
 {
 	Super::PostInitProperties();
-	NodeName = BlueprintNodeHelpers::GetNodeName(this);
-	BBKeyObserver = FOnBlackboardChange::CreateUObject(this, &UBTDecorator_BlueprintBase::OnBlackboardChange);
+		
+	InitializeProperties();
+	
+	if (bIsObservingBB)
+	{
+		BBKeyObserver = FOnBlackboardChange::CreateUObject(this, &UBTDecorator_BlueprintBase::OnBlackboardChange);
+	}
+
+	if (ReceiveConditionCheckImplementations || bIsObservingBB)
+	{
+		bNotifyBecomeRelevant = true;
+		bNotifyCeaseRelevant = true;
+	}
+}
+
+void UBTDecorator_BlueprintBase::PostLoad()
+{
+	Super::PostLoad();
+
+	if (GetFlowAbortMode() != EBTFlowAbortMode::None && bIsObservingBB)
+	{
+		ObservedKeyNames.Reset();
+		UClass* StopAtClass = UBTDecorator_BlueprintBase::StaticClass();
+		BlueprintNodeHelpers::CollectBlackboardSelectors(this, StopAtClass, ObservedKeyNames);
+		ensure(ObservedKeyNames.Num() > 0);
+	}
 }
 
 void UBTDecorator_BlueprintBase::SetOwner(AActor* InActorOwner)
@@ -77,13 +91,19 @@ void UBTDecorator_BlueprintBase::OnBecomeRelevant(UBehaviorTreeComponent& OwnerC
 		ReceiveObserverActivated(ActorOwner);
 	}
 
+	if (GetNeedsTickForConditionChecking())
+	{
+		// if set up as observer, and has a condition check, we want to check the condition every frame
+		// highly inefficient, but hopefully people will use it only for prototyping.
+		bNotifyTick = true;
+	}
+
 	UBlackboardComponent* BlackboardComp = OwnerComp.GetBlackboardComponent();
-	UBTDecorator_BlueprintBase* CDO = (UBTDecorator_BlueprintBase*)(GetClass()->GetDefaultObject());
-	if (BlackboardComp && CDO)
+	if (BlackboardComp)
 	{
-		for (int32 NameIndex = 0; NameIndex < CDO->ObservedKeyNames.Num(); NameIndex++)
+		for (int32 NameIndex = 0; NameIndex < ObservedKeyNames.Num(); NameIndex++)
 		{
-			const FBlackboard::FKey KeyID = BlackboardComp->GetKeyID(CDO->ObservedKeyNames[NameIndex]);
+			const FBlackboard::FKey KeyID = BlackboardComp->GetKeyID(ObservedKeyNames[NameIndex]);
 			if (KeyID != FBlackboard::InvalidKey)
 			{
 				BlackboardComp->RegisterObserver(KeyID, BBKeyObserver);
@@ -95,12 +115,11 @@ void UBTDecorator_BlueprintBase::OnBecomeRelevant(UBehaviorTreeComponent& OwnerC
 void UBTDecorator_BlueprintBase::OnCeaseRelevant(UBehaviorTreeComponent& OwnerComp, uint8* NodeMemory)
 {
 	UBlackboardComponent* BlackboardComp = OwnerComp.GetBlackboardComponent();
-	UBTDecorator_BlueprintBase* CDO = (UBTDecorator_BlueprintBase*)(GetClass()->GetDefaultObject());
-	if (BlackboardComp && CDO)
+	if (BlackboardComp)
 	{
-		for (int32 NameIndex = 0; NameIndex < CDO->ObservedKeyNames.Num(); NameIndex++)
+		for (int32 NameIndex = 0; NameIndex < ObservedKeyNames.Num(); NameIndex++)
 		{
-			const FBlackboard::FKey KeyID = BlackboardComp->GetKeyID(CDO->ObservedKeyNames[NameIndex]);
+			const FBlackboard::FKey KeyID = BlackboardComp->GetKeyID(ObservedKeyNames[NameIndex]);
 			if (KeyID != FBlackboard::InvalidKey)
 			{
 				BlackboardComp->UnregisterObserver(KeyID, BBKeyObserver);
@@ -116,6 +135,12 @@ void UBTDecorator_BlueprintBase::OnCeaseRelevant(UBehaviorTreeComponent& OwnerCo
 	{
 		ReceiveObserverDeactivated(ActorOwner);
 	}
+
+	if (GetNeedsTickForConditionChecking() == true && ReceiveTickImplementations == 0)
+	{
+		// clean up the tick request if no longer "active"
+		bNotifyTick = false;
+	}
 }
 
 void UBTDecorator_BlueprintBase::OnNodeActivation(FBehaviorTreeSearchData& SearchData)
@@ -152,11 +177,40 @@ void UBTDecorator_BlueprintBase::TickNode(UBehaviorTreeComponent& OwnerComp, uin
 	{
 		ReceiveTick(ActorOwner, DeltaSeconds);
 	}
+		
+	// possible this got ticked due to the decorator being configured as an observer
+	if (GetNeedsTickForConditionChecking() && GetShouldAbort(OwnerComp))
+	{
+		OwnerComp.RequestExecution(this);
+	}
 }
 
-bool UBTDecorator_BlueprintBase::CalculateRawConditionValue(UBehaviorTreeComponent& OwnerComp, uint8* NodeMemory) const
+bool UBTDecorator_BlueprintBase::GetShouldAbort(UBehaviorTreeComponent& OwnerComp) const 
 {
-	CurrentCallResult = false;
+	// if there's no condition-checking function implemented we always want to abort on any change
+	if (ReceiveConditionCheckImplementations == 0)
+	{
+		return true;
+	}
+
+	const bool bIsOnActiveBranch = OwnerComp.IsExecutingBranch(GetMyNode(), GetChildIndex());
+
+	bool bShouldAbort = false;
+	if (bIsOnActiveBranch)
+	{
+		bShouldAbort = (FlowAbortMode == EBTFlowAbortMode::Self || FlowAbortMode == EBTFlowAbortMode::Both) && CalculateRawConditionValueImpl(OwnerComp) == IsInversed();
+	}
+	else
+	{
+		bShouldAbort = (FlowAbortMode == EBTFlowAbortMode::LowerPriority || FlowAbortMode == EBTFlowAbortMode::Both) && CalculateRawConditionValueImpl(OwnerComp) != IsInversed();
+	}
+
+	return bShouldAbort;
+}
+
+bool UBTDecorator_BlueprintBase::CalculateRawConditionValueImpl(UBehaviorTreeComponent& OwnerComp) const
+{
+	bool CurrentCallResult = false;
 	if (ReceiveConditionCheckImplementations != 0)
 	{
 		// can't use const functions with blueprints
@@ -164,20 +218,20 @@ bool UBTDecorator_BlueprintBase::CalculateRawConditionValue(UBehaviorTreeCompone
 
 		if (AIOwner != nullptr && ReceiveConditionCheckImplementations & FBTNodeBPImplementationHelper::AISpecific)
 		{
-			MyNode->ReceiveConditionCheckAI(MyNode->AIOwner, MyNode->AIOwner->GetPawn());
+			CurrentCallResult = MyNode->ReceiveConditionCheckAI(MyNode->AIOwner, MyNode->AIOwner->GetPawn());
 		}
 		else if (ReceiveConditionCheckImplementations & FBTNodeBPImplementationHelper::Generic)
 		{
-			MyNode->ReceiveConditionCheck(MyNode->ActorOwner);
+			CurrentCallResult = MyNode->ReceiveConditionCheck(MyNode->ActorOwner);
 		}
 	}
 
 	return CurrentCallResult;
 }
 
-void UBTDecorator_BlueprintBase::FinishConditionCheck(bool bAllowExecution)
+bool UBTDecorator_BlueprintBase::CalculateRawConditionValue(UBehaviorTreeComponent& OwnerComp, uint8* NodeMemory) const
 {
-	CurrentCallResult = bAllowExecution;
+	return CalculateRawConditionValueImpl(OwnerComp);	
 }
 
 bool UBTDecorator_BlueprintBase::IsDecoratorExecutionActive() const
@@ -223,7 +277,7 @@ void UBTDecorator_BlueprintBase::DescribeRuntimeValues(const UBehaviorTreeCompon
 void UBTDecorator_BlueprintBase::OnBlackboardChange(const UBlackboardComponent& Blackboard, FBlackboard::FKey ChangedKeyID)
 {
 	UBehaviorTreeComponent* BehaviorComp = (UBehaviorTreeComponent*)Blackboard.GetBrainComponent();
-	if (BehaviorComp)
+	if (BehaviorComp && GetShouldAbort(*BehaviorComp))
 	{
 		BehaviorComp->RequestExecution(this);		
 	}
diff --git a/Engine/Source/Runtime/AIModule/Public/BlueprintNodeHelpers.h b/Engine/Source/Runtime/AIModule/Public/BlueprintNodeHelpers.h
--- a/Engine/Source/Runtime/AIModule/Public/BlueprintNodeHelpers.h
+++ b/Engine/Source/Runtime/AIModule/Public/BlueprintNodeHelpers.h
@@ -11,6 +11,7 @@ namespace BlueprintNodeHelpers
 	uint16 GetPropertiesMemorySize(const TArray<UProperty*>& PropertyData);
 
 	void CollectBlackboardSelectors(const UObject* Ob, const UClass* StopAtClass, TArray<FName>& KeyNames);
+	bool HasAnyBlackboardSelectors(const UObject* Ob, const UClass* StopAtClass);
 
 	AIMODULE_API FString DescribeProperty(const UProperty* Prop, const uint8* PropertyAddr);
 	void DescribeRuntimeValues(const UObject* Ob, const TArray<UProperty*>& PropertyData, TArray<FString>& RuntimeValues);
