Introducing virtual AIController::GetFocalPointOnActor to control which specific location is AI focusing on while given Actor is its focus #UE4

- also, based position from FFocusKnowledge::FFocusItem has been removed in favour of regular FVector
- github PR#632, by slonopotamus

diff --git a/Engine/Source/Runtime/AIModule/Classes/AIController.h b/Engine/Source/Runtime/AIModule/Classes/AIController.h
--- a/Engine/Source/Runtime/AIModule/Classes/AIController.h
+++ b/Engine/Source/Runtime/AIModule/Classes/AIController.h
@@ -16,7 +16,6 @@ class APawn;
 class UPathFollowingComponent;
 class UBrainComponent;
 class UAIPerceptionComponent;
-struct FBasedPosition;
 class UPawnAction;
 class UPawnActionsComponent;
 class UNavigationQueryFilter;
@@ -30,6 +29,9 @@ struct FPathFindingQuery;
 
 DECLARE_DYNAMIC_MULTICAST_DELEGATE_TwoParams(FAIMoveCompletedSignature, FAIRequestID, RequestID, EPathFollowingResult::Type, Result);
 
+// the reason for this being namespace instead of a regular enum is
+// so that it can be expanded in game-specific code
+// @todo this is a bit messy, needs to be refactored
 namespace EAIFocusPriority
 {
 	typedef uint8 Type;
@@ -46,19 +48,15 @@ struct FFocusKnowledge
 	struct FFocusItem
 	{
 		TWeakObjectPtr<AActor> Actor;
-		FBasedPosition Position;    
+		FVector Position;
 
-		FVector GetLocation() const 
+		FFocusItem()
 		{
-			const AActor* FocusActor = Actor.Get();
-			return FocusActor ? FocusActor->GetActorLocation() : *Position;
+			Actor = nullptr;
+			Position = FAISystem::InvalidLocation;
 		}
 	};
-
-	FFocusKnowledge() 
-	{
-		Priorities.Reserve(6);
-	}
+	
 	TArray<FFocusItem> Priorities;
 };
 
@@ -120,7 +118,7 @@ private_subobject:
 	UPawnActionsComponent* ActionsComp;
 
 public:
-	
+
 	AAIController(const FObjectInitializer& ObjectInitializer);
 
 	/** Event called when PossessedPawn is possesed by this controller. */
@@ -199,7 +197,7 @@ public:
 	/** Blueprint notification that we've completed the current movement request */
 	UPROPERTY(BlueprintAssignable, meta=(DisplayName="MoveCompleted"))
 	FAIMoveCompletedSignature ReceiveMoveCompleted;
-	
+
 	/** Returns status of path following */
 	UFUNCTION(BlueprintCallable, Category="AI|Navigation")
 	EPathFollowingStatus::Type GetMoveStatus() const;
@@ -236,14 +234,17 @@ protected:
 public:
 	/** Retrieve the final position that controller should be looking at. */
 	UFUNCTION(BlueprintCallable, Category="AI")
-	virtual FVector GetFocalPoint() const;
+	FVector GetFocalPoint() const;
 
-	FVector GetFocalPoint(EAIFocusPriority::Type Priority) const;
-	FORCEINLINE FFocusKnowledge::FFocusItem GetFocusItem(EAIFocusPriority::Type Priority) const { return FocusInformation.Priorities.IsValidIndex(Priority) ? FocusInformation.Priorities[Priority] : FFocusKnowledge::FFocusItem(); }
-	
-	/** Set FocalPoint as absolute position or offset from base. */
+	FVector GetFocalPointForPriority(EAIFocusPriority::Type InPriority) const;
+
+	/** Retrieve the focal point this controller should focus to on given actor. */
+	UFUNCTION(BlueprintCallable, Category="AI")
+	virtual FVector GetFocalPointOnActor(const AActor *Actor) const;
+
+	/** Set the position that controller should be looking at. */
 	UFUNCTION(BlueprintCallable, Category="AI", meta=(FriendlyName="SetFocalPoint"))
-	void K2_SetFocalPoint(FVector FP, bool bOffsetFromBase = false);
+	void K2_SetFocalPoint(FVector FP);
 
 	/** Set Focus for actor, will set FocalPoint as a result. */
 	UFUNCTION(BlueprintCallable, Category="AI", meta=(FriendlyName="SetFocus"))
@@ -253,7 +254,7 @@ public:
 	UFUNCTION(BlueprintCallable, Category="AI")
 	AActor* GetFocusActor() const;
 
-	FORCEINLINE AActor* GetFocusActor(EAIFocusPriority::Type Priority) const {  return FocusInformation.Priorities.IsValidIndex(Priority) ? FocusInformation.Priorities[Priority].Actor.Get() : NULL; }
+	FORCEINLINE AActor* GetFocusActorForPriority(EAIFocusPriority::Type InPriority) const {  return FocusInformation.Priorities.IsValidIndex(InPriority) ? FocusInformation.Priorities[InPriority].Actor.Get() : nullptr; }
 
 	/** Clears Focus, will also clear FocalPoint as a result */
 	UFUNCTION(BlueprintCallable, Category="AI", meta=(FriendlyName="ClearFocus"))
@@ -309,7 +310,7 @@ public:
 	virtual void UpdateControlRotation(float DeltaTime, bool bUpdatePawn = true);
 
 	/** Set FocalPoint for given priority as absolute position or offset from base. */
-	virtual void SetFocalPoint(FVector FP, bool bOffsetFromBase=false, uint8 InPriority=EAIFocusPriority::Gameplay);
+	virtual void SetFocalPoint(FVector NewFocus, EAIFocusPriority::Type InPriority=EAIFocusPriority::Gameplay);
 
 	/* Set Focus actor for given priority, will set FocalPoint as a result. */
 	virtual void SetFocus(AActor* NewFocus, EAIFocusPriority::Type InPriority = EAIFocusPriority::Gameplay);
diff --git a/Engine/Source/Runtime/AIModule/Classes/BehaviorTree/Services/BTService_DefaultFocus.h b/Engine/Source/Runtime/AIModule/Classes/BehaviorTree/Services/BTService_DefaultFocus.h
--- a/Engine/Source/Runtime/AIModule/Classes/BehaviorTree/Services/BTService_DefaultFocus.h
+++ b/Engine/Source/Runtime/AIModule/Classes/BehaviorTree/Services/BTService_DefaultFocus.h
@@ -12,7 +12,9 @@ struct FBTFocusMemory
 
 	void Reset()
 	{
-		FMemory::MemZero(*this);
+		FocusActorSet = nullptr;
+		FocusLocationSet = FAISystem::InvalidLocation;
+		bActorSet = false;
 	}
 };
 
diff --git a/Engine/Source/Runtime/AIModule/Private/AIController.cpp b/Engine/Source/Runtime/AIModule/Private/AIController.cpp
--- a/Engine/Source/Runtime/AIModule/Private/AIController.cpp
+++ b/Engine/Source/Runtime/AIModule/Private/AIController.cpp
@@ -24,11 +24,13 @@
 #define MEDSIGHTTHRESHOLDSQUARED		(MEDSIGHTTHRESHOLD*MEDSIGHTTHRESHOLD)
 #define FARSIGHTTHRESHOLDSQUARED		(FARSIGHTTHRESHOLD*FARSIGHTTHRESHOLD)
 
-
+//----------------------------------------------------------------------//
+// AAIController
+//----------------------------------------------------------------------//
 bool AAIController::bAIIgnorePlayers = false;
 
-DECLARE_CYCLE_STAT(TEXT("MoveToLocation"),STAT_MoveToLocation,STATGROUP_AI);
-DECLARE_CYCLE_STAT(TEXT("MoveToActor"),STAT_MoveToActor,STATGROUP_AI);
+DECLARE_CYCLE_STAT(TEXT("MoveToLocation"), STAT_MoveToLocation, STATGROUP_AI);
+DECLARE_CYCLE_STAT(TEXT("MoveToActor"), STAT_MoveToActor, STATGROUP_AI);
 
 DEFINE_LOG_CATEGORY(LogAINavigation);
 
@@ -94,10 +96,11 @@ void AAIController::DisplayDebug(UCanvas* Canvas, const FDebugDisplayInfo& Debug
 		{
 			PathFollowingComponent->DisplayDebug(Canvas, DebugDisplay, YL, YPos);
 		}
-		
-		if ( GetFocusActor() )
+
+		AActor* FocusActor = GetFocusActor();
+		if (FocusActor)
 		{
-			Canvas->DrawText(GEngine->GetSmallFont(), FString::Printf(TEXT("      Focus %s"), *GetFocusActor()->GetName()), 4.0f, YPos);
+			Canvas->DrawText(GEngine->GetSmallFont(), FString::Printf(TEXT("      Focus %s"), *FocusActor->GetName()), 4.0f, YPos);
 			YPos += YL;
 		}
 	}
@@ -112,7 +115,8 @@ void AAIController::GrabDebugSnapshot(FVisualLogEntry* Snapshot) const
 	MyCategory.Add(TEXT("Pawn"), GetNameSafe(GetPawn()));
 	AActor* FocusActor = GetFocusActor();
 	MyCategory.Add(TEXT("Focus"), GetDebugName(FocusActor));
-	if (FocusActor == NULL)
+
+	if (FocusActor == nullptr)
 	{
 		MyCategory.Add(TEXT("Focus Location"), GetFocalPoint().ToString());
 	}
@@ -154,89 +158,76 @@ void AAIController::GetPlayerViewPoint(FVector& out_Location, FRotator& out_Rota
 	}
 }
 
-void AAIController::SetFocalPoint( FVector FP, bool bOffsetFromBase, uint8 InPriority)
+void AAIController::SetFocalPoint(FVector NewFocus, EAIFocusPriority::Type InPriority)
 {
 	if (InPriority >= FocusInformation.Priorities.Num())
 	{
-		FocusInformation.Priorities.SetNum(InPriority+1);
-	}
-	FFocusKnowledge::FFocusItem& Focusitem = FocusInformation.Priorities[InPriority];
-
-	{
-		AActor* FocalActor = NULL;
-		if (bOffsetFromBase)
-		{
-			APawn *Pawn = GetPawn();
-			if (Pawn && Pawn->GetMovementBase())
-			{
-				FocalActor = Pawn->GetMovementBase()->GetOwner();
-			}
-		}
-		Focusitem.Position.Set( FocalActor, FP );
+		FocusInformation.Priorities.SetNum(InPriority + 1);
 	}
+	FFocusKnowledge::FFocusItem& FocusItem = FocusInformation.Priorities[InPriority];
 
-	Focusitem.Actor = NULL;
+	FocusItem.Actor = nullptr;
+	FocusItem.Position = NewFocus;
 }
 
-FVector AAIController::GetFocalPoint(EAIFocusPriority::Type Priority) const
-{	
-	FBasedPosition FocalPointForPriority;
-
-	const FFocusKnowledge::FFocusItem& FocusItem = GetFocusItem(Priority);
+FVector AAIController::GetFocalPointForPriority(EAIFocusPriority::Type InPriority) const
+{
+	FVector Result = FAISystem::InvalidLocation;
 
-	if (FocusItem.Actor.IsValid())
-	{
-		const AActor* Focus = FocusItem.Actor.Get();
-		UPrimitiveComponent* MyBase = GetPawn() ? GetPawn()->GetMovementBase() : NULL;
-		const bool bRequestedFocusIsBased = MyBase && Cast<const APawn>(Focus) && (Cast<const APawn>(Focus)->GetMovementBase() == MyBase);
-		FocalPointForPriority.Set(bRequestedFocusIsBased && MyBase ? MyBase->GetOwner() : NULL, Focus->GetActorLocation());
-	}
-	else if (!(*FocusItem.Position).IsZero())
+	if (InPriority < FocusInformation.Priorities.Num())
 	{
-		FocalPointForPriority = FocusItem.Position;
+		const FFocusKnowledge::FFocusItem& FocusItem = FocusInformation.Priorities[InPriority];
+
+		AActor* FocusActor = FocusItem.Actor.Get();
+		if (FocusActor)
+		{
+			Result = GetFocalPointOnActor(FocusActor);
+		}
+		else
+		{
+			Result = FocusItem.Position;
+		}
 	}
 
-	return *FocalPointForPriority;
+	return Result;
 }
 
 FVector AAIController::GetFocalPoint() const
 {
-	FBasedPosition FinalFocus;
+	FVector Result = FAISystem::InvalidLocation;
 
 	// find focus with highest priority
-	for( int32 Index = FocusInformation.Priorities.Num()-1; Index >= 0; --Index)
+	for (int32 Index = FocusInformation.Priorities.Num() - 1; Index >= 0; --Index)
 	{
-		const FFocusKnowledge::FFocusItem& Focusitem = FocusInformation.Priorities[Index];
-		if ( Focusitem.Actor.IsValid() )
+		const FFocusKnowledge::FFocusItem& FocusItem = FocusInformation.Priorities[Index];
+		AActor* FocusActor = FocusItem.Actor.Get();
+		if (FocusActor)
 		{
-			const AActor* Focus = Focusitem.Actor.Get();
-			UPrimitiveComponent* MyBase = GetPawn() ? GetPawn()->GetMovementBase() : NULL;
-			const bool bRequestedFocusIsBased = MyBase && Cast<const APawn>(Focus) && (Cast<const APawn>(Focus)->GetMovementBase() == MyBase);
-			FinalFocus.Set(bRequestedFocusIsBased && MyBase ? MyBase->GetOwner() : NULL, Focus->GetActorLocation());
+			Result = GetFocalPointOnActor(FocusActor);
 			break;
 		}
-		else if( !(*Focusitem.Position).IsZero() )
+		else if (FAISystem::IsValidLocation(FocusItem.Position))
 		{
-			FinalFocus = Focusitem.Position;
+			Result = FocusItem.Position;
 			break;
 		}
 	}
 
-	return *FinalFocus;
+	return Result;
 }
 
 AActor* AAIController::GetFocusActor() const
 {
-	AActor* FocusActor = NULL;
-	for( int32 Index = FocusInformation.Priorities.Num()-1; Index >= 0; --Index)
+	AActor* FocusActor = nullptr;
+	for (int32 Index = FocusInformation.Priorities.Num() - 1; Index >= 0; --Index)
 	{
-		const FFocusKnowledge::FFocusItem& Focusitem = FocusInformation.Priorities[Index];
-		if ( Focusitem.Actor.IsValid() )
+		const FFocusKnowledge::FFocusItem& FocusItem = FocusInformation.Priorities[Index];
+		FocusActor = FocusItem.Actor.Get();
+		if (FocusActor)
 		{
-			FocusActor = Focusitem.Actor.Get();
 			break;
 		}
-		else if( !(*Focusitem.Position).IsZero() )
+		else if (FAISystem::IsValidLocation(FocusItem.Position))
 		{
 			break;
 		}
@@ -245,14 +236,19 @@ AActor* AAIController::GetFocusActor() const
 	return FocusActor;
 }
 
+FVector AAIController::GetFocalPointOnActor(const AActor *Actor) const
+{
+	return Actor != nullptr ? Actor->GetActorLocation() : FAISystem::InvalidLocation;
+}
+
 void AAIController::K2_SetFocus(AActor* NewFocus)
-{ 
-	SetFocus(NewFocus, EAIFocusPriority::Gameplay); 
+{
+	SetFocus(NewFocus, EAIFocusPriority::Gameplay);
 }
 
-void AAIController::K2_SetFocalPoint(FVector FP, bool bOffsetFromBase)
+void AAIController::K2_SetFocalPoint(FVector NewFocus)
 {
-	SetFocalPoint(FP, bOffsetFromBase, EAIFocusPriority::Gameplay);
+	SetFocalPoint(NewFocus, EAIFocusPriority::Gameplay);
 }
 
 void AAIController::K2_ClearFocus()
@@ -262,11 +258,11 @@ void AAIController::K2_ClearFocus()
 
 void AAIController::SetFocus(AActor* NewFocus, EAIFocusPriority::Type InPriority)
 {
-	if( NewFocus )
+	if (NewFocus)
 	{
 		if (InPriority >= FocusInformation.Priorities.Num())
 		{
-			FocusInformation.Priorities.SetNum(InPriority+1);
+			FocusInformation.Priorities.SetNum(InPriority + 1);
 		}
 		FocusInformation.Priorities[InPriority].Actor = NewFocus;
 	}
@@ -280,23 +276,23 @@ void AAIController::ClearFocus(EAIFocusPriority::Type InPriority)
 {
 	if (InPriority < FocusInformation.Priorities.Num())
 	{
-		FocusInformation.Priorities[InPriority].Actor = NULL;
-		FocusInformation.Priorities[InPriority].Position.Clear();
+		FocusInformation.Priorities[InPriority].Actor = nullptr;
+		FocusInformation.Priorities[InPriority].Position = FAISystem::InvalidLocation;
 	}
 }
 
 bool AAIController::LineOfSightTo(const AActor* Other, FVector ViewPoint, bool bAlternateChecks) const
 {
-	if( !Other )
+	if (Other == nullptr)
 	{
 		return false;
 	}
-		
-	if ( ViewPoint.IsZero() )
+
+	if (ViewPoint.IsZero())
 	{
 		AActor*	ViewTarg = GetViewTarget();
 		ViewPoint = ViewTarg->GetActorLocation();
-		if( ViewTarg == GetPawn() )
+		if (ViewTarg == GetPawn())
 		{
 			ViewPoint.Z += GetPawn()->BaseEyeHeight; //look from eyes
 		}
@@ -309,7 +305,7 @@ bool AAIController::LineOfSightTo(const AActor* Other, FVector ViewPoint, bool b
 	CollisionParams.AddIgnoredActor(Other);
 
 	bool bHit = GetWorld()->LineTraceTest(ViewPoint, TargetLocation, ECC_Visibility, CollisionParams);
-	if( !bHit )
+	if (!bHit)
 	{
 		return true;
 	}
@@ -324,12 +320,12 @@ bool AAIController::LineOfSightTo(const AActor* Other, FVector ViewPoint, bool b
 
 	const FVector OtherActorLocation = Other->GetActorLocation();
 	const float DistSq = (OtherActorLocation - ViewPoint).SizeSquared();
-	if ( DistSq > FARSIGHTTHRESHOLDSQUARED )
+	if (DistSq > FARSIGHTTHRESHOLDSQUARED)
 	{
 		return false;
 	}
 
-	if ( !OtherPawn && (DistSq > NEARSIGHTTHRESHOLDSQUARED) ) 
+	if (!OtherPawn && (DistSq > NEARSIGHTTHRESHOLDSQUARED))
 	{
 		return false;
 	}
@@ -337,20 +333,20 @@ bool AAIController::LineOfSightTo(const AActor* Other, FVector ViewPoint, bool b
 	float OtherRadius, OtherHeight;
 	Other->GetSimpleCollisionCylinder(OtherRadius, OtherHeight);
 
-	if ( !bAlternateChecks || !bLOSflag )
+	if (!bAlternateChecks || !bLOSflag)
 	{
 		//try viewpoint to head
-		bHit = GetWorld()->LineTraceTest(ViewPoint,  OtherActorLocation + FVector(0.f,0.f,OtherHeight), ECC_Visibility, CollisionParams);
-		if ( !bHit )
+		bHit = GetWorld()->LineTraceTest(ViewPoint, OtherActorLocation + FVector(0.f, 0.f, OtherHeight), ECC_Visibility, CollisionParams);
+		if (!bHit)
 		{
 			return true;
 		}
 	}
 
-	if( !bSkipExtraLOSChecks && (!bAlternateChecks || bLOSflag) )
+	if (!bSkipExtraLOSChecks && (!bAlternateChecks || bLOSflag))
 	{
 		// only check sides if width of other is significant compared to distance
-		if( OtherRadius * OtherRadius/(OtherActorLocation - ViewPoint).SizeSquared() < 0.0001f )
+		if (OtherRadius * OtherRadius / (OtherActorLocation - ViewPoint).SizeSquared() < 0.0001f)
 		{
 			return false;
 		}
@@ -364,7 +360,7 @@ bool AAIController::LineOfSightTo(const AActor* Other, FVector ViewPoint, bool b
 		int32 IndexMax = 0;
 		float CurrentMax = (Points[0] - ViewPoint).SizeSquared();
 		float CurrentMin = CurrentMax;
-		for ( int32 PointIndex=1; PointIndex<4; PointIndex++ )
+		for (int32 PointIndex = 1; PointIndex<4; PointIndex++)
 		{
 			const float NextSize = (Points[PointIndex] - ViewPoint).SizeSquared();
 			if (NextSize > CurrentMin)
@@ -379,12 +375,12 @@ bool AAIController::LineOfSightTo(const AActor* Other, FVector ViewPoint, bool b
 			}
 		}
 
-		for ( int32 PointIndex=0; PointIndex<4; PointIndex++ )
+		for (int32 PointIndex = 0; PointIndex<4; PointIndex++)
 		{
 			if ((PointIndex != IndexMin) && (PointIndex != IndexMax))
 			{
-				bHit = GetWorld()->LineTraceTest(ViewPoint,  Points[PointIndex], ECC_Visibility, CollisionParams);
-				if ( !bHit )
+				bHit = GetWorld()->LineTraceTest(ViewPoint, Points[PointIndex], ECC_Visibility, CollisionParams);
+				if (!bHit)
 				{
 					return true;
 				}
@@ -405,14 +401,14 @@ void AAIController::UpdateControlRotation(float DeltaTime, bool bUpdatePawn)
 {
 	// Look toward focus
 	FVector FocalPoint = GetFocalPoint();
-	if( !FocalPoint.IsZero() && GetPawn())
+	if (FAISystem::IsValidLocation(FocalPoint) && GetPawn())
 	{
 		FVector Direction = FocalPoint - GetPawn()->GetActorLocation();
 		FRotator NewControlRotation = Direction.Rotation();
 
 		// Don't pitch view of walking pawns unless looking at another pawn
-		if ( GetPawn()->GetMovementComponent() && GetPawn()->GetMovementComponent()->IsMovingOnGround() &&
-			PathFollowingComponent && (!PathFollowingComponent->GetMoveGoal() || !Cast<APawn>(PathFollowingComponent->GetMoveGoal()) ) )
+		if (GetPawn()->GetMovementComponent() && GetPawn()->GetMovementComponent()->IsMovingOnGround() &&
+			PathFollowingComponent && (!PathFollowingComponent->GetMoveGoal() || !Cast<APawn>(PathFollowingComponent->GetMoveGoal())))
 		{
 			NewControlRotation.Pitch = 0.f;
 		}
@@ -440,7 +436,7 @@ void AAIController::Possess(APawn* InPawn)
 	{
 		return;
 	}
-	
+
 	// no point in doing navigation setup if pawn has no movement component
 	const UPawnMovementComponent* MovementComp = InPawn->GetMovementComponent();
 	if (MovementComp != NULL)
@@ -536,12 +532,12 @@ EPathFollowingRequestResult::Type AAIController::MoveToLocation(const FVector& D
 		*Dest.ToString(), AcceptanceRadius, bStopOnOverlap ? TEXT(" + agent") : TEXT(""), bUsePathfinding, bCanStrafe, *GetNameSafe(FilterClass));
 
 	// Check input is valid
-	if( Dest.ContainsNaN() )
+	if (Dest.ContainsNaN())
 	{
 		UE_VLOG(this, LogAINavigation, Error, TEXT("AAIController::MoveToLocation: Destination is not valid! Goal(%s) AcceptRadius(%.1f%s) bUsePathfinding(%d) bCanStrafe(%d)"),
 			*Dest.ToString(), AcceptanceRadius, bStopOnOverlap ? TEXT(" + agent") : TEXT(""), bUsePathfinding, bCanStrafe);
-		
-		ensure( !Dest.ContainsNaN() );
+
+		ensure(!Dest.ContainsNaN());
 		bCanRequestMove = false;
 	}
 
@@ -661,7 +657,7 @@ bool AAIController::PreparePathfinding(FPathFindingQuery& Query, const FVector&
 		}
 
 		Query = FPathFindingQuery(this, NavData, GetNavAgentLocation(), GoalLocation, UNavigationQueryFilter::GetQueryFilter(NavData, FilterClass));
-		
+
 		if (PathFollowingComponent)
 		{
 			PathFollowingComponent->OnPathfindingQuery(Query);
@@ -748,7 +744,7 @@ bool AAIController::RunBehaviorTree(UBehaviorTree* BTAsset)
 		bSuccess = UseBlackboard(BTAsset->BlackboardAsset);
 		BlackboardComp = FindComponentByClass<UBlackboardComponent>();
 	}
-	
+
 	if (bSuccess)
 	{
 		UBehaviorTreeComponent* BTComp = Cast<UBehaviorTreeComponent>(BrainComponent);
diff --git a/Engine/Source/Runtime/AIModule/Private/BehaviorTree/Services/BTService_DefaultFocus.cpp b/Engine/Source/Runtime/AIModule/Private/BehaviorTree/Services/BTService_DefaultFocus.cpp
--- a/Engine/Source/Runtime/AIModule/Private/BehaviorTree/Services/BTService_DefaultFocus.cpp
+++ b/Engine/Source/Runtime/AIModule/Private/BehaviorTree/Services/BTService_DefaultFocus.cpp
@@ -46,7 +46,7 @@ void UBTService_DefaultFocus::OnBecomeRelevant(UBehaviorTreeComponent& OwnerComp
 		else
 		{
 			const FVector FocusLocation = MyBlackboard->GetValueAsVector(BlackboardKey.GetSelectedKeyID());
-			OwnerController->SetFocalPoint(FocusLocation, /*bOffsetFromBase=*/false, EAIFocusPriority::Default);
+			OwnerController->SetFocalPoint(FocusLocation, EAIFocusPriority::Default);
 			MyMemory->FocusLocationSet = FocusLocation;
 		}
 	}
@@ -64,11 +64,11 @@ void UBTService_DefaultFocus::OnCeaseRelevant(UBehaviorTreeComponent& OwnerComp,
 		bool bClearFocus = false;
 		if (MyMemory->bActorSet)
 		{
-			bClearFocus = (MyMemory->FocusActorSet == OwnerController->GetFocusActor(EAIFocusPriority::Default));
+			bClearFocus = (MyMemory->FocusActorSet == OwnerController->GetFocusActorForPriority(EAIFocusPriority::Default));
 		}
 		else
 		{
-			bClearFocus = (MyMemory->FocusLocationSet == OwnerController->GetFocalPoint(EAIFocusPriority::Default));
+			bClearFocus = (MyMemory->FocusLocationSet == OwnerController->GetFocalPointForPriority(EAIFocusPriority::Default));
 		}
 
 		if (bClearFocus)
@@ -102,4 +102,4 @@ FName UBTService_DefaultFocus::GetNodeIconName() const
 	return FName("BTEditor.Graph.BTNode.Service.DefaultFocus.Icon");
 }
 
-#endif	// WITH_EDITOR
\ No newline at end of file
+#endif	// WITH_EDITOR
diff --git a/Engine/Source/Runtime/AIModule/Private/BehaviorTree/Tasks/BTTask_RotateToFaceBBEntry.cpp b/Engine/Source/Runtime/AIModule/Private/BehaviorTree/Tasks/BTTask_RotateToFaceBBEntry.cpp
--- a/Engine/Source/Runtime/AIModule/Private/BehaviorTree/Tasks/BTTask_RotateToFaceBBEntry.cpp
+++ b/Engine/Source/Runtime/AIModule/Private/BehaviorTree/Tasks/BTTask_RotateToFaceBBEntry.cpp
@@ -96,7 +96,7 @@ EBTNodeResult::Type UBTTask_RotateToFaceBBEntry::ExecuteTask(UBehaviorTreeCompon
 			}
 			else
 			{
-				AIController->SetFocalPoint(KeyValue, false, EAIFocusPriority::Gameplay);
+				AIController->SetFocalPoint(KeyValue, EAIFocusPriority::Gameplay);
 				MyMemory->FocusLocationSet = KeyValue;
 				Result = EBTNodeResult::InProgress;
 			}
@@ -119,7 +119,7 @@ EBTNodeResult::Type UBTTask_RotateToFaceBBEntry::ExecuteTask(UBehaviorTreeCompon
 			{
 				const FVector FocalPoint = PawnLocation + DirectionVector * (MAX_FLT / 2);
 				// set focal somewhere far in the indicated direction
-				AIController->SetFocalPoint(FocalPoint, false, EAIFocusPriority::Gameplay);
+				AIController->SetFocalPoint(FocalPoint, EAIFocusPriority::Gameplay);
 				MyMemory->FocusLocationSet = FocalPoint;
 				Result = EBTNodeResult::InProgress;
 			}
@@ -140,7 +140,7 @@ void UBTTask_RotateToFaceBBEntry::TickTask(UBehaviorTreeComponent& OwnerComp, ui
 	else
 	{
 		const FVector PawnDirection = AIController->GetPawn()->GetActorRotation().Vector();				
-		const FVector FocalPoint = AIController->GetFocalPoint(EAIFocusPriority::Gameplay);
+		const FVector FocalPoint = AIController->GetFocalPointForPriority(EAIFocusPriority::Gameplay);
 
 		if (CalculateAngleDifferenceDot(PawnDirection, (FocalPoint - AIController->GetPawn()->GetActorLocation()).GetSafeNormal2D()) >= PrecisionDot)
 		{
@@ -158,11 +158,11 @@ void UBTTask_RotateToFaceBBEntry::CleanUp(AAIController& AIController, uint8* No
 	bool bClearFocus = false;
 	if (MyMemory->bActorSet)
 	{
-		bClearFocus = (MyMemory->FocusActorSet == AIController.GetFocusActor(EAIFocusPriority::Gameplay));
+		bClearFocus = (MyMemory->FocusActorSet == AIController.GetFocusActorForPriority(EAIFocusPriority::Gameplay));
 	}
 	else
 	{
-		bClearFocus = (MyMemory->FocusLocationSet == AIController.GetFocalPoint(EAIFocusPriority::Gameplay));
+		bClearFocus = (MyMemory->FocusLocationSet == AIController.GetFocalPointForPriority(EAIFocusPriority::Gameplay));
 	}
 	
 	if (bClearFocus)
@@ -193,7 +193,7 @@ void UBTTask_RotateToFaceBBEntry::DescribeRuntimeValues(const UBehaviorTreeCompo
 	if (AIController != NULL && AIController->GetPawn() != NULL)
 	{
 		const FVector PawnDirection = AIController->GetPawn()->GetActorRotation().Vector();
-		const FVector FocalPoint = AIController->GetFocusItem(EAIFocusPriority::Gameplay).GetLocation();
+		const FVector FocalPoint = AIController->GetFocalPointForPriority(EAIFocusPriority::Gameplay);
 
 		const float CurrentAngleRadians = CalculateAngleDifferenceDot(PawnDirection, (FocalPoint - AIController->GetPawn()->GetActorLocation()).GetSafeNormal2D());
 		Values.Add(FString::Printf(TEXT("Current angle: %.2f"), FMath::RadiansToDegrees(FMath::Acos(CurrentAngleRadians))));
diff --git a/Engine/Source/Runtime/AIModule/Private/Navigation/PathFollowingComponent.cpp b/Engine/Source/Runtime/AIModule/Private/Navigation/PathFollowingComponent.cpp
--- a/Engine/Source/Runtime/AIModule/Private/Navigation/PathFollowingComponent.cpp
+++ b/Engine/Source/Runtime/AIModule/Private/Navigation/PathFollowingComponent.cpp
@@ -1161,7 +1161,7 @@ void UPathFollowingComponent::UpdateMoveFocus()
 		if (Status == EPathFollowingStatus::Moving)
 		{
 			const FVector MoveFocus = GetMoveFocus(AIOwner->bAllowStrafe);
-			AIOwner->SetFocalPoint(MoveFocus, false, EAIFocusPriority::Move);
+			AIOwner->SetFocalPoint(MoveFocus, EAIFocusPriority::Move);
 		}
 		else
 		{
