Fixed AIPerceptionSystem force-starting engine's MakeNoise implementation #UE4

- addresses long broken UE-4814

diff --git a/Engine/Source/Runtime/AIModule/Private/Perception/AIPerceptionSystem.cpp b/Engine/Source/Runtime/AIModule/Private/Perception/AIPerceptionSystem.cpp
--- a/Engine/Source/Runtime/AIModule/Private/Perception/AIPerceptionSystem.cpp
+++ b/Engine/Source/Runtime/AIModule/Private/Perception/AIPerceptionSystem.cpp
@@ -30,10 +30,6 @@ UAIPerceptionSystem::UAIPerceptionSystem(const FObjectInitializer& ObjectInitial
 	, PerceptionAgingRate(0.3f)
 	, CurrentTime(0.f)
 {
-	if (HasAnyFlags(RF_ClassDefaultObject) == true)
-	{
-		AActor::SetMakeNoiseDelegate(FMakeNoiseDelegate::CreateStatic(&UAIPerceptionSystem::MakeNoiseImpl));
-	}
 }
 
 void UAIPerceptionSystem::RegisterSenseClass(TSubclassOf<UAISense> SenseClass)
diff --git a/Engine/Source/Runtime/AIModule/Private/Perception/AISense_Hearing.cpp b/Engine/Source/Runtime/AIModule/Private/Perception/AISense_Hearing.cpp
--- a/Engine/Source/Runtime/AIModule/Private/Perception/AISense_Hearing.cpp
+++ b/Engine/Source/Runtime/AIModule/Private/Perception/AISense_Hearing.cpp
@@ -61,6 +61,13 @@ UAISense_Hearing::UAISense_Hearing(const FObjectInitializer& ObjectInitializer)
 		OnNewListenerDelegate.BindUObject(this, &UAISense_Hearing::OnNewListenerImpl);
 		OnListenerUpdateDelegate.BindUObject(this, &UAISense_Hearing::OnListenerUpdateImpl);
 		OnListenerRemovedDelegate.BindUObject(this, &UAISense_Hearing::OnListenerRemovedImpl);
+
+		static bool bMakeNoiseInterceptionSetUp = false;
+		if (bMakeNoiseInterceptionSetUp == false)
+		{
+			AActor::SetMakeNoiseDelegate(FMakeNoiseDelegate::CreateStatic(&UAIPerceptionSystem::MakeNoiseImpl));
+			bMakeNoiseInterceptionSetUp = true;
+		}
 	}
 }
 
