fixed execution result of behavior tree's parallel node when decorators on main task are blocking

diff --git a/Engine/Source/Runtime/AIModule/Private/BehaviorTree/Composites/BTComposite_SimpleParallel.cpp b/Engine/Source/Runtime/AIModule/Private/BehaviorTree/Composites/BTComposite_SimpleParallel.cpp
--- a/Engine/Source/Runtime/AIModule/Private/BehaviorTree/Composites/BTComposite_SimpleParallel.cpp
+++ b/Engine/Source/Runtime/AIModule/Private/BehaviorTree/Composites/BTComposite_SimpleParallel.cpp
@@ -19,6 +19,7 @@ int32 UBTComposite_SimpleParallel::GetNextChildHandler(FBehaviorTreeSearchData&
 	if (PrevChild == BTSpecialChild::NotInitialized)
 	{
 		NextChildIdx = EBTParallelChild::MainTask;
+		MyMemory->MainTaskResult = EBTNodeResult::Failed;
 	}
 	else if ((MyMemory->bMainTaskIsActive || MyMemory->bForceBackgroundTree) && !SearchData.OwnerComp.IsRestartPending())
 	{
