Made sure we're not trying to run EQS queries as result of changes to EQSTestingPawn's defaults #UE4

diff --git a/Engine/Source/Runtime/AIModule/Private/EnvironmentQuery/EQSTestingPawn.cpp b/Engine/Source/Runtime/AIModule/Private/EnvironmentQuery/EQSTestingPawn.cpp
--- a/Engine/Source/Runtime/AIModule/Private/EnvironmentQuery/EQSTestingPawn.cpp
+++ b/Engine/Source/Runtime/AIModule/Private/EnvironmentQuery/EQSTestingPawn.cpp
@@ -145,6 +145,11 @@ void AEQSTestingPawn::PostLoad()
 
 void AEQSTestingPawn::RunEQSQuery()
 {
+	if (HasAnyFlags(RF_ClassDefaultObject) == true)
+	{
+		return;
+	}
+
 	Reset();
 	
 	// make one step if TimeLimitPerStep > 0.f, else all steps
@@ -198,6 +203,10 @@ void AEQSTestingPawn::MakeOneStep()
 void AEQSTestingPawn::UpdateDrawing()
 {
 #if WITH_EDITORONLY_DATA
+	if (HasAnyFlags(RF_ClassDefaultObject) == true)
+	{
+		return;
+	}
 
 	UBillboardComponent* SpriteComponent = FindComponentByClass<UBillboardComponent>();
 	if (SpriteComponent != NULL)
