Fixed BP-implemented EQS generators erroneously complaining about wrong item types #UE4

diff --git a/Engine/Source/Runtime/AIModule/Classes/EnvironmentQuery/Generators/EnvQueryGenerator_BlueprintBase.h b/Engine/Source/Runtime/AIModule/Classes/EnvironmentQuery/Generators/EnvQueryGenerator_BlueprintBase.h
--- a/Engine/Source/Runtime/AIModule/Classes/EnvironmentQuery/Generators/EnvQueryGenerator_BlueprintBase.h
+++ b/Engine/Source/Runtime/AIModule/Classes/EnvironmentQuery/Generators/EnvQueryGenerator_BlueprintBase.h
@@ -41,10 +41,10 @@ class AIMODULE_API UEnvQueryGenerator_BlueprintBase : public UEnvQueryGenerator
 	virtual FText GetDescriptionDetails() const override;
 	
 	UFUNCTION(BlueprintCallable, Category = "EQS")
-	void AddGeneratedVector(FVector GeneratedVector);
+	void AddGeneratedVector(FVector GeneratedVector) const;
 
 	UFUNCTION(BlueprintCallable, Category = "EQS")
-	void AddGeneratedActor(AActor* GeneratedActor);
+	void AddGeneratedActor(AActor* GeneratedActor) const;
 
 	UFUNCTION(BlueprintCallable, Category = "EQS")
 	UObject* GetQuerier() const;
diff --git a/Engine/Source/Runtime/AIModule/Private/EnvironmentQuery/Generators/EnvQueryGenerator_BlueprintBase.cpp b/Engine/Source/Runtime/AIModule/Private/EnvironmentQuery/Generators/EnvQueryGenerator_BlueprintBase.cpp
--- a/Engine/Source/Runtime/AIModule/Private/EnvironmentQuery/Generators/EnvQueryGenerator_BlueprintBase.cpp
+++ b/Engine/Source/Runtime/AIModule/Private/EnvironmentQuery/Generators/EnvQueryGenerator_BlueprintBase.cpp
@@ -43,21 +43,21 @@ void UEnvQueryGenerator_BlueprintBase::GenerateItems(FEnvQueryInstance& QueryIns
 	QueryInstance.PrepareContext(Context, ContextLocations);
 
 	CachedQueryInstance = &QueryInstance;
-	DoItemGeneration(ContextLocations);
+	const_cast<UEnvQueryGenerator_BlueprintBase*>(this)->DoItemGeneration(ContextLocations);
 	CachedQueryInstance = NULL;
 }
 
-void UEnvQueryGenerator_BlueprintBase::AddGeneratedVector(FVector Vector)
+void UEnvQueryGenerator_BlueprintBase::AddGeneratedVector(FVector Vector) const
 {
 	check(CachedQueryInstance);
-	ensure(ItemType->GetClass()->IsChildOf(UEnvQueryItemType_VectorBase::StaticClass()));
+	ensure(ItemType->IsChildOf<UEnvQueryItemType_VectorBase>());
 	CachedQueryInstance->AddItemData<UEnvQueryItemType_Point>(Vector);
 }
 
-void UEnvQueryGenerator_BlueprintBase::AddGeneratedActor(AActor* Actor)
+void UEnvQueryGenerator_BlueprintBase::AddGeneratedActor(AActor* Actor) const
 {
 	check(CachedQueryInstance);
-	//ensure(ItemType->GetClass()->IsChildOf(UEnvQueryItemType_ActorBase::StaticClass()));
+	ensure(ItemType->IsChildOf<UEnvQueryItemType_ActorBase>());
 	CachedQueryInstance->AddItemData<UEnvQueryItemType_Actor>(Actor);
 }
 
