Added mechanics for FAIStimulus auto-expiration after specified time #UE4

diff --git a/Engine/Source/Runtime/AIModule/Classes/Perception/AIPerceptionComponent.h b/Engine/Source/Runtime/AIModule/Classes/Perception/AIPerceptionComponent.h
--- a/Engine/Source/Runtime/AIModule/Classes/Perception/AIPerceptionComponent.h
+++ b/Engine/Source/Runtime/AIModule/Classes/Perception/AIPerceptionComponent.h
@@ -211,7 +211,8 @@ public:
 
 	void RegisterStimulus(AActor* Source, const FAIStimulus& Stimulus);
 	void ProcessStimuli();
-	void AgeStimuli(const float ConstPerceptionAgingRate);
+	/** returns true if, as result of stimuli aging, this listener needs an update (like if some stimuli expired)*/
+	bool AgeStimuli(const float ConstPerceptionAgingRate);
 	void ForgetActor(AActor* ActorToForget);
 
 	float GetYoungestStimulusAge(const AActor& Source) const;
@@ -254,6 +255,9 @@ protected:
 
 	/** Updates the stimulus entry in StimulusStore, if NewStimulus is more recent or stronger */
 	virtual void RefreshStimulus(FAIStimulus& StimulusStore, const FAIStimulus& NewStimulus);
+
+	/** @note no need to call super implementation, it's there just for some validity checking */
+	virtual void HandleExpiredStimulus(FAIStimulus& StimulusStore);
 	
 private:
 	FPerceptionListenerID PerceptionListenerId;
diff --git a/Engine/Source/Runtime/AIModule/Classes/Perception/AIPerceptionSystem.h b/Engine/Source/Runtime/AIModule/Classes/Perception/AIPerceptionSystem.h
--- a/Engine/Source/Runtime/AIModule/Classes/Perception/AIPerceptionSystem.h
+++ b/Engine/Source/Runtime/AIModule/Classes/Perception/AIPerceptionSystem.h
@@ -42,6 +42,11 @@ protected:
 	UPROPERTY(config, EditAnywhere, Category = Perception)
 	float PerceptionAgingRate;
 
+	/** gets set to true if as as result of stimuli aging (that's done outside of Tick, on timer)
+	 *	one of listeners requires an update. The update, as usual is tone in Tick where 
+	 *	bSomeListenersNeedUpdateDueToStimuliAging gets reset to false */
+	uint32 bSomeListenersNeedUpdateDueToStimuliAging : 1;
+
 	struct FDelayedStimulus
 	{
 		float DeliveryTimestamp;
diff --git a/Engine/Source/Runtime/AIModule/Classes/Perception/AIPerceptionTypes.h b/Engine/Source/Runtime/AIModule/Classes/Perception/AIPerceptionTypes.h
--- a/Engine/Source/Runtime/AIModule/Classes/Perception/AIPerceptionTypes.h
+++ b/Engine/Source/Runtime/AIModule/Classes/Perception/AIPerceptionTypes.h
@@ -74,6 +74,9 @@ struct AIMODULE_API FAIStimulus
 protected:
 	UPROPERTY(BlueprintReadWrite, Category = "AI|Perception")
 	float Age;
+
+	UPROPERTY(BlueprintReadWrite, Category = "AI|Perception")
+	float ExpirationAge;
 public:
 	UPROPERTY(BlueprintReadWrite, Category = "AI|Perception")
 	float Strength;
@@ -85,10 +88,18 @@ public:
 	FAISenseID Type;
 protected:
 	uint32 bLastSensingResult:1; // currently used only for marking failed sight tests
+	/** this means the stimulus was originally created with a "time limit" and this time has passed. 
+	 *	Expiration also results in calling MarkNoLongerSensed */
+	uint32 bExpired:1;	
 	
 public:
-	FAIStimulus(FAISenseID SenseType, float StimulusStrength, const FVector& InStimulusLocation, const FVector& InReceiverLocation, FResult Result = SensingSucceeded, float StimulusAge = 0.f)
-		: Age(StimulusAge), Strength(Result == SensingSucceeded ? StimulusStrength : -1.f)
+	
+	/** this is the recommended constructor. Use others if you know what you're doing. */
+	FAIStimulus(const UAISense& Sense, float StimulusStrength, const FVector& InStimulusLocation, const FVector& InReceiverLocation, FResult Result = SensingSucceeded);
+
+	FAIStimulus(FAISenseID SenseType, float StimulusStrength, const FVector& InStimulusLocation, const FVector& InReceiverLocation, FResult Result = SensingSucceeded)
+		: Age(0.f), ExpirationAge(NeverHappenedAge)
+		, Strength(Result == SensingSucceeded ? StimulusStrength : -1.f)
 		, StimulusLocation(InStimulusLocation)
 		, ReceiverLocation(InReceiverLocation), Type(SenseType), bLastSensingResult(Result == SensingSucceeded)
 	{}
@@ -99,10 +110,20 @@ public:
 		, ReceiverLocation(FAISystem::InvalidLocation), Type(FAISenseID::InvalidID()), bLastSensingResult(false)
 	{}
 
+	FAIStimulus& SetExpirationAge(float InExpirationAge) { ExpirationAge = InExpirationAge; return *this; }
+	FAIStimulus& SetStimulusAge(float StimulusAge) { Age = StimulusAge; return *this; }
+	
 	FORCEINLINE float GetAge() const { return Strength > 0 ? Age : NeverHappenedAge; }
-	FORCEINLINE void AgeStimulus(float ConstPerceptionAgingRate) { Age += ConstPerceptionAgingRate; }
+	/** @return false when this stimulus is no longer valid, when it is Expired */
+	FORCEINLINE bool AgeStimulus(float ConstPerceptionAgingRate) 
+	{ 
+		Age += ConstPerceptionAgingRate; 
+		return Age > ExpirationAge;
+	}
 	FORCEINLINE bool WasSuccessfullySensed() const { return bLastSensingResult; }
+	FORCEINLINE bool IsExpired() const { return bExpired; }
 	FORCEINLINE void MarkNoLongerSensed() { bLastSensingResult = false; }
+	FORCEINLINE void MarkExpired() { bExpired = true; MarkNoLongerSensed(); }
 	FORCEINLINE bool IsActive() const { return WasSuccessfullySensed() == true && GetAge() < NeverHappenedAge; }
 };
 
@@ -135,9 +156,9 @@ struct AIMODULE_API FPerceptionListener
 
 	FGenericTeamId TeamIdentifier;
 
+private:
 	uint32 bHasStimulusToProcess : 1;
 
-private:
 	FPerceptionListenerID ListenerID;
 
 	FPerceptionListener();
@@ -173,6 +194,7 @@ public:
 private:
 	friend class UAIPerceptionSystem;
 	FORCEINLINE void SetListenerID(FPerceptionListenerID InListenerID) { ListenerID = InListenerID; }
+	FORCEINLINE void MarkForStimulusProcessing() { bHasStimulusToProcess = true; }
 };
 
 namespace AIPerception
diff --git a/Engine/Source/Runtime/AIModule/Classes/Perception/AISense.h b/Engine/Source/Runtime/AIModule/Classes/Perception/AISense.h
--- a/Engine/Source/Runtime/AIModule/Classes/Perception/AISense.h
+++ b/Engine/Source/Runtime/AIModule/Classes/Perception/AISense.h
@@ -10,7 +10,7 @@ class UAIPerceptionSystem;
 
 DECLARE_DELEGATE_OneParam(FOnPerceptionListenerUpdateDelegate, const FPerceptionListener&);
 
-UCLASS(ClassGroup=AI, abstract)
+UCLASS(ClassGroup = AI, abstract, config = Engine)
 class AIMODULE_API UAISense : public UObject
 {
 	GENERATED_UCLASS_BODY()
@@ -25,6 +25,10 @@ protected:
 	UPROPERTY(EditAnywhere, BlueprintReadOnly, Category = Debug)
 	FString DebugName;
 
+	/** age past which stimulus of this sense are "forgotten"*/
+	UPROPERTY(EditAnywhere, BlueprintReadOnly, Category = "AI Perception", config)
+	float DefaultExpirationAge;
+	
 private:
 	UPROPERTY()
 	UAIPerceptionSystem* PerceptionSystemInstance;
@@ -87,6 +91,8 @@ public:
 	FORCEINLINE void OnListenerUpdate(const FPerceptionListener& NewListener) { OnListenerUpdateDelegate.ExecuteIfBound(NewListener); }
 	FORCEINLINE void OnListenerRemoved(const FPerceptionListener& NewListener) { OnListenerRemovedDelegate.ExecuteIfBound(NewListener); }
 
+	FORCEINLINE float GetDefaultExpirationAge() const { return DefaultExpirationAge; }
+
 protected:
 	/** @return time until next update */
 	virtual float Update() { return FLT_MAX; }
diff --git a/Engine/Source/Runtime/AIModule/Private/Perception/AIPerceptionComponent.cpp b/Engine/Source/Runtime/AIModule/Private/Perception/AIPerceptionComponent.cpp
--- a/Engine/Source/Runtime/AIModule/Private/Perception/AIPerceptionComponent.cpp
+++ b/Engine/Source/Runtime/AIModule/Private/Perception/AIPerceptionComponent.cpp
@@ -394,6 +394,7 @@ void UAIPerceptionComponent::ProcessStimuli()
 
 		FAIStimulus& StimulusStore = PerceptualInfo->LastSensedStimuli[SourcedStimulus->Stimulus.Type];
 
+		// if the new stimulus is "valid" or it's info that "no longer sensed" and it used to be sensed successfully
 		if (SourcedStimulus->Stimulus.WasSuccessfullySensed() || StimulusStore.WasSuccessfullySensed())
 		{
 			UpdatedActors.AddUnique(SourcedStimulus->Source);
@@ -403,6 +404,10 @@ void UAIPerceptionComponent::ProcessStimuli()
 		{
 			RefreshStimulus(StimulusStore, SourcedStimulus->Stimulus);
 		}
+		else if (StimulusStore.IsExpired())
+		{
+			HandleExpiredStimulus(StimulusStore);
+		}
 		else
 		{
 			// @note there some more valid info in SourcedStimulus->Stimulus regarding test that failed
@@ -435,17 +440,35 @@ void UAIPerceptionComponent::RefreshStimulus(FAIStimulus& StimulusStore, const F
 	}
 }
 
-void UAIPerceptionComponent::AgeStimuli(const float ConstPerceptionAgingRate)
+void UAIPerceptionComponent::HandleExpiredStimulus(FAIStimulus& StimulusStore)
+{
+	ensure(StimulusStore.IsExpired() == true && StimulusStore.WasSuccessfullySensed() == false && StimulusStore.IsActive() == false);
+}
+
+bool UAIPerceptionComponent::AgeStimuli(const float ConstPerceptionAgingRate)
 {
+	bool bExpiredStimuli = false;
+
 	for (TActorPerceptionContainer::TIterator It(PerceptualData); It; ++It)
 	{
 		FActorPerceptionInfo& ActorPerceptionInfo = It->Value;
 
-		for (auto& Stimulus : ActorPerceptionInfo.LastSensedStimuli)
+		for (FAIStimulus& Stimulus : ActorPerceptionInfo.LastSensedStimuli)
 		{
-			Stimulus.AgeStimulus(ConstPerceptionAgingRate);
+			if (Stimulus.AgeStimulus(ConstPerceptionAgingRate) == false)
+			{
+				AActor* TargetActor = ActorPerceptionInfo.Target.Get();
+				if (TargetActor)
+				{
+					Stimulus.MarkExpired();
+					RegisterStimulus(TargetActor, Stimulus);
+					bExpiredStimuli = true;
+				}
+			}
 		}
 	}
+
+	return bExpiredStimuli;
 }
 
 void UAIPerceptionComponent::ForgetActor(AActor* ActorToForget)
diff --git a/Engine/Source/Runtime/AIModule/Private/Perception/AIPerceptionSystem.cpp b/Engine/Source/Runtime/AIModule/Private/Perception/AIPerceptionSystem.cpp
--- a/Engine/Source/Runtime/AIModule/Private/Perception/AIPerceptionSystem.cpp
+++ b/Engine/Source/Runtime/AIModule/Private/Perception/AIPerceptionSystem.cpp
@@ -28,6 +28,7 @@ FAISenseID UAISenseConfig::GetSenseID() const
 UAIPerceptionSystem::UAIPerceptionSystem(const FObjectInitializer& ObjectInitializer)
 	: Super(ObjectInitializer)
 	, PerceptionAgingRate(0.3f)
+	, bSomeListenersNeedUpdateDueToStimuliAging(false)
 	, CurrentTime(0.f)
 {
 }
@@ -203,7 +204,7 @@ void UAIPerceptionSystem::Tick(float DeltaSeconds)
 		/** no point in soring in no new stimuli was processed */
 		bool bStimuliDelivered = DeliverDelayedStimuli(bNeedsUpdate ? RequiresSorting : NoNeedToSort);
 
-		if (bNeedsUpdate || bStimuliDelivered)
+		if (bNeedsUpdate || bStimuliDelivered || bSomeListenersNeedUpdateDueToStimuliAging)
 		{
 			for (AIPerception::FListenerMap::TIterator ListenerIt(ListenerContainer); ListenerIt; ++ListenerIt)
 			{
@@ -213,6 +214,8 @@ void UAIPerceptionSystem::Tick(float DeltaSeconds)
 					ListenerIt->Value.ProcessStimuli();
 				}
 			}
+
+			bSomeListenersNeedUpdateDueToStimuliAging = false;
 		}
 	}
 }
@@ -227,7 +230,12 @@ void UAIPerceptionSystem::AgeStimuli()
 		FPerceptionListener& Listener = ListenerIt->Value;
 		if (Listener.Listener.IsValid())
 		{
-			Listener.Listener->AgeStimuli(ConstPerceptionAgingRate);
+			// AgeStimuli will return true if this listener requires an update after stimuli aging
+			if (Listener.Listener->AgeStimuli(ConstPerceptionAgingRate))
+			{
+				Listener.MarkForStimulusProcessing();
+				bSomeListenersNeedUpdateDueToStimuliAging = true;
+			}
 		}
 	}
 }
diff --git a/Engine/Source/Runtime/AIModule/Private/Perception/AIPerceptionTypes.cpp b/Engine/Source/Runtime/AIModule/Private/Perception/AIPerceptionTypes.cpp
--- a/Engine/Source/Runtime/AIModule/Private/Perception/AIPerceptionTypes.cpp
+++ b/Engine/Source/Runtime/AIModule/Private/Perception/AIPerceptionTypes.cpp
@@ -2,10 +2,23 @@
 
 #include "AIModulePrivate.h"
 
+//----------------------------------------------------------------------//
+// FAIStimulus
+//----------------------------------------------------------------------//
+
 // mind that this needs to be > 0 since checks are done if Age < FAIStimulus::NeverHappenedAge
 // @todo maybe should be a function (IsValidAge)
 const float FAIStimulus::NeverHappenedAge = FLT_MAX;
 
+FAIStimulus::FAIStimulus(const UAISense& Sense, float StimulusStrength, const FVector& InStimulusLocation, const FVector& InReceiverLocation, FResult Result)
+	: Age(0.f), Strength(Result == SensingSucceeded ? StimulusStrength : -1.f)
+	, StimulusLocation(InStimulusLocation)
+	, ReceiverLocation(InReceiverLocation), bLastSensingResult(Result == SensingSucceeded)
+{
+	Type = Sense.GetSenseID();
+	ExpirationAge = Sense.GetDefaultExpirationAge();
+}
+
 //----------------------------------------------------------------------//
 // FPerceptionListener
 //----------------------------------------------------------------------//
diff --git a/Engine/Source/Runtime/AIModule/Private/Perception/AISense.cpp b/Engine/Source/Runtime/AIModule/Private/Perception/AISense.cpp
--- a/Engine/Source/Runtime/AIModule/Private/Perception/AISense.cpp
+++ b/Engine/Source/Runtime/AIModule/Private/Perception/AISense.cpp
@@ -17,6 +17,7 @@ UAISense::UAISense(const FObjectInitializer& ObjectInitializer)
 	, SenseID(FAISenseID::InvalidID())
 {
 	DebugName = GetName();
+	DefaultExpirationAge = FAIStimulus::NeverHappenedAge;
 
 	if (HasAnyFlags(RF_ClassDefaultObject) == false)
 	{
diff --git a/Engine/Source/Runtime/AIModule/Private/Perception/AISense_Damage.cpp b/Engine/Source/Runtime/AIModule/Private/Perception/AISense_Damage.cpp
--- a/Engine/Source/Runtime/AIModule/Private/Perception/AISense_Damage.cpp
+++ b/Engine/Source/Runtime/AIModule/Private/Perception/AISense_Damage.cpp
@@ -74,7 +74,7 @@ float UAISense_Damage::Update()
 				// this has to succeed, will assert a failure
 				FPerceptionListener& Listener = ListenersMap[PerceptionComponent->GetListenerId()];
 
-				Listener.RegisterStimulus(Event.Instigator, FAIStimulus(GetSenseID(), Event.Amount, Event.Location, Event.HitLocation));
+				Listener.RegisterStimulus(Event.Instigator, FAIStimulus(*this, Event.Amount, Event.Location, Event.HitLocation));
 			}
 		}
 	}
diff --git a/Engine/Source/Runtime/AIModule/Private/Perception/AISense_Hearing.cpp b/Engine/Source/Runtime/AIModule/Private/Perception/AISense_Hearing.cpp
--- a/Engine/Source/Runtime/AIModule/Private/Perception/AISense_Hearing.cpp
+++ b/Engine/Source/Runtime/AIModule/Private/Perception/AISense_Hearing.cpp
@@ -135,7 +135,7 @@ float UAISense_Hearing::Update()
 			const float Delay = SpeedOfSoundSq > 0.f ? FVector::DistSquared(Event.NoiseLocation, Listener.CachedLocation) / SpeedOfSoundSq : 0;
 			// pass over to listener to process 			
 			PerseptionSys->RegisterDelayedStimulus(Listener.GetListenerID(), Delay, Event.Instigator
-				, FAIStimulus(GetSenseID(), Event.Loudness, Event.NoiseLocation, Listener.CachedLocation) );
+				, FAIStimulus(*this, Event.Loudness, Event.NoiseLocation, Listener.CachedLocation) );
 		}
 	}
 
diff --git a/Engine/Source/Runtime/AIModule/Private/Perception/AISense_Prediction.cpp b/Engine/Source/Runtime/AIModule/Private/Perception/AISense_Prediction.cpp
--- a/Engine/Source/Runtime/AIModule/Private/Perception/AISense_Prediction.cpp
+++ b/Engine/Source/Runtime/AIModule/Private/Perception/AISense_Prediction.cpp
@@ -32,7 +32,7 @@ float UAISense_Prediction::Update()
 					// calculate the prediction here:
 					const FVector PredictedLocation = Event.PredictedActor->GetActorLocation() + Event.PredictedActor->GetVelocity() * Event.TimeToPredict;
 
-					Listener.RegisterStimulus(Event.PredictedActor, FAIStimulus(GetSenseID(), 1.f, PredictedLocation, Listener.CachedLocation));
+					Listener.RegisterStimulus(Event.PredictedActor, FAIStimulus(*this, 1.f, PredictedLocation, Listener.CachedLocation));
 				}
 			}
 		}
diff --git a/Engine/Source/Runtime/AIModule/Private/Perception/AISense_Sight.cpp b/Engine/Source/Runtime/AIModule/Private/Perception/AISense_Sight.cpp
--- a/Engine/Source/Runtime/AIModule/Private/Perception/AISense_Sight.cpp
+++ b/Engine/Source/Runtime/AIModule/Private/Perception/AISense_Sight.cpp
@@ -144,13 +144,13 @@ float UAISense_Sight::Update()
 						int32 NumberOfLoSChecksPerformed = 0;
 						if (Target.SightTargetInterface->CanBeSeenFrom(Listener.CachedLocation, OutSeenLocation, NumberOfLoSChecksPerformed, Listener.Listener->GetBodyActor()) == true)
 						{
-							Listener.RegisterStimulus(TargetActor, FAIStimulus(GetSenseID(), 1.f, OutSeenLocation, Listener.CachedLocation));
+							Listener.RegisterStimulus(TargetActor, FAIStimulus(*this, 1.f, OutSeenLocation, Listener.CachedLocation));
 							SightQuery->bLastResult = true;
 						}
 						else
 						{
 //							UE_VLOG_LOCATION(Listener.Listener.Get()->GetOwner(), TargetLocation, 25.f, FColor::Red, TEXT(""));
-							Listener.RegisterStimulus(TargetActor, FAIStimulus(GetSenseID(), 0.f, TargetLocation, Listener.CachedLocation, FAIStimulus::SensingFailed));
+							Listener.RegisterStimulus(TargetActor, FAIStimulus(*this, 0.f, TargetLocation, Listener.CachedLocation, FAIStimulus::SensingFailed));
 							SightQuery->bLastResult = false;
 						}
 
@@ -171,13 +171,13 @@ float UAISense_Sight::Update()
 
 						if (bHit == false || (HitResult.Actor.IsValid() && HitResult.Actor->IsOwnedBy(TargetActor)))
 						{
-							Listener.RegisterStimulus(TargetActor, FAIStimulus(GetSenseID(), 1.f, TargetLocation, Listener.CachedLocation));
+							Listener.RegisterStimulus(TargetActor, FAIStimulus(*this, 1.f, TargetLocation, Listener.CachedLocation));
 							SightQuery->bLastResult = true;
 						}
 						else
 						{
 //							UE_VLOG_LOCATION(Listener.Listener.Get()->GetOwner(), TargetLocation, 25.f, FColor::Red, TEXT(""));
-							Listener.RegisterStimulus(TargetActor, FAIStimulus(GetSenseID(), 0.f, TargetLocation, Listener.CachedLocation, FAIStimulus::SensingFailed));
+							Listener.RegisterStimulus(TargetActor, FAIStimulus(*this, 0.f, TargetLocation, Listener.CachedLocation, FAIStimulus::SensingFailed));
 							SightQuery->bLastResult = false;
 						}
 					}
@@ -185,7 +185,7 @@ float UAISense_Sight::Update()
 				else
 				{
 //					UE_VLOG_SEGMENT(Listener.Listener.Get()->GetOwner(), Listener.CachedLocation, TargetLocation, FColor::Red, TEXT("%s"), *(Target.TargetId.ToString()));
-					Listener.RegisterStimulus(TargetActor, FAIStimulus(GetSenseID(), 0.f, TargetLocation, Listener.CachedLocation, FAIStimulus::SensingFailed));
+					Listener.RegisterStimulus(TargetActor, FAIStimulus(*this, 0.f, TargetLocation, Listener.CachedLocation, FAIStimulus::SensingFailed));
 					SightQuery->bLastResult = false;
 				}
 
diff --git a/Engine/Source/Runtime/AIModule/Private/Perception/AISense_Team.cpp b/Engine/Source/Runtime/AIModule/Private/Perception/AISense_Team.cpp
--- a/Engine/Source/Runtime/AIModule/Private/Perception/AISense_Team.cpp
+++ b/Engine/Source/Runtime/AIModule/Private/Perception/AISense_Team.cpp
@@ -49,7 +49,7 @@ float UAISense_Team::Update()
 				continue;
 			}
 			
-			Listener.RegisterStimulus(Event.Enemy, FAIStimulus(GetSenseID(), 1.f, Event.LastKnowLocation, Event.GetBroadcastLocation(), FAIStimulus::SensingSucceeded, Event.InformationAge));
+			Listener.RegisterStimulus(Event.Enemy, FAIStimulus(*this, 1.f, Event.LastKnowLocation, Event.GetBroadcastLocation(), FAIStimulus::SensingSucceeded).SetStimulusAge(Event.InformationAge));
 		}
 	}
 
diff --git a/Engine/Source/Runtime/AIModule/Private/Perception/AISense_Touch.cpp b/Engine/Source/Runtime/AIModule/Private/Perception/AISense_Touch.cpp
--- a/Engine/Source/Runtime/AIModule/Private/Perception/AISense_Touch.cpp
+++ b/Engine/Source/Runtime/AIModule/Private/Perception/AISense_Touch.cpp
@@ -29,7 +29,7 @@ float UAISense_Touch::Update()
 					// this has to succeed, will assert a failure
 					FPerceptionListener& Listener = ListenersMap[PerceptionComponent->GetListenerId()];
 
-					Listener.RegisterStimulus(Event.OtherActor, FAIStimulus(GetSenseID(), 1.f, Event.Location, Event.Location));
+					Listener.RegisterStimulus(Event.OtherActor, FAIStimulus(*this, 1.f, Event.Location, Event.Location));
 				}
 			}
 		}
