CrowdManager's properties exposed to confing Project Settings #UE4

- github PR#655 by bozaro

diff --git a/Engine/Source/Editor/ProjectSettingsViewer/Private/ProjectSettingsViewerModule.cpp b/Engine/Source/Editor/ProjectSettingsViewer/Private/ProjectSettingsViewerModule.cpp
--- a/Engine/Source/Editor/ProjectSettingsViewer/Private/ProjectSettingsViewerModule.cpp
+++ b/Engine/Source/Editor/ProjectSettingsViewer/Private/ProjectSettingsViewerModule.cpp
@@ -16,6 +16,7 @@
 #include "Runtime/Engine/Classes/Engine/RendererSettings.h"
 #include "Runtime/Engine/Classes/Engine/UserInterfaceSettings.h"
 #include "Runtime/Engine/Classes/Sound/AudioSettings.h"
+#include "Runtime/AIModule/Classes/Navigation/CrowdManager.h"
 #include "AISystem.h"
 
 #define LOCTEXT_NAMESPACE "FProjectSettingsViewerModule"
@@ -146,7 +147,15 @@ protected:
 			LOCTEXT("AISystemSettingsDescription", "Settings for the AI System."),
 			GetMutableDefault<UAISystem>()
 			);
-/*
+
+		// Crowd manager
+		SettingsModule.RegisterSettings("Project", "Engine", "CrowdManager",
+			LOCTEXT("CrowdManagerSettingsName", "Crowd Manager"),
+			LOCTEXT("CrowdManagerSettingsDescription", "Settings for the AI Crowd Manager."),
+			GetMutableDefault<UCrowdManager>()
+			);
+
+		/*
 		// network settings
 		SettingsModule.RegisterSettings("Project", "Engine", "NetworkManager",
 			LOCTEXT("GameNetworkManagerSettingsName", "Network Manager"),
@@ -259,6 +268,7 @@ protected:
 
 			// engine settings
 			SettingsModule->UnregisterSettings("Project", "Engine", "General");
+			SettingsModule->UnregisterSettings("Project", "Engine", "CrowdManager");
 			SettingsModule->UnregisterSettings("Project", "Engine", "NavigationSystem");
 			SettingsModule->UnregisterSettings("Project", "Engine", "NavigationMesh");
 			SettingsModule->UnregisterSettings("Project", "Engine", "Input");
diff --git a/Engine/Source/Runtime/AIModule/Classes/Navigation/CrowdManager.h b/Engine/Source/Runtime/AIModule/Classes/Navigation/CrowdManager.h
--- a/Engine/Source/Runtime/AIModule/Classes/Navigation/CrowdManager.h
+++ b/Engine/Source/Runtime/AIModule/Classes/Navigation/CrowdManager.h
@@ -150,7 +150,7 @@ struct FCrowdTickHelper : FTickableGameObject
 	virtual TStatId GetStatId() const;
 };
 
-UCLASS()
+UCLASS(config = Engine, defaultconfig)
 class AIMODULE_API UCrowdManager : public UObject
 {
 	GENERATED_UCLASS_BODY()
@@ -236,35 +236,35 @@ protected:
 	ANavigationData* MyNavData;
 
 	/** obstacle avoidance params */
-	UPROPERTY(EditAnywhere, Category=Config)
+	UPROPERTY(config, EditAnywhere, Category = Config)
 	TArray<FCrowdAvoidanceConfig> AvoidanceConfig;
 
 	/** obstacle avoidance params */
-	UPROPERTY(EditAnywhere, Category=Config)
+	UPROPERTY(config, EditAnywhere, Category = Config)
 	TArray<FCrowdAvoidanceSamplingPattern> SamplingPatterns;
 
 	/** max number of agents supported by crowd */
-	UPROPERTY(EditAnywhere, Category=Config)
+	UPROPERTY(config, EditAnywhere, Category = Config)
 	int32 MaxAgents;
 
 	/** max radius of agent that can be added to crowd */
-	UPROPERTY(EditAnywhere, Category=Config)
+	UPROPERTY(config, EditAnywhere, Category = Config)
 	float MaxAgentRadius;
 
 	/** max number of neighbor agents for velocity avoidance */
-	UPROPERTY(EditAnywhere, Category=Config)
+	UPROPERTY(config, EditAnywhere, Category = Config)
 	int32 MaxAvoidedAgents;
 
 	/** max number of wall segments for velocity avoidance */
-	UPROPERTY(EditAnywhere, Category=Config)
+	UPROPERTY(config, EditAnywhere, Category = Config)
 	int32 MaxAvoidedWalls;
 
 	/** how often should agents check their position after moving off navmesh? */
-	UPROPERTY(EditAnywhere, Category=Config)
+	UPROPERTY(config, EditAnywhere, Category = Config)
 	float NavmeshCheckInterval;
 
 	/** how often should agents try to optimize their paths? */
-	UPROPERTY(EditAnywhere, Category=Config)
+	UPROPERTY(config, EditAnywhere, Category = Config)
 	float PathOptimizationInterval;
 
 	uint32 bPruneStartedOffmeshConnections : 1;
