#UE4 - Fixed AgeStimulus() returning true for an expired stimulus
  - Prevent unsensed stimuli from being marked and registered as expired

diff --git a/Engine/Source/Runtime/AIModule/Classes/Perception/AIPerceptionTypes.h b/Engine/Source/Runtime/AIModule/Classes/Perception/AIPerceptionTypes.h
--- a/Engine/Source/Runtime/AIModule/Classes/Perception/AIPerceptionTypes.h
+++ b/Engine/Source/Runtime/AIModule/Classes/Perception/AIPerceptionTypes.h
@@ -118,7 +118,7 @@ public:
 	FORCEINLINE bool AgeStimulus(float ConstPerceptionAgingRate) 
 	{ 
 		Age += ConstPerceptionAgingRate; 
-		return Age > ExpirationAge;
+		return Age < ExpirationAge;
 	}
 	FORCEINLINE bool WasSuccessfullySensed() const { return bLastSensingResult; }
 	FORCEINLINE bool IsExpired() const { return bExpired; }
diff --git a/Engine/Source/Runtime/AIModule/Private/Perception/AIPerceptionComponent.cpp b/Engine/Source/Runtime/AIModule/Private/Perception/AIPerceptionComponent.cpp
--- a/Engine/Source/Runtime/AIModule/Private/Perception/AIPerceptionComponent.cpp
+++ b/Engine/Source/Runtime/AIModule/Private/Perception/AIPerceptionComponent.cpp
@@ -455,7 +455,10 @@ bool UAIPerceptionComponent::AgeStimuli(const float ConstPerceptionAgingRate)
 
 		for (FAIStimulus& Stimulus : ActorPerceptionInfo.LastSensedStimuli)
 		{
-			if (Stimulus.AgeStimulus(ConstPerceptionAgingRate) == false)
+			// Age the stimulus. If it is active but has just expired, mark it as such
+			if (Stimulus.AgeStimulus(ConstPerceptionAgingRate) == false && 
+				Stimulus.IsActive() && 
+				!Stimulus.IsExpired())
 			{
 				AActor* TargetActor = ActorPerceptionInfo.Target.Get();
 				if (TargetActor)
