Fixed missing FAIStimulus members initialization #UE4

diff --git a/Engine/Source/Runtime/AIModule/Classes/Perception/AIPerceptionTypes.h b/Engine/Source/Runtime/AIModule/Classes/Perception/AIPerceptionTypes.h
--- a/Engine/Source/Runtime/AIModule/Classes/Perception/AIPerceptionTypes.h
+++ b/Engine/Source/Runtime/AIModule/Classes/Perception/AIPerceptionTypes.h
@@ -101,13 +101,13 @@ public:
 		: Age(0.f), ExpirationAge(NeverHappenedAge)
 		, Strength(Result == SensingSucceeded ? StimulusStrength : -1.f)
 		, StimulusLocation(InStimulusLocation)
-		, ReceiverLocation(InReceiverLocation), Type(SenseType), bLastSensingResult(Result == SensingSucceeded)
+		, ReceiverLocation(InReceiverLocation), Type(SenseType), bLastSensingResult(Result == SensingSucceeded), bExpired(false)
 	{}
 
 	// default constructor
 	FAIStimulus()
-		: Age(NeverHappenedAge), Strength(-1.f), StimulusLocation(FAISystem::InvalidLocation)
-		, ReceiverLocation(FAISystem::InvalidLocation), Type(FAISenseID::InvalidID()), bLastSensingResult(false)
+		: Age(NeverHappenedAge), ExpirationAge(NeverHappenedAge), Strength(-1.f), StimulusLocation(FAISystem::InvalidLocation)
+		, ReceiverLocation(FAISystem::InvalidLocation), Type(FAISenseID::InvalidID()), bLastSensingResult(false), bExpired(false)
 	{}
 
 	FAIStimulus& SetExpirationAge(float InExpirationAge) { ExpirationAge = InExpirationAge; return *this; }
diff --git a/Engine/Source/Runtime/AIModule/Private/Perception/AIPerceptionTypes.cpp b/Engine/Source/Runtime/AIModule/Private/Perception/AIPerceptionTypes.cpp
--- a/Engine/Source/Runtime/AIModule/Private/Perception/AIPerceptionTypes.cpp
+++ b/Engine/Source/Runtime/AIModule/Private/Perception/AIPerceptionTypes.cpp
@@ -14,6 +14,7 @@ FAIStimulus::FAIStimulus(const UAISense& Sense, float StimulusStrength, const FV
 	: Age(0.f), Strength(Result == SensingSucceeded ? StimulusStrength : -1.f)
 	, StimulusLocation(InStimulusLocation)
 	, ReceiverLocation(InReceiverLocation), bLastSensingResult(Result == SensingSucceeded)
+	, bExpired(false)
 {
 	Type = Sense.GetSenseID();
 	ExpirationAge = Sense.GetDefaultExpirationAge();
