An option to specify a random deviation added to BTTask_Wait #UE4

- github PR#663 by slonopotamus

diff --git a/Engine/Source/Runtime/AIModule/Classes/BehaviorTree/Tasks/BTTask_Wait.h b/Engine/Source/Runtime/AIModule/Classes/BehaviorTree/Tasks/BTTask_Wait.h
--- a/Engine/Source/Runtime/AIModule/Classes/BehaviorTree/Tasks/BTTask_Wait.h
+++ b/Engine/Source/Runtime/AIModule/Classes/BehaviorTree/Tasks/BTTask_Wait.h
@@ -19,9 +19,13 @@ class AIMODULE_API UBTTask_Wait : public UBTTaskNode
 	GENERATED_UCLASS_BODY()
 
 	/** wait time in seconds */
-	UPROPERTY(Category=Wait, EditAnywhere)
+	UPROPERTY(Category = Wait, EditAnywhere)
 	float WaitTime;
 
+	/** allows adding random time to wait time */
+	UPROPERTY(Category = Wait, EditAnywhere, meta = (UIMin = 0, ClampMin = 0))
+	float RandomDeviation;
+
 	virtual EBTNodeResult::Type ExecuteTask(UBehaviorTreeComponent& OwnerComp, uint8* NodeMemory) override;
 	virtual uint16 GetInstanceMemorySize() const override;
 	virtual void DescribeRuntimeValues(const UBehaviorTreeComponent& OwnerComp, uint8* NodeMemory, EBTDescriptionVerbosity::Type Verbosity, TArray<FString>& Values) const override;
diff --git a/Engine/Source/Runtime/AIModule/Private/BehaviorTree/Tasks/BTTask_Wait.cpp b/Engine/Source/Runtime/AIModule/Private/BehaviorTree/Tasks/BTTask_Wait.cpp
--- a/Engine/Source/Runtime/AIModule/Private/BehaviorTree/Tasks/BTTask_Wait.cpp
+++ b/Engine/Source/Runtime/AIModule/Private/BehaviorTree/Tasks/BTTask_Wait.cpp
@@ -13,7 +13,7 @@ UBTTask_Wait::UBTTask_Wait(const FObjectInitializer& ObjectInitializer) : Super(
 EBTNodeResult::Type UBTTask_Wait::ExecuteTask(UBehaviorTreeComponent& OwnerComp, uint8* NodeMemory)
 {
 	FBTWaitTaskMemory* MyMemory = (FBTWaitTaskMemory*)NodeMemory;
-	MyMemory->RemainingWaitTime = WaitTime;
+	MyMemory->RemainingWaitTime = FMath::FRandRange(FMath::Max(0.0f, WaitTime - RandomDeviation), (WaitTime + RandomDeviation));
 	
 	return EBTNodeResult::InProgress;
 }
@@ -32,7 +32,14 @@ void UBTTask_Wait::TickTask(UBehaviorTreeComponent& OwnerComp, uint8* NodeMemory
 
 FString UBTTask_Wait::GetStaticDescription() const
 {
-	return FString::Printf(TEXT("%s: %.1fs"), *Super::GetStaticDescription(), WaitTime);
+	if (FMath::IsNearlyZero(RandomDeviation))
+	{
+		return FString::Printf(TEXT("%s: %.1fs"), *Super::GetStaticDescription(), WaitTime);
+	}
+	else
+	{
+		return FString::Printf(TEXT("%s: %.1f+-%.1fs"), *Super::GetStaticDescription(), WaitTime, RandomDeviation);
+	}
 }
 
 void UBTTask_Wait::DescribeRuntimeValues(const UBehaviorTreeComponent& OwnerComp, uint8* NodeMemory, EBTDescriptionVerbosity::Type Verbosity, TArray<FString>& Values) const
