Exposed BrainComponent's functions for stopping and restarting AI logic #UE4

diff --git a/Engine/Source/Runtime/AIModule/Classes/BrainComponent.h b/Engine/Source/Runtime/AIModule/Classes/BrainComponent.h
--- a/Engine/Source/Runtime/AIModule/Classes/BrainComponent.h
+++ b/Engine/Source/Runtime/AIModule/Classes/BrainComponent.h
@@ -126,9 +126,14 @@ public:
 	 *	On subsequent ResumeLogic instead RestartLogic will be called. 
 	 *	@note this call does nothing if logic is not locked at the moment of call */
 	void RequestLogicRestartOnUnlock();
-	virtual void RestartLogic() {}
+
+	UFUNCTION(BlueprintCallable, Category = "AI|Logic")
+	virtual void RestartLogic();
+
+	UFUNCTION(BlueprintCallable, Category = "AI|Logic")
+	virtual void StopLogic(const FString& Reason);
+
 protected:
-	virtual void StopLogic(const FString& Reason) {}
 	virtual void PauseLogic(const FString& Reason) {}
 	/** MUST be called by child implementations!
 	 *	@return indicates whether child class' ResumeLogic should be called (true) or has it been 
diff --git a/Engine/Source/Runtime/AIModule/Private/BrainComponent.cpp b/Engine/Source/Runtime/AIModule/Private/BrainComponent.cpp
--- a/Engine/Source/Runtime/AIModule/Private/BrainComponent.cpp
+++ b/Engine/Source/Runtime/AIModule/Private/BrainComponent.cpp
@@ -314,3 +314,13 @@ EAILogicResuming::Type UBrainComponent::ResumeLogic(const FString& Reason)
 
 	return EAILogicResuming::Continue;
 }
+
+void UBrainComponent::RestartLogic()
+{
+
+}
+
+void UBrainComponent::StopLogic(const FString& Reason)
+{
+
+}
