Update to AAIController::UpdateControlRotation to properly use AI's focus mechanics #UE4

- github PR#666 by slonopotamus

diff --git a/Engine/Source/Runtime/AIModule/Private/AIController.cpp b/Engine/Source/Runtime/AIModule/Private/AIController.cpp
--- a/Engine/Source/Runtime/AIModule/Private/AIController.cpp
+++ b/Engine/Source/Runtime/AIModule/Private/AIController.cpp
@@ -401,27 +401,27 @@ void AAIController::UpdateControlRotation(float DeltaTime, bool bUpdatePawn)
 {
 	// Look toward focus
 	FVector FocalPoint = GetFocalPoint();
-	if (FAISystem::IsValidLocation(FocalPoint) && GetPawn())
+	APawn* const Pawn = GetPawn();
+
+	if (Pawn)
 	{
-		FVector Direction = FocalPoint - GetPawn()->GetActorLocation();
+		FVector Direction = FAISystem::IsValidLocation(FocalPoint) ? (FocalPoint - Pawn->GetActorLocation()) : Pawn->GetActorForwardVector();
 		FRotator NewControlRotation = Direction.Rotation();
 
-		// Don't pitch view of walking pawns unless looking at another pawn
-		if (GetPawn()->GetMovementComponent() && GetPawn()->GetMovementComponent()->IsMovingOnGround() &&
-			PathFollowingComponent && (!PathFollowingComponent->GetMoveGoal() || !Cast<APawn>(PathFollowingComponent->GetMoveGoal())))
+		// Don't pitch view unless looking at another pawn
+		if (Cast<APawn>(GetFocusActor()) == nullptr)
 		{
 			NewControlRotation.Pitch = 0.f;
 		}
 		NewControlRotation.Yaw = FRotator::ClampAxis(NewControlRotation.Yaw);
 
-		if (!GetControlRotation().Equals(NewControlRotation, 1e-3f))
+		if (GetControlRotation().Equals(NewControlRotation, 1e-3f) == false)
 		{
 			SetControlRotation(NewControlRotation);
 
-			APawn* const P = GetPawn();
-			if (P && bUpdatePawn)
+			if (bUpdatePawn)
 			{
-				P->FaceRotation(NewControlRotation, DeltaTime);
+				Pawn->FaceRotation(NewControlRotation, DeltaTime);
 			}
 		}
 	}
