Made PawnAction aborting be treated like all other action events, meaning it's processed in PawnActionsComponent's tick rather than instantly #UE4

- merged over from dev branch

diff --git a/Engine/Source/Runtime/AIModule/Classes/AITypes.h b/Engine/Source/Runtime/AIModule/Classes/AITypes.h
--- a/Engine/Source/Runtime/AIModule/Classes/AITypes.h
+++ b/Engine/Source/Runtime/AIModule/Classes/AITypes.h
@@ -104,6 +104,7 @@ namespace EPawnActionEventType
 	{
 		Invalid,
 		FailedToStart,
+		InstantAbort,
 		FinishedAborting,
 		FinishedExecution,
 		Push,
diff --git a/Engine/Source/Runtime/AIModule/Private/Actions/PawnActionsComponent.cpp b/Engine/Source/Runtime/AIModule/Private/Actions/PawnActionsComponent.cpp
--- a/Engine/Source/Runtime/AIModule/Private/Actions/PawnActionsComponent.cpp
+++ b/Engine/Source/Runtime/AIModule/Private/Actions/PawnActionsComponent.cpp
@@ -180,8 +180,19 @@ void UPawnActionsComponent::TickComponent(float DeltaTime, enum ELevelTick TickT
 		for (int32 EventIndex = 0; EventIndex < ActionEvents.Num(); ++EventIndex)
 		{
 			FPawnActionEvent& Event = ActionEvents[EventIndex];
+
+			if (Event.Action == nullptr)
+			{
+				UE_VLOG(ControlledPawn, LogPawnAction, Warning, TEXT("NULL action encountered during ActionEvents processing. May result in some notifies not being sent out."));
+				continue;
+			}
+
 			switch (Event.EventType)
 			{
+			case EPawnActionEventType::InstantAbort:
+				Event.Action->Abort(EAIForceParam::Force);
+				ActionStacks[Event.Priority].PopAction(*Event.Action);
+				break;
 			case EPawnActionEventType::FinishedAborting:
 			case EPawnActionEventType::FinishedExecution:
 			case EPawnActionEventType::FailedToStart:
@@ -370,7 +381,7 @@ uint32 UPawnActionsComponent::AbortActionsInstigatedBy(UObject* const Instigator
 		{
 			if (Action->GetInstigator() == Instigator)
 			{
-				Action->Abort(EAIForceParam::Force);
+				OnEvent(*Action, EPawnActionEventType::InstantAbort);
 				++AbortedActionsCount;
 			}
 			Action = Action->ParentAction;
