Introduced a dedicated AIController setup with CrowdFollowingComponent in place of default PathFollowingComponent #UE4

- not that will no longer be needed once we get a way to override component classes from BP

diff --git a/Engine/Source/Runtime/AIModule/Classes/DetourCrowdAIController.h b/Engine/Source/Runtime/AIModule/Classes/DetourCrowdAIController.h
new file mode 100644
--- /dev/null
+++ b/Engine/Source/Runtime/AIModule/Classes/DetourCrowdAIController.h
@@ -0,0 +1,14 @@
+// Copyright 1998-2014 Epic Games, Inc. All Rights Reserved.
+
+#pragma once
+
+#include "AIController.h"
+#include "DetourCrowdAIController.generated.h"
+
+UCLASS()
+class ADetourCrowdAIController : public AAIController
+{
+	GENERATED_BODY()
+public:
+	ADetourCrowdAIController(const FObjectInitializer& ObjectInitializer);
+};
\ No newline at end of file
diff --git a/Engine/Source/Runtime/AIModule/Private/DetourCrowdAIController.cpp b/Engine/Source/Runtime/AIModule/Private/DetourCrowdAIController.cpp
new file mode 100644
--- /dev/null
+++ b/Engine/Source/Runtime/AIModule/Private/DetourCrowdAIController.cpp
@@ -0,0 +1,11 @@
+// Copyright 1998-2015 Epic Games, Inc. All Rights Reserved.
+
+#include "AIModulePrivate.h"
+#include "DetourCrowdAIController.h"
+#include "Navigation/CrowdFollowingComponent.h"
+
+ADetourCrowdAIController::ADetourCrowdAIController(const FObjectInitializer& ObjectInitializer)
+	: Super(ObjectInitializer.SetDefaultSubobjectClass<UCrowdFollowingComponent>(TEXT("PathFollowingComponent")))
+{
+
+}
