Fix to AAIController::MoveToLocation misusing bStopOnOverlap when testing if already at destination location #UE4

- github PR#679 by mollstam

diff --git a/Engine/Source/Runtime/AIModule/Private/AIController.cpp b/Engine/Source/Runtime/AIModule/Private/AIController.cpp
--- a/Engine/Source/Runtime/AIModule/Private/AIController.cpp
+++ b/Engine/Source/Runtime/AIModule/Private/AIController.cpp
@@ -559,7 +559,7 @@ EPathFollowingRequestResult::Type AAIController::MoveToLocation(const FVector& D
 		GoalLocation = ProjectedLocation.Location;
 	}
 
-	if (bCanRequestMove && PathFollowingComponent && PathFollowingComponent->HasReached(GoalLocation, AcceptanceRadius, bStopOnOverlap))
+	if (bCanRequestMove && PathFollowingComponent && PathFollowingComponent->HasReached(GoalLocation, AcceptanceRadius, !bStopOnOverlap))
 	{
 		UE_VLOG(this, LogAINavigation, Log, TEXT("MoveToLocation: already at goal!"));
 
