Fixed AI Move To BP node's issues resulting from NavigationComponent's removal #UE4

diff --git a/Engine/Source/Runtime/AIModule/Private/AIController.cpp b/Engine/Source/Runtime/AIModule/Private/AIController.cpp
--- a/Engine/Source/Runtime/AIModule/Private/AIController.cpp
+++ b/Engine/Source/Runtime/AIModule/Private/AIController.cpp
@@ -654,6 +654,10 @@ bool AAIController::PreparePathfinding(FPathFindingQuery& Query, const FVector&
 				const FVector Offset = NavGoal->GetMoveGoalOffset(this);
 				GoalLocation = FRotationTranslationMatrix(Goal->GetActorRotation(), NavGoal->GetNavAgentLocation()).TransformPosition(Offset);
 			}
+			else
+			{
+				GoalLocation = Goal->GetActorLocation();
+			}
 		}
 
 		Query = FPathFindingQuery(this, NavData, GetNavAgentLocation(), GoalLocation, UNavigationQueryFilter::GetQueryFilter(NavData, FilterClass));
@@ -677,17 +681,24 @@ FAIRequestID AAIController::RequestPathAndMove(FPathFindingQuery& Query, AActor*
 	if (NavSys)
 	{
 		FPathFindingResult PathResult = NavSys->FindPathSync(Query);
-		if (PathResult.IsSuccessful() && PathResult.Path.IsValid())
+		if (PathResult.Result != ENavigationQueryResult::Error)
 		{
-			if (Goal)
+			if (PathResult.IsSuccessful() && PathResult.Path.IsValid())
 			{
-				PathResult.Path->SetGoalActorObservation(*Goal, 100.0f);
+				if (Goal)
+				{
+					PathResult.Path->SetGoalActorObservation(*Goal, 100.0f);
+				}
+
+				PathResult.Path->EnableRecalculationOnInvalidation(true);
 			}
 
-			PathResult.Path->EnableRecalculationOnInvalidation(true);
+			RequestID = RequestMove(PathResult.Path, Goal, AcceptanceRadius, bStopOnOverlap, CustomData);
+		}
+		else
+		{
+			UE_VLOG(this, LogBehaviorTree, Error, TEXT("Trying to find path to %s resulted in Error"), *GetNameSafe(Goal));
 		}
-
-		RequestID = RequestMove(PathResult.Path, Goal, AcceptanceRadius, bStopOnOverlap, CustomData);
 	}
 
 	return RequestID;
