BT task node for playing Animation #UE4

diff --git a/Engine/Source/Runtime/AIModule/Classes/BehaviorTree/Tasks/BTTask_PlayAnimation.h b/Engine/Source/Runtime/AIModule/Classes/BehaviorTree/Tasks/BTTask_PlayAnimation.h
new file mode 100644
--- /dev/null
+++ b/Engine/Source/Runtime/AIModule/Classes/BehaviorTree/Tasks/BTTask_PlayAnimation.h
@@ -0,0 +1,45 @@
+// Copyright 1998-2015 Epic Games, Inc. All Rights Reserved.
+
+#pragma once
+#include "BehaviorTree/BTTaskNode.h"
+#include "BTTask_PlayAnimation.generated.h"
+
+class UAnimationAsset;
+
+/**
+ *	Play indicated AnimationAsset on Pawn controlled by BT 
+ *	Note that this node is generic and is handing multiple special cases,
+ *	If you want a more efficient solution you'll need to implement it yourself (or wait for our BTTask_PlayCharacterAnimation)
+ */
+UCLASS()
+class AIMODULE_API UBTTask_PlayAnimation : public UBTTaskNode
+{
+	GENERATED_UCLASS_BODY()
+
+	/** Animation asset to play. Note that it needs to match the skeleton of pawn this BT is controlling */
+	UPROPERTY(Category = Node, EditAnywhere)
+	UAnimationAsset* AnimationToPlay;
+	
+	UPROPERTY(Category = Node, EditAnywhere)
+	uint32 bLooping : 1; 
+
+	/** if true the task will just trigger the animation and instantly finish. Fire and Forget. */
+	UPROPERTY(Category = Node, EditAnywhere)
+	uint32 bNonBlocking : 1;
+
+	UPROPERTY()
+	UBehaviorTreeComponent* MyOwnerComp;
+
+	FTimerDelegate TimerDelegate;
+	FTimerHandle TimerHandle;
+
+	virtual EBTNodeResult::Type ExecuteTask(UBehaviorTreeComponent& OwnerComp, uint8* NodeMemory) override;
+	virtual EBTNodeResult::Type AbortTask(UBehaviorTreeComponent& OwnerComp, uint8* NodeMemory) override;
+	virtual FString GetStaticDescription() const override;
+
+	void OnAnimationTimerDone();
+
+#if WITH_EDITOR
+	virtual FName GetNodeIconName() const override;
+#endif // WITH_EDITOR
+};
diff --git a/Engine/Source/Runtime/AIModule/Private/BehaviorTree/Tasks/BTTask_PlayAnimation.cpp b/Engine/Source/Runtime/AIModule/Private/BehaviorTree/Tasks/BTTask_PlayAnimation.cpp
new file mode 100644
--- /dev/null
+++ b/Engine/Source/Runtime/AIModule/Private/BehaviorTree/Tasks/BTTask_PlayAnimation.cpp
@@ -0,0 +1,95 @@
+// Copyright 1998-2015 Epic Games, Inc. All Rights Reserved.
+
+#include "AIModulePrivate.h"
+#include "BehaviorTree/Tasks/BTTask_PlayAnimation.h"
+
+//----------------------------------------------------------------------//
+// UBTTask_PlayAnimation
+//----------------------------------------------------------------------//
+UBTTask_PlayAnimation::UBTTask_PlayAnimation(const FObjectInitializer& ObjectInitializer) 
+	: Super(ObjectInitializer)
+{
+	NodeName = "Play Animation";
+	// instantiating to be able to use Timers
+	bCreateNodeInstance = true;
+
+	bLooping = false;
+	bNonBlocking = false;
+
+	TimerDelegate = FTimerDelegate::CreateUObject(this, &UBTTask_PlayAnimation::OnAnimationTimerDone);
+}
+
+EBTNodeResult::Type UBTTask_PlayAnimation::ExecuteTask(UBehaviorTreeComponent& OwnerComp, uint8* NodeMemory)
+{
+	AAIController* const MyController = OwnerComp.GetAIOwner();
+	EBTNodeResult::Type Result = EBTNodeResult::Failed;
+
+	// reset timer handle
+	TimerHandle.Invalidate();
+	MyOwnerComp = &OwnerComp;
+
+	if (AnimationToPlay && MyController)
+	{
+		ACharacter* const MyCharacter = Cast<ACharacter>(MyController->GetPawn());
+		if (MyCharacter && MyCharacter->GetMesh())
+		{
+			MyCharacter->GetMesh()->PlayAnimation(AnimationToPlay, bLooping);
+			const float FinishDelay = AnimationToPlay->GetMaxCurrentTime();
+
+			if (bNonBlocking == false && FinishDelay > 0)
+			{
+				if (bLooping == false)
+				{
+					MyController->GetWorld()->GetTimerManager().SetTimer(TimerHandle, TimerDelegate, FinishDelay, /*bLoop=*/false);
+				}
+				Result = EBTNodeResult::InProgress;
+			}
+			else
+			{
+				UE_CVLOG(bNonBlocking == false, MyController, LogBehaviorTree, Log, TEXT("%s> Instant success due to having a valid AnimationToPlay and Character with SkelMesh, but 0-length animation"), *GetNodeName());
+				// we're done here, report success so that BT can pick next task
+				Result = EBTNodeResult::Succeeded;
+			}
+		}
+	}
+
+	return Result;
+}
+
+EBTNodeResult::Type UBTTask_PlayAnimation::AbortTask(UBehaviorTreeComponent& OwnerComp, uint8* NodeMemory)
+{
+	AAIController* const MyController = OwnerComp.GetAIOwner();
+
+	if (AnimationToPlay && MyController && TimerHandle.IsValid())
+	{
+		MyController->GetWorld()->GetTimerManager().ClearTimer(TimerHandle);
+	}
+
+	TimerHandle.Invalidate();
+
+	return EBTNodeResult::Aborted;
+}
+
+FString UBTTask_PlayAnimation::GetStaticDescription() const
+{
+	return FString::Printf(TEXT("%s: '%s'"), *Super::GetStaticDescription(), *GetNameSafe(AnimationToPlay)
+		, bLooping ? TEXT(", looping") : TEXT("")
+		, bNonBlocking ? TEXT(", non-blocking") : TEXT("blocking"));
+}
+
+void UBTTask_PlayAnimation::OnAnimationTimerDone()
+{
+	if (MyOwnerComp)
+	{
+		FinishLatentTask(*MyOwnerComp, EBTNodeResult::Succeeded);
+	}
+}
+
+#if WITH_EDITOR
+
+FName UBTTask_PlayAnimation::GetNodeIconName() const
+{
+	return FName("BTEditor.Graph.BTNode.Task.PlaySound.Icon");
+}
+
+#endif	// WITH_EDITOR
