Fixing build

diff --git a/Engine/Source/Runtime/AIModule/Private/BehaviorTree/Tasks/BTTask_PlayAnimation.cpp b/Engine/Source/Runtime/AIModule/Private/BehaviorTree/Tasks/BTTask_PlayAnimation.cpp
--- a/Engine/Source/Runtime/AIModule/Private/BehaviorTree/Tasks/BTTask_PlayAnimation.cpp
+++ b/Engine/Source/Runtime/AIModule/Private/BehaviorTree/Tasks/BTTask_PlayAnimation.cpp
@@ -2,6 +2,8 @@
 
 #include "AIModulePrivate.h"
 #include "BehaviorTree/Tasks/BTTask_PlayAnimation.h"
+#include "GameFramework/Character.h"
+#include "Animation/AnimationAsset.h"
 
 //----------------------------------------------------------------------//
 // UBTTask_PlayAnimation
