Another build fix

diff --git a/Engine/Source/Runtime/AIModule/Private/BehaviorTree/Tasks/BTTask_PlayAnimation.cpp b/Engine/Source/Runtime/AIModule/Private/BehaviorTree/Tasks/BTTask_PlayAnimation.cpp
--- a/Engine/Source/Runtime/AIModule/Private/BehaviorTree/Tasks/BTTask_PlayAnimation.cpp
+++ b/Engine/Source/Runtime/AIModule/Private/BehaviorTree/Tasks/BTTask_PlayAnimation.cpp
@@ -4,6 +4,7 @@
 #include "BehaviorTree/Tasks/BTTask_PlayAnimation.h"
 #include "GameFramework/Character.h"
 #include "Animation/AnimationAsset.h"
+#include "Components/SkeletalMeshComponent.h"
 
 //----------------------------------------------------------------------//
 // UBTTask_PlayAnimation
