integrated pawn action changes from dev branch
integrated behavior tree fixes from dev branch

diff --git a/Engine/Source/Runtime/AIModule/Classes/AITypes.h b/Engine/Source/Runtime/AIModule/Classes/AITypes.h
--- a/Engine/Source/Runtime/AIModule/Classes/AITypes.h
+++ b/Engine/Source/Runtime/AIModule/Classes/AITypes.h
@@ -418,10 +418,3 @@ public:
 	static const FAIRequestID CurrentRequest;
 	static const FAIRequestID InvalidRequest;
 };
-
-// used in BT nodes' memory for nodes spawning actions
-enum class EBTActionMemoryHelper : uint8
-{
-	ActionAbortingDone = 0,
-	WaitingForActionToFinishAborting = 1
-};
diff --git a/Engine/Source/Runtime/AIModule/Classes/Actions/PawnAction.h b/Engine/Source/Runtime/AIModule/Classes/Actions/PawnAction.h
--- a/Engine/Source/Runtime/AIModule/Classes/Actions/PawnAction.h
+++ b/Engine/Source/Runtime/AIModule/Classes/Actions/PawnAction.h
@@ -181,6 +181,7 @@ public:
 	virtual void HandleAIMessage(UBrainComponent*, const FAIMessage&){};
 
 	void SetActionObserver(const FPawnActionEventDelegate& ActionObserver) { this->ActionObserver = ActionObserver; }
+	bool HasActionObserver() const { return ActionObserver.IsBound(); }
 
 	//----------------------------------------------------------------------//
 	// Blueprint interface
diff --git a/Engine/Source/Runtime/AIModule/Classes/Actions/PawnActionsComponent.h b/Engine/Source/Runtime/AIModule/Classes/Actions/PawnActionsComponent.h
--- a/Engine/Source/Runtime/AIModule/Classes/Actions/PawnActionsComponent.h
+++ b/Engine/Source/Runtime/AIModule/Classes/Actions/PawnActionsComponent.h
@@ -140,6 +140,8 @@ public:
 	void DescribeSelfToVisLog(struct FVisualLogEntry* Snapshot) const;
 #endif // ENABLE_VISUAL_LOG
 
+	static FString DescribeEventType(EPawnActionEventType::Type EventType);
+
 	//----------------------------------------------------------------------//
 	// Debugging-testing purposes 
 	//----------------------------------------------------------------------//
diff --git a/Engine/Source/Runtime/AIModule/Classes/BehaviorTree/Tasks/BTTask_PawnActionBase.h b/Engine/Source/Runtime/AIModule/Classes/BehaviorTree/Tasks/BTTask_PawnActionBase.h
new file mode 100644
--- /dev/null
+++ b/Engine/Source/Runtime/AIModule/Classes/BehaviorTree/Tasks/BTTask_PawnActionBase.h
@@ -0,0 +1,44 @@
+// Copyright 1998-2015 Epic Games, Inc. All Rights Reserved.
+
+#pragma once
+
+#include "BehaviorTree/BTTaskNode.h"
+#include "Actions/PawnAction.h"
+#include "BTTask_PawnActionBase.generated.h"
+
+class UPawnAction;
+
+/**
+ * Base class for managing pawn actions
+ *
+ * Task will set itself as action observer before pushing it to AI Controller,
+ * override OnActionEvent if you need any special event handling.
+ *
+ * Please use result returned by PushAction for ExecuteTask function.
+ */
+
+UCLASS(Abstract)
+class AIMODULE_API UBTTask_PawnActionBase : public UBTTaskNode
+{
+	GENERATED_UCLASS_BODY()
+
+	UPROPERTY(EditAnywhere, Category=Action)
+	TEnumAsByte<EAIRequestPriority::Type> ActionPriority;
+
+	virtual EBTNodeResult::Type AbortTask(UBehaviorTreeComponent& OwnerComp, uint8* NodeMemory) override;
+
+protected:
+
+	/** starts executing pawn action */
+	EBTNodeResult::Type PushAction(UBehaviorTreeComponent& OwnerComp, UPawnAction& Action);
+
+	/** action observer, updates state of task */
+	virtual void OnActionEvent(UPawnAction& Action, EPawnActionEventType::Type Event);
+
+public:
+
+	/** helper functions, should be used when behavior tree task deals with pawn actions, but can't derive from this class */
+
+	static EBTNodeResult::Type AbortTaskHandler(UBTTaskNode* TaskNode, UBehaviorTreeComponent& OwnerComp, EAIRequestPriority::Type Priority = EAIRequestPriority::Logic);
+	static void ActionEventHandler(UBTTaskNode* TaskNode, UPawnAction& Action, EPawnActionEventType::Type Event);
+};
diff --git a/Engine/Source/Runtime/AIModule/Private/Actions/PawnActionsComponent.cpp b/Engine/Source/Runtime/AIModule/Private/Actions/PawnActionsComponent.cpp
--- a/Engine/Source/Runtime/AIModule/Private/Actions/PawnActionsComponent.cpp
+++ b/Engine/Source/Runtime/AIModule/Private/Actions/PawnActionsComponent.cpp
@@ -544,4 +544,9 @@ void UPawnActionsComponent::DescribeSelfToVisLog(FVisualLogEntry* Snapshot) cons
 		Snapshot->Status.Add(StatusCategory);
 	}
 }
-#endif // ENABLE_VISUAL_LOG
\ No newline at end of file
+#endif // ENABLE_VISUAL_LOG
+
+FString UPawnActionsComponent::DescribeEventType(EPawnActionEventType::Type EventType)
+{
+	return GetEventName(EventType);
+}
diff --git a/Engine/Source/Runtime/AIModule/Private/BehaviorTree/BehaviorTreeComponent.cpp b/Engine/Source/Runtime/AIModule/Private/BehaviorTree/BehaviorTreeComponent.cpp
--- a/Engine/Source/Runtime/AIModule/Private/BehaviorTree/BehaviorTreeComponent.cpp
+++ b/Engine/Source/Runtime/AIModule/Private/BehaviorTree/BehaviorTreeComponent.cpp
@@ -236,6 +236,13 @@ void UBehaviorTreeComponent::OnTaskFinished(const UBTTaskNode* TaskNode, EBTNode
 				RequestExecution(TaskResult);
 			}
 		}
+		else if (TaskResult == EBTNodeResult::Aborted && InstanceStack[TaskInstanceIdx].ActiveNode == TaskNode)
+		{
+			// active instance may be already changed when getting back from AbortCurrentTask 
+			// (e.g. new task is higher on stack)
+
+			InstanceStack[TaskInstanceIdx].ActiveNodeType = EBTActiveNode::InactiveTask;
+		}
 
 		// update state of aborting tasks after currently finished one was set to Inactive
 		UpdateAbortingTasks();
@@ -1133,6 +1140,7 @@ void UBehaviorTreeComponent::ExecuteTask(UBTTaskNode* TaskNode)
 	uint8* NodeMemory = (uint8*)(TaskNode->GetNodeMemory<uint8>(ActiveInstance));
 	EBTNodeResult::Type TaskResult = TaskNode->WrappedExecuteTask(*this, NodeMemory);
 
+	// pass task finished if wasn't already notified (FinishLatentTask)
 	const UBTNode* ActiveNodeAfterExecution = GetActiveNode();
 	if (ActiveNodeAfterExecution == TaskNode)
 	{
@@ -1145,23 +1153,26 @@ void UBehaviorTreeComponent::ExecuteTask(UBTTaskNode* TaskNode)
 
 void UBehaviorTreeComponent::AbortCurrentTask()
 {
-	FBehaviorTreeInstance& ActiveInstance = InstanceStack.Last();
-	UBTTaskNode* ActiveTask = (UBTTaskNode*)ActiveInstance.ActiveNode;
+	const int32 CurrentInstanceIdx = InstanceStack.Num() - 1;
+	FBehaviorTreeInstance& CurrentInstance = InstanceStack[CurrentInstanceIdx];
+	CurrentInstance.ActiveNodeType = EBTActiveNode::AbortingTask;
+
+	UBTTaskNode* CurrentTask = (UBTTaskNode*)CurrentInstance.ActiveNode;
 
 	// remove all observers before requesting abort
-	UnregisterMessageObserversFrom(ActiveTask);
+	UnregisterMessageObserversFrom(CurrentTask);
 
-	UE_VLOG(GetOwner(), LogBehaviorTree, Log, TEXT("Abort task: %s"), *UBehaviorTreeTypes::DescribeNodeHelper(ActiveTask));
+	UE_VLOG(GetOwner(), LogBehaviorTree, Log, TEXT("Abort task: %s"), *UBehaviorTreeTypes::DescribeNodeHelper(CurrentTask));
 
 	// abort task using current state of tree
-	uint8* NodeMemory = (uint8*)(ActiveTask->GetNodeMemory<uint8>(ActiveInstance));
-	EBTNodeResult::Type TaskResult = ActiveTask->WrappedAbortTask(*this, NodeMemory);
+	uint8* NodeMemory = (uint8*)(CurrentTask->GetNodeMemory<uint8>(CurrentInstance));
+	EBTNodeResult::Type TaskResult = CurrentTask->WrappedAbortTask(*this, NodeMemory);
 
-	const UBTNode* ActiveNodeAfterAbort = GetActiveNode();
-	if (ActiveNodeAfterAbort == ActiveTask)
+	// pass task finished if wasn't already notified (FinishLatentAbort)
+	if (CurrentInstance.ActiveNodeType == EBTActiveNode::AbortingTask &&
+		CurrentInstanceIdx == (InstanceStack.Num() - 1))
 	{
-		ActiveInstance.ActiveNodeType = EBTActiveNode::AbortingTask;
-		OnTaskFinished(ActiveTask, TaskResult);
+		OnTaskFinished(CurrentTask, TaskResult);
 	}
 }
 
diff --git a/Engine/Source/Runtime/AIModule/Private/BehaviorTree/Tasks/BTTask_PawnActionBase.cpp b/Engine/Source/Runtime/AIModule/Private/BehaviorTree/Tasks/BTTask_PawnActionBase.cpp
new file mode 100644
--- /dev/null
+++ b/Engine/Source/Runtime/AIModule/Private/BehaviorTree/Tasks/BTTask_PawnActionBase.cpp
@@ -0,0 +1,110 @@
+// Copyright 1998-2015 Epic Games, Inc. All Rights Reserved.
+
+#include "AIModulePrivate.h"
+#include "AIController.h"
+#include "Actions/PawnActionsComponent.h"
+#include "BehaviorTree/Tasks/BTTask_PawnActionBase.h"
+
+UBTTask_PawnActionBase::UBTTask_PawnActionBase(const FObjectInitializer& ObjectInitializer) : Super(ObjectInitializer)
+{
+	NodeName = "Action Base";
+	ActionPriority = EAIRequestPriority::Logic;
+}
+
+EBTNodeResult::Type UBTTask_PawnActionBase::AbortTask(UBehaviorTreeComponent& OwnerComp, uint8* NodeMemory)
+{
+	return AbortTaskHandler(this, OwnerComp, ActionPriority);
+}
+
+EBTNodeResult::Type UBTTask_PawnActionBase::AbortTaskHandler(UBTTaskNode* TaskNode, UBehaviorTreeComponent& OwnerComp, EAIRequestPriority::Type Priority)
+{
+	EBTNodeResult::Type NodeResult = EBTNodeResult::Aborted;
+
+	AAIController* AIOwner = Cast<AAIController>(OwnerComp.GetOwner());
+	if (AIOwner && AIOwner->GetActionsComp())
+	{
+		const int32 NumRequests = AIOwner->GetActionsComp()->AbortActionsInstigatedBy(TaskNode, Priority);
+		if (NumRequests > 0)
+		{
+			NodeResult = EBTNodeResult::InProgress;
+		}
+	}
+
+	return NodeResult;
+}
+
+EBTNodeResult::Type UBTTask_PawnActionBase::PushAction(UBehaviorTreeComponent& OwnerComp, UPawnAction& Action)
+{
+	EBTNodeResult::Type NodeResult = EBTNodeResult::Failed;
+
+	AAIController* AIOwner = Cast<AAIController>(OwnerComp.GetOwner());
+	if (AIOwner)
+	{
+		if (Action.HasActionObserver())
+		{
+			UE_VLOG(AIOwner, LogBehaviorTree, Warning, TEXT("Action %s already had an observer, it will be overwritten!"), *Action.GetName());
+		}
+
+		Action.SetActionObserver(FPawnActionEventDelegate::CreateUObject(this, &UBTTask_PawnActionBase::OnActionEvent));
+
+		const bool bResult = AIOwner->PerformAction(Action, ActionPriority, this);		
+		if (bResult)
+		{
+			// don't bother with handling action events, it will be processed in next tick
+			NodeResult = EBTNodeResult::InProgress;
+		}
+	}
+
+	return NodeResult;
+}
+
+void UBTTask_PawnActionBase::OnActionEvent(UPawnAction& Action, EPawnActionEventType::Type Event)
+{
+	ActionEventHandler(this, Action, Event);
+}
+
+void UBTTask_PawnActionBase::ActionEventHandler(UBTTaskNode* TaskNode, UPawnAction& Action, EPawnActionEventType::Type Event)
+{
+	AAIController* AIOwner = Cast<AAIController>(Action.GetController());
+	UBehaviorTreeComponent* OwnerComp = AIOwner ? Cast<UBehaviorTreeComponent>(AIOwner->BrainComponent) : nullptr;
+	if (OwnerComp == nullptr)
+	{
+		UE_VLOG(Action.GetController(), LogBehaviorTree, Warning, TEXT("%s: Unhandled PawnAction event: %s, can't find owning component!"),
+			*UBehaviorTreeTypes::DescribeNodeHelper(TaskNode),
+			*UPawnActionsComponent::DescribeEventType(Event));
+
+		return;
+	}
+
+	const EBTTaskStatus::Type TaskStatus = OwnerComp->GetTaskStatus(TaskNode);
+	bool bHandled = false;
+
+	if (TaskStatus == EBTTaskStatus::Active)
+	{
+		if (Event == EPawnActionEventType::FinishedExecution || Event == EPawnActionEventType::FailedToStart)
+		{
+			const bool bSucceeded = (Action.GetResult() == EPawnActionResult::Success);
+			TaskNode->FinishLatentTask(*OwnerComp, bSucceeded ? EBTNodeResult::Succeeded : EBTNodeResult::Failed);
+			
+			bHandled = true;
+		}
+	}
+	else if (TaskStatus == EBTTaskStatus::Aborting)
+	{
+		if (Event == EPawnActionEventType::FinishedAborting ||
+			Event == EPawnActionEventType::FinishedExecution || Event == EPawnActionEventType::FailedToStart)
+		{
+			TaskNode->FinishLatentAbort(*OwnerComp);
+
+			bHandled = true;
+		}
+	}
+
+	if (!bHandled)
+	{
+		UE_VLOG(Action.GetController(), LogBehaviorTree, Verbose, TEXT("%s: Unhandled PawnAction event: %s TaskStatus: %s"),
+			*UBehaviorTreeTypes::DescribeNodeHelper(TaskNode),
+			*UPawnActionsComponent::DescribeEventType(Event),
+			*UBehaviorTreeTypes::DescribeTaskStatus(TaskStatus));
+	}
+}
