integrated behavior tree full restart fix from dev branch

diff --git a/Engine/Source/Runtime/AIModule/Private/BehaviorTree/BehaviorTreeComponent.cpp b/Engine/Source/Runtime/AIModule/Private/BehaviorTree/BehaviorTreeComponent.cpp
--- a/Engine/Source/Runtime/AIModule/Private/BehaviorTree/BehaviorTreeComponent.cpp
+++ b/Engine/Source/Runtime/AIModule/Private/BehaviorTree/BehaviorTreeComponent.cpp
@@ -527,9 +527,8 @@ void UBehaviorTreeComponent::RequestExecution(UBTCompositeNode* RequestedOn, int
 	ExecutionIdx.ExecutionIndex = RequestedBy->GetExecutionIndex();
 	uint16 LastExecutionIndex = MAX_uint16;
 
-	if (bSwitchToHigherPriority)
+	if (bSwitchToHigherPriority && RequestedByChildIndex >= 0)
 	{
-		RequestedByChildIndex = FMath::Max(0, RequestedByChildIndex);
 		ExecutionIdx.ExecutionIndex = RequestedOn->GetChildExecutionIndex(RequestedByChildIndex, EBTChildIndex::FirstNode);
 		
 		// first index outside allowed range		
