Blackboard keys' values synchronization #UE4

- original implementation by [REDACTED_PERSON] @ Black Tusk Studios (https://udn.unrealengine.com/questions/228443/behaviortree-synchronized-blackboard-fields.html)
- butchered and optimized a bit in the process of integration

diff --git a/Engine/Source/Developer/AITestSuite/Private/MockAI/MockAI_BT.cpp b/Engine/Source/Developer/AITestSuite/Private/MockAI/MockAI_BT.cpp
--- a/Engine/Source/Developer/AITestSuite/Private/MockAI/MockAI_BT.cpp
+++ b/Engine/Source/Developer/AITestSuite/Private/MockAI/MockAI_BT.cpp
@@ -27,7 +27,10 @@ bool UMockAI_BT::IsRunning() const
 
 bool UMockAI_BT::RunBT(UBehaviorTree& BTAsset, EBTExecutionMode::Type RunType)
 {
-	BBComp->InitializeBlackboard(BTAsset.BlackboardAsset);
+	if (BTAsset.BlackboardAsset)
+	{
+		BBComp->InitializeBlackboard(*BTAsset.BlackboardAsset);
+	}
 	BBComp->CacheBrainComponent(*BTComp);
 	BTComp->CacheBlackboardComponent(BBComp);
 
diff --git a/Engine/Source/Editor/BehaviorTreeEditor/Private/DetailCustomizations/BlackboardDataDetails.cpp b/Engine/Source/Editor/BehaviorTreeEditor/Private/DetailCustomizations/BlackboardDataDetails.cpp
--- a/Engine/Source/Editor/BehaviorTreeEditor/Private/DetailCustomizations/BlackboardDataDetails.cpp
+++ b/Engine/Source/Editor/BehaviorTreeEditor/Private/DetailCustomizations/BlackboardDataDetails.cpp
@@ -61,6 +61,9 @@ void FBlackboardDataDetails::CustomizeDetails( IDetailLayoutBuilder& DetailLayou
 
 			TSharedPtr<IPropertyHandle> KeyTypeProperty = KeyHandle->GetChildHandle(GET_MEMBER_NAME_CHECKED(FBlackboardEntry, KeyType));
 			DetailCategoryBuilder.AddProperty(KeyTypeProperty);
+
+			TSharedPtr<IPropertyHandle> bInstanceSyncedProperty = KeyHandle->GetChildHandle(GET_MEMBER_NAME_CHECKED(FBlackboardEntry, bInstanceSynced));
+			DetailCategoryBuilder.AddProperty(bInstanceSyncedProperty);
 		}	
 	}
 }
diff --git a/Engine/Source/Runtime/AIModule/Classes/AIController.h b/Engine/Source/Runtime/AIModule/Classes/AIController.h
--- a/Engine/Source/Runtime/AIModule/Classes/AIController.h
+++ b/Engine/Source/Runtime/AIModule/Classes/AIController.h
@@ -10,6 +10,7 @@
 #include "GameFramework/Controller.h"
 #include "Actions/PawnActionsComponent.h"
 #include "Perception/AIPerceptionListenerInterface.h"
+#include "BehaviorTree/BehaviorTreeTypes.h"
 #include "AIController.generated.h"
 
 class APawn;
@@ -225,6 +226,10 @@ public:
 	UFUNCTION(BlueprintCallable, Category = "AI")
 	bool UseBlackboard(UBlackboardData* BlackboardAsset);
 
+	/** does this AIController allow given UBlackboardComponent sync data with it */
+	virtual bool ShouldSyncBlackboardWith(const UBlackboardComponent& OtherBlackboardComponent) const { return true; }
+		
+
 protected:
 	UFUNCTION(BlueprintImplementableEvent)
 	virtual void OnUsingBlackBoard(UBlackboardComponent* BlackboardComp, UBlackboardData* BlackboardAsset);
diff --git a/Engine/Source/Runtime/AIModule/Classes/AISystem.h b/Engine/Source/Runtime/AIModule/Classes/AISystem.h
--- a/Engine/Source/Runtime/AIModule/Classes/AISystem.h
+++ b/Engine/Source/Runtime/AIModule/Classes/AISystem.h
@@ -12,6 +12,8 @@ class UEnvQueryManager;
 class UAIPerceptionSystem;
 class UAIAsyncTaskBlueprintProxy;
 class UAIHotSpotManager;
+class UBlackboardData;
+class UBlackboardComponent;
 
 UCLASS(config=Engine, defaultconfig)
 class AIMODULE_API UAISystem : public UAISystemBase
@@ -41,6 +43,11 @@ protected:
 
 	UPROPERTY(Transient)
 	UAIHotSpotManager* HotSpotManager;
+
+	typedef TMultiMap<TWeakObjectPtr<UBlackboardData>, TWeakObjectPtr<UBlackboardComponent>> FBlackboardDataToComponentsMap;
+
+	/** UBlackboardComponent instances that reference the blackboard data definition */
+	FBlackboardDataToComponentsMap BlackboardDataToComponentsMap;
 	
 public:
 	virtual ~UAISystem();
@@ -98,4 +105,76 @@ public:
 
 	/** insta-runs EQS query for given Target */
 	void RunEQS(const FString& QueryName, UObject* Target);
+
+	/**
+	* Iterator for traversing all UBlackboardComponent instances associated
+	* with this blackboard data asset. This is a forward only iterator.
+	*/
+	struct FBlackboardDataToComponentsIterator
+	{
+	public:
+		FBlackboardDataToComponentsIterator(FBlackboardDataToComponentsMap& BlackboardDataToComponentsMap, class UBlackboardData* BlackboardAsset);
+
+		FORCEINLINE FBlackboardDataToComponentsIterator& operator++()
+		{
+			++GetCurrentIteratorRef();
+			TryMoveIteratorToParentBlackboard();
+			return *this;
+		}
+		FORCEINLINE FBlackboardDataToComponentsIterator operator++(int)
+		{
+			FBlackboardDataToComponentsIterator Tmp(*this);
+			++GetCurrentIteratorRef();
+			TryMoveIteratorToParentBlackboard();
+			return Tmp;
+		}
+
+		SAFE_BOOL_OPERATORS(FBlackboardDataToComponentsIterator);
+		FORCEINLINE_EXPLICIT_OPERATOR_BOOL() const { return CurrentIteratorIndex < Iterators.Num() && (bool)GetCurrentIteratorRef(); }
+		FORCEINLINE bool operator !() const { return !(bool)*this; }
+
+		FORCEINLINE UBlackboardData* Key() const { return GetCurrentIteratorRef().Key().Get(); }
+		FORCEINLINE UBlackboardComponent* Value() const { return GetCurrentIteratorRef().Value().Get(); }
+
+	private:
+		FORCEINLINE const FBlackboardDataToComponentsMap::TConstKeyIterator& GetCurrentIteratorRef() const { return Iterators[CurrentIteratorIndex]; }
+		FORCEINLINE FBlackboardDataToComponentsMap::TConstKeyIterator& GetCurrentIteratorRef() { return Iterators[CurrentIteratorIndex]; }
+
+		void TryMoveIteratorToParentBlackboard()
+		{
+			if (!GetCurrentIteratorRef() && CurrentIteratorIndex < Iterators.Num() - 1)
+			{
+				++CurrentIteratorIndex;
+				TryMoveIteratorToParentBlackboard(); // keep incrementing until we find a valid iterator.
+			}
+		}
+
+		int32 CurrentIteratorIndex;
+
+		static const int32 InlineSize = 8;
+		TArray<FBlackboardDataToComponentsMap::TConstKeyIterator, TInlineAllocator<InlineSize>> Iterators;
+	};
+
+	/**
+	* Registers a UBlackboardComponent instance with this blackboard data asset.
+	* This will also register the component for each parent UBlackboardData
+	* asset. This should be called after the component has been initialized
+	* (i.e. InitializeComponent). The user is responsible for calling
+	* UnregisterBlackboardComponent (i.e. UninitializeComponent).
+	*/
+	void RegisterBlackboardComponent(class UBlackboardData& BlackboardAsset, class UBlackboardComponent& BlackboardComp);
+
+	/**
+	* Unregisters a UBlackboardComponent instance with this blackboard data
+	* asset. This should be called before the component has been uninitialized
+	* (i.e. UninitializeComponent).
+	*/
+	void UnregisterBlackboardComponent(class UBlackboardData& BlackboardAsset, class UBlackboardComponent& BlackboardComp);
+
+	/**
+	* Creates a forward only iterator for that will iterate all
+	* UBlackboardComponent instances that reference the specified
+	* BlackboardAsset and it's parents.
+	*/
+	FBlackboardDataToComponentsIterator CreateBlackboardDataToComponentsIterator(class UBlackboardData& BlackboardAsset);
 };
diff --git a/Engine/Source/Runtime/AIModule/Classes/BehaviorTree/Blackboard/BlackboardKeyType.h b/Engine/Source/Runtime/AIModule/Classes/BehaviorTree/Blackboard/BlackboardKeyType.h
--- a/Engine/Source/Runtime/AIModule/Classes/BehaviorTree/Blackboard/BlackboardKeyType.h
+++ b/Engine/Source/Runtime/AIModule/Classes/BehaviorTree/Blackboard/BlackboardKeyType.h
@@ -2,6 +2,8 @@
 #pragma once
 #include "BlackboardKeyType.generated.h"
 
+class UBlackboardComponent;
+
 namespace EBlackboardCompare
 {
 	enum Type
@@ -142,4 +144,10 @@ protected:
 
 		return bChanged;
 	}
+
+	friend UBlackboardComponent;
+	void CopyValue(uint8* To, const uint8* From) const
+	{
+		FMemory::Memcpy(To, From, GetValueSize());
+	}
 };
diff --git a/Engine/Source/Runtime/AIModule/Classes/BehaviorTree/Blackboard/BlackboardKeyType_Bool.h b/Engine/Source/Runtime/AIModule/Classes/BehaviorTree/Blackboard/BlackboardKeyType_Bool.h
--- a/Engine/Source/Runtime/AIModule/Classes/BehaviorTree/Blackboard/BlackboardKeyType_Bool.h
+++ b/Engine/Source/Runtime/AIModule/Classes/BehaviorTree/Blackboard/BlackboardKeyType_Bool.h
@@ -8,6 +8,8 @@ class AIMODULE_API UBlackboardKeyType_Bool : public UBlackboardKeyType
 {
 	GENERATED_UCLASS_BODY()
 
+	typedef bool FDataType;
+
 	static bool GetValue(const uint8* RawData);
 	static bool SetValue(uint8* RawData, bool bValue);
 
diff --git a/Engine/Source/Runtime/AIModule/Classes/BehaviorTree/Blackboard/BlackboardKeyType_Class.h b/Engine/Source/Runtime/AIModule/Classes/BehaviorTree/Blackboard/BlackboardKeyType_Class.h
--- a/Engine/Source/Runtime/AIModule/Classes/BehaviorTree/Blackboard/BlackboardKeyType_Class.h
+++ b/Engine/Source/Runtime/AIModule/Classes/BehaviorTree/Blackboard/BlackboardKeyType_Class.h
@@ -8,6 +8,8 @@ class AIMODULE_API UBlackboardKeyType_Class : public UBlackboardKeyType
 {
 	GENERATED_UCLASS_BODY()
 
+	typedef UClass* FDataType;
+
 	UPROPERTY(Category=Blackboard, EditDefaultsOnly, meta=(AllowAbstract="1"))
 	UClass* BaseClass;
 	
diff --git a/Engine/Source/Runtime/AIModule/Classes/BehaviorTree/Blackboard/BlackboardKeyType_Enum.h b/Engine/Source/Runtime/AIModule/Classes/BehaviorTree/Blackboard/BlackboardKeyType_Enum.h
--- a/Engine/Source/Runtime/AIModule/Classes/BehaviorTree/Blackboard/BlackboardKeyType_Enum.h
+++ b/Engine/Source/Runtime/AIModule/Classes/BehaviorTree/Blackboard/BlackboardKeyType_Enum.h
@@ -8,6 +8,8 @@ class AIMODULE_API UBlackboardKeyType_Enum : public UBlackboardKeyType
 {
 	GENERATED_UCLASS_BODY()
 
+	typedef uint8 FDataType;
+
 	UPROPERTY(Category=Blackboard, EditDefaultsOnly)
 	UEnum* EnumType;
 
diff --git a/Engine/Source/Runtime/AIModule/Classes/BehaviorTree/Blackboard/BlackboardKeyType_Float.h b/Engine/Source/Runtime/AIModule/Classes/BehaviorTree/Blackboard/BlackboardKeyType_Float.h
--- a/Engine/Source/Runtime/AIModule/Classes/BehaviorTree/Blackboard/BlackboardKeyType_Float.h
+++ b/Engine/Source/Runtime/AIModule/Classes/BehaviorTree/Blackboard/BlackboardKeyType_Float.h
@@ -8,6 +8,8 @@ class AIMODULE_API UBlackboardKeyType_Float : public UBlackboardKeyType
 {
 	GENERATED_UCLASS_BODY()
 
+	typedef float FDataType;
+
 	static float GetValue(const uint8* RawData);
 	static bool SetValue(uint8* RawData, float Value);
 
diff --git a/Engine/Source/Runtime/AIModule/Classes/BehaviorTree/Blackboard/BlackboardKeyType_Int.h b/Engine/Source/Runtime/AIModule/Classes/BehaviorTree/Blackboard/BlackboardKeyType_Int.h
--- a/Engine/Source/Runtime/AIModule/Classes/BehaviorTree/Blackboard/BlackboardKeyType_Int.h
+++ b/Engine/Source/Runtime/AIModule/Classes/BehaviorTree/Blackboard/BlackboardKeyType_Int.h
@@ -8,6 +8,8 @@ class AIMODULE_API UBlackboardKeyType_Int : public UBlackboardKeyType
 {
 	GENERATED_UCLASS_BODY()
 
+	typedef int32 FDataType;
+
 	static int32 GetValue(const uint8* RawData);
 	static bool SetValue(uint8* RawData, int32 Value);
 
diff --git a/Engine/Source/Runtime/AIModule/Classes/BehaviorTree/Blackboard/BlackboardKeyType_Name.h b/Engine/Source/Runtime/AIModule/Classes/BehaviorTree/Blackboard/BlackboardKeyType_Name.h
--- a/Engine/Source/Runtime/AIModule/Classes/BehaviorTree/Blackboard/BlackboardKeyType_Name.h
+++ b/Engine/Source/Runtime/AIModule/Classes/BehaviorTree/Blackboard/BlackboardKeyType_Name.h
@@ -8,6 +8,8 @@ class AIMODULE_API UBlackboardKeyType_Name : public UBlackboardKeyType
 {
 	GENERATED_UCLASS_BODY()
 
+	typedef FName FDataType;
+
 	static FName GetValue(const uint8* RawData);
 	static bool SetValue(uint8* RawData, const FName& Value);
 
diff --git a/Engine/Source/Runtime/AIModule/Classes/BehaviorTree/Blackboard/BlackboardKeyType_NativeEnum.h b/Engine/Source/Runtime/AIModule/Classes/BehaviorTree/Blackboard/BlackboardKeyType_NativeEnum.h
--- a/Engine/Source/Runtime/AIModule/Classes/BehaviorTree/Blackboard/BlackboardKeyType_NativeEnum.h
+++ b/Engine/Source/Runtime/AIModule/Classes/BehaviorTree/Blackboard/BlackboardKeyType_NativeEnum.h
@@ -9,6 +9,8 @@ class AIMODULE_API UBlackboardKeyType_NativeEnum : public UBlackboardKeyType
 {
 	GENERATED_UCLASS_BODY()
 
+	typedef uint8 FDataType;
+
 	UPROPERTY(Category=Blackboard, EditDefaultsOnly)
 	FString EnumName;
 
diff --git a/Engine/Source/Runtime/AIModule/Classes/BehaviorTree/Blackboard/BlackboardKeyType_Object.h b/Engine/Source/Runtime/AIModule/Classes/BehaviorTree/Blackboard/BlackboardKeyType_Object.h
--- a/Engine/Source/Runtime/AIModule/Classes/BehaviorTree/Blackboard/BlackboardKeyType_Object.h
+++ b/Engine/Source/Runtime/AIModule/Classes/BehaviorTree/Blackboard/BlackboardKeyType_Object.h
@@ -8,6 +8,8 @@ class AIMODULE_API UBlackboardKeyType_Object : public UBlackboardKeyType
 {
 	GENERATED_UCLASS_BODY()
 
+	typedef UObject* FDataType;
+
 	UPROPERTY(Category=Blackboard, EditDefaultsOnly, meta=(AllowAbstract="1"))
 	UClass* BaseClass;
 
diff --git a/Engine/Source/Runtime/AIModule/Classes/BehaviorTree/Blackboard/BlackboardKeyType_Rotator.h b/Engine/Source/Runtime/AIModule/Classes/BehaviorTree/Blackboard/BlackboardKeyType_Rotator.h
--- a/Engine/Source/Runtime/AIModule/Classes/BehaviorTree/Blackboard/BlackboardKeyType_Rotator.h
+++ b/Engine/Source/Runtime/AIModule/Classes/BehaviorTree/Blackboard/BlackboardKeyType_Rotator.h
@@ -6,7 +6,9 @@ UCLASS(EditInlineNew, meta=(DisplayName="Rotator"))
 class AIMODULE_API UBlackboardKeyType_Rotator : public UBlackboardKeyType
 {
 	GENERATED_UCLASS_BODY()
-
+	
+	typedef FRotator FDataType; 
+	
 	static FRotator GetValue(const uint8* RawData);
 	static bool SetValue(uint8* RawData, const FRotator& Value);
 
diff --git a/Engine/Source/Runtime/AIModule/Classes/BehaviorTree/Blackboard/BlackboardKeyType_String.h b/Engine/Source/Runtime/AIModule/Classes/BehaviorTree/Blackboard/BlackboardKeyType_String.h
--- a/Engine/Source/Runtime/AIModule/Classes/BehaviorTree/Blackboard/BlackboardKeyType_String.h
+++ b/Engine/Source/Runtime/AIModule/Classes/BehaviorTree/Blackboard/BlackboardKeyType_String.h
@@ -8,6 +8,8 @@ class AIMODULE_API UBlackboardKeyType_String : public UBlackboardKeyType
 {
 	GENERATED_UCLASS_BODY()
 
+	typedef FString FDataType;
+
 	static FString GetValue(const uint8* RawData);
 	static bool SetValue(uint8* RawData, const FString& Value);
 
diff --git a/Engine/Source/Runtime/AIModule/Classes/BehaviorTree/Blackboard/BlackboardKeyType_Vector.h b/Engine/Source/Runtime/AIModule/Classes/BehaviorTree/Blackboard/BlackboardKeyType_Vector.h
--- a/Engine/Source/Runtime/AIModule/Classes/BehaviorTree/Blackboard/BlackboardKeyType_Vector.h
+++ b/Engine/Source/Runtime/AIModule/Classes/BehaviorTree/Blackboard/BlackboardKeyType_Vector.h
@@ -8,6 +8,8 @@ class AIMODULE_API UBlackboardKeyType_Vector : public UBlackboardKeyType
 {
 	GENERATED_UCLASS_BODY()
 
+	typedef FVector FDataType; 
+	
 	static FVector GetValue(const uint8* RawData);
 	static bool SetValue(uint8* RawData, const FVector& Value);
 
diff --git a/Engine/Source/Runtime/AIModule/Classes/BehaviorTree/BlackboardComponent.h b/Engine/Source/Runtime/AIModule/Classes/BehaviorTree/BlackboardComponent.h
--- a/Engine/Source/Runtime/AIModule/Classes/BehaviorTree/BlackboardComponent.h
+++ b/Engine/Source/Runtime/AIModule/Classes/BehaviorTree/BlackboardComponent.h
@@ -35,10 +35,14 @@ namespace EBlackboardDescription
 UCLASS()
 class AIMODULE_API UBlackboardComponent : public UActorComponent
 {
-	GENERATED_UCLASS_BODY()
+	GENERATED_BODY()
+
+public:
+	UBlackboardComponent(const FObjectInitializer& ObjectInitializer);
 
 	/** BEGIN UActorComponent overrides */
 	virtual void InitializeComponent() override;
+	virtual void UninitializeComponent() override;
 	/** END UActorComponent overrides */
 
 	/** @return name of key */
@@ -50,6 +54,9 @@ class AIMODULE_API UBlackboardComponent : public UActorComponent
 	/** @return class of value for given key */
 	TSubclassOf<UBlackboardKeyType> GetKeyType(FBlackboard::FKey KeyID) const;
 
+	/** @return true if the key is marked as instance synced */
+	bool IsKeyInstanceSynced(FBlackboard::FKey KeyID) const;
+	
 	/** @return number of entries in data asset */
 	int32 GetNumKeys() const;
 
@@ -78,7 +85,7 @@ class AIMODULE_API UBlackboardComponent : public UActorComponent
 	void CacheBrainComponent(UBrainComponent& BrainComponent);
 
 	/** setup component for using given blackboard asset */
-	bool InitializeBlackboard(UBlackboardData* NewAsset);
+	bool InitializeBlackboard(UBlackboardData& NewAsset);
 	
 	/** @return true if component can be used with specified blackboard asset */
 	bool IsCompatibleWith(UBlackboardData* TestAsset) const;
@@ -125,39 +132,30 @@ class AIMODULE_API UBlackboardComponent : public UActorComponent
 
 	UFUNCTION(BlueprintCallable, Category="AI|Components|Blackboard")
 	void SetValueAsObject(const FName& KeyName, UObject* ObjectValue);
-	void SetValueAsObject(FBlackboard::FKey KeyID, UObject* ObjectValue);
 
 	UFUNCTION(BlueprintCallable, Category="AI|Components|Blackboard")
 	void SetValueAsClass(const FName& KeyName, UClass* ClassValue);
-	void SetValueAsClass(FBlackboard::FKey KeyID, UClass* ClassValue);
 
 	UFUNCTION(BlueprintCallable, Category="AI|Components|Blackboard")
 	void SetValueAsEnum(const FName& KeyName, uint8 EnumValue);
-	void SetValueAsEnum(FBlackboard::FKey KeyID, uint8 EnumValue);
 
 	UFUNCTION(BlueprintCallable, Category="AI|Components|Blackboard")
 	void SetValueAsInt(const FName& KeyName, int32 IntValue);
-	void SetValueAsInt(FBlackboard::FKey KeyID, int32 IntValue);
 
 	UFUNCTION(BlueprintCallable, Category="AI|Components|Blackboard")
 	void SetValueAsFloat(const FName& KeyName, float FloatValue);
-	void SetValueAsFloat(FBlackboard::FKey KeyID, float FloatValue);
 
 	UFUNCTION(BlueprintCallable, Category="AI|Components|Blackboard")
 	void SetValueAsBool(const FName& KeyName, bool BoolValue);
-	void SetValueAsBool(FBlackboard::FKey KeyID, bool BoolValue);
 
 	UFUNCTION(BlueprintCallable, Category="AI|Components|Blackboard")
 	void SetValueAsString(const FName& KeyName, FString StringValue);
-	void SetValueAsString(FBlackboard::FKey KeyID, const FString& StringValue);
 
 	UFUNCTION(BlueprintCallable, Category="AI|Components|Blackboard")
 	void SetValueAsName(const FName& KeyName, FName NameValue);
-	void SetValueAsName(FBlackboard::FKey KeyID, const FName& NameValue);
 
 	UFUNCTION(BlueprintCallable, Category="AI|Components|Blackboard")
 	void SetValueAsVector(const FName& KeyName, FVector VectorValue);
-	void SetValueAsVector(FBlackboard::FKey KeyID, const FVector& VectorValue);
 
 	UFUNCTION(BlueprintCallable, Category="AI|Components|Blackboard")
 	void ClearValueAsVector(const FName& KeyName);
@@ -190,6 +188,9 @@ class AIMODULE_API UBlackboardComponent : public UActorComponent
 	void ClearValue(const FName& KeyName);
 	void ClearValue(FBlackboard::FKey KeyID);
 
+	template<class TDataClass>
+	bool SetValue(FBlackboard::FKey KeyID, typename TDataClass::FDataType Value);
+
 	/** get pointer to raw data for given key */
 	FORCEINLINE uint8* GetKeyRawData(const FName& KeyName) { return GetKeyRawData(GetKeyID(KeyName)); }
 	FORCEINLINE uint8* GetKeyRawData(FBlackboard::FKey KeyID) { return ValueMemory.Num() && ValueOffsets.IsValidIndex(KeyID) ? (ValueMemory.GetData() + ValueOffsets[KeyID]) : NULL; }
@@ -238,11 +239,19 @@ protected:
 	/** set when notifies are paused and shouldn't be passed to observers */
 	uint32 bPausedNotifies : 1;
 
+	/** reset to false every time a new BB asset is assigned to this component */
+	uint32 bSynchronizedKeyPopulated : 1;
+
 	/** notifies behavior tree decorators about change in blackboard */
 	void NotifyObservers(FBlackboard::FKey KeyID) const;
 
 	/** initializes parent chain in asset */
 	void InitializeParentChain(UBlackboardData* NewAsset);
+
+	/** populates BB's synchronized entries */
+	void PopulateSynchronizedKeys();
+
+	bool ShouldSyncWithBlackboard(UBlackboardComponent& OtherBlackboardComponent) const;
 };
 
 //////////////////////////////////////////////////////////////////////////
@@ -252,3 +261,39 @@ FORCEINLINE bool UBlackboardComponent::HasValidAsset() const
 {
 	return BlackboardAsset && BlackboardAsset->IsValid();
 }
+
+template<class TDataClass>
+bool UBlackboardComponent::SetValue(FBlackboard::FKey KeyID, typename TDataClass::FDataType Value)
+{
+	if (GetKeyType(KeyID) != TDataClass::StaticClass())
+	{
+		return false;
+	}
+
+	uint8* RawData = GetKeyRawData(KeyID);
+	if (RawData)
+	{
+		const bool bChanged = TDataClass::SetValue(RawData, Value);
+		if (bChanged)
+		{
+			NotifyObservers(KeyID);
+			if (BlackboardAsset->HasSyncronizedKeys() && IsKeyInstanceSynced(KeyID))
+			{
+				UAISystem* AISystem = UAISystem::GetCurrentSafe(GetWorld());
+				for (auto Iter = AISystem->CreateBlackboardDataToComponentsIterator(*BlackboardAsset); Iter; ++Iter)
+				{
+					UBlackboardComponent* OtherBlackboard = Iter.Value();
+					if (OtherBlackboard != nullptr && ShouldSyncWithBlackboard(*OtherBlackboard))
+					{
+						TDataClass::SetValue(OtherBlackboard->GetKeyRawData(KeyID), Value);
+						OtherBlackboard->NotifyObservers(KeyID);
+					}
+				}
+			}
+		}
+
+		return true;
+	}
+
+	return false;
+}
\ No newline at end of file
diff --git a/Engine/Source/Runtime/AIModule/Classes/BehaviorTree/BlackboardData.h b/Engine/Source/Runtime/AIModule/Classes/BehaviorTree/BlackboardData.h
--- a/Engine/Source/Runtime/AIModule/Classes/BehaviorTree/BlackboardData.h
+++ b/Engine/Source/Runtime/AIModule/Classes/BehaviorTree/BlackboardData.h
@@ -24,6 +24,10 @@ struct FBlackboardEntry
 	UPROPERTY(EditAnywhere, Instanced, Category=Blackboard)
 	UBlackboardKeyType* KeyType;
 
+	/** if set to true then this field will be synchronized across all instances of this blackboard */
+	UPROPERTY(EditAnywhere, Category = Blackboard)
+	uint32 bInstanceSynced : 1;
+
 	bool operator==(const FBlackboardEntry& Other) const;
 };
 
@@ -47,6 +51,17 @@ class AIMODULE_API UBlackboardData : public UDataAsset
 	UPROPERTY(EditAnywhere, Category=Blackboard)
 	TArray<FBlackboardEntry> Keys;
 
+private:
+	UPROPERTY()
+	uint32 bHasSyncronizedKeys : 1;
+
+public:
+
+	FORCEINLINE bool HasSyncronizedKeys() const { return bHasSyncronizedKeys; }
+
+	/** @return true if the key is instance synced */
+	bool IsKeyInstanceSynced(FBlackboard::FKey KeyID) const;
+	
 	/** @return key ID from name */
 	FBlackboard::FKey GetKeyID(const FName& KeyName) const;
 
@@ -66,6 +81,7 @@ class AIMODULE_API UBlackboardData : public UDataAsset
 
 	virtual void PostLoad() override;
 #if WITH_EDITOR
+	virtual void PreEditChange(class FEditPropertyChain& PropertyAboutToChange) override;
 	virtual void PostEditChangeProperty(struct FPropertyChangedEvent& PropertyChangedEvent) override;
 #endif
 
@@ -112,6 +128,8 @@ class AIMODULE_API UBlackboardData : public UDataAsset
 	/** forces update of FirstKeyID, which depends on parent chain */
 	void UpdateKeyIDs();
 
+	void UpdateIfHasSynchronizedKeys();
+
 protected:
 
 	enum EKeyLookupMode
diff --git a/Engine/Source/Runtime/AIModule/Private/AIController.cpp b/Engine/Source/Runtime/AIModule/Private/AIController.cpp
--- a/Engine/Source/Runtime/AIModule/Private/AIController.cpp
+++ b/Engine/Source/Runtime/AIModule/Private/AIController.cpp
@@ -777,7 +777,7 @@ bool AAIController::RunBehaviorTree(UBehaviorTree* BTAsset)
 bool AAIController::InitializeBlackboard(UBlackboardComponent& BlackboardComp, UBlackboardData& BlackboardAsset)
 {
 	check(BlackboardComp.GetOwner() == this);
-	if (BlackboardComp.InitializeBlackboard(&BlackboardAsset))
+	if (BlackboardComp.InitializeBlackboard(BlackboardAsset))
 	{
 		OnUsingBlackBoard(&BlackboardComp, &BlackboardAsset);
 		return true;
diff --git a/Engine/Source/Runtime/AIModule/Private/AISystem.cpp b/Engine/Source/Runtime/AIModule/Private/AISystem.cpp
--- a/Engine/Source/Runtime/AIModule/Private/AISystem.cpp
+++ b/Engine/Source/Runtime/AIModule/Private/AISystem.cpp
@@ -113,4 +113,50 @@ void UAISystem::RunEQS(const FString& QueryName, UObject* Target)
 		MyPC->ClientMessage(TEXT("No debugging target"));
 	}
 #endif
-}
\ No newline at end of file
+}
+
+UAISystem::FBlackboardDataToComponentsIterator::FBlackboardDataToComponentsIterator(FBlackboardDataToComponentsMap& BlackboardDataToComponentsMap, UBlackboardData* BlackboardAsset)
+	: CurrentIteratorIndex(0)
+	, Iterators()
+{
+	while (BlackboardAsset)
+	{
+		Iterators.Add(BlackboardDataToComponentsMap.CreateConstKeyIterator(BlackboardAsset));
+		BlackboardAsset = BlackboardAsset->Parent;
+	}
+	TryMoveIteratorToParentBlackboard();
+}
+
+void UAISystem::RegisterBlackboardComponent(UBlackboardData& BlackboardData, UBlackboardComponent& BlackboardComp)
+{
+	// mismatch of register/unregister.
+	ensure(BlackboardDataToComponentsMap.FindPair(&BlackboardData, &BlackboardComp) == nullptr);
+
+	BlackboardDataToComponentsMap.Add(&BlackboardData, &BlackboardComp);
+	if (BlackboardData.Parent)
+	{
+		RegisterBlackboardComponent(*BlackboardData.Parent, BlackboardComp);
+	}
+}
+
+void UAISystem::UnregisterBlackboardComponent(UBlackboardData& BlackboardData, UBlackboardComponent& BlackboardComp)
+{
+	// this is actually possible, we can end up unregistering before UBlackboardComponent cached its BrainComponent
+	// which currently is tied to the whole process. 
+	// @todo remove this dependency
+	ensure(BlackboardDataToComponentsMap.FindPair(&BlackboardData, &BlackboardComp) != nullptr);
+
+	if (BlackboardData.Parent)
+	{
+		UnregisterBlackboardComponent(*BlackboardData.Parent, BlackboardComp);
+	}
+	BlackboardDataToComponentsMap.RemoveSingle(&BlackboardData, &BlackboardComp);
+
+	// mismatch of Register/Unregister.
+	check(BlackboardDataToComponentsMap.FindPair(&BlackboardData, &BlackboardComp) == nullptr);
+}
+
+UAISystem::FBlackboardDataToComponentsIterator UAISystem::CreateBlackboardDataToComponentsIterator(UBlackboardData& BlackboardAsset)
+{
+	return UAISystem::FBlackboardDataToComponentsIterator(BlackboardDataToComponentsMap, &BlackboardAsset);
+}
diff --git a/Engine/Source/Runtime/AIModule/Private/BehaviorTree/BlackboardComponent.cpp b/Engine/Source/Runtime/AIModule/Private/BehaviorTree/BlackboardComponent.cpp
--- a/Engine/Source/Runtime/AIModule/Private/BehaviorTree/BlackboardComponent.cpp
+++ b/Engine/Source/Runtime/AIModule/Private/BehaviorTree/BlackboardComponent.cpp
@@ -10,6 +10,7 @@ UBlackboardComponent::UBlackboardComponent(const FObjectInitializer& ObjectIniti
 	PrimaryComponentTick.bCanEverTick = false;
 	bWantsInitializeComponent = true;
 	bPausedNotifies = false;
+	bSynchronizedKeyPopulated = false;
 }
 
 void UBlackboardComponent::InitializeComponent()
@@ -28,7 +29,23 @@ void UBlackboardComponent::InitializeComponent()
 		}
 	}
 
-	InitializeBlackboard(BlackboardAsset);
+	if (BlackboardAsset)
+	{
+		InitializeBlackboard(*BlackboardAsset);
+	}
+}
+
+void UBlackboardComponent::UninitializeComponent()
+{
+	if (BlackboardAsset && BlackboardAsset->HasSyncronizedKeys())
+	{
+		UAISystem* AISystem = UAISystem::GetCurrentSafe(GetWorld());
+		if (AISystem)
+		{
+			AISystem->UnregisterBlackboardComponent(*BlackboardAsset, *this);
+		}
+	}
+	Super::UninitializeComponent();
 }
 
 void UBlackboardComponent::CacheBrainComponent(UBrainComponent& BrainComponent)
@@ -68,75 +85,130 @@ void UBlackboardComponent::InitializeParentChain(UBlackboardData* NewAsset)
 	}
 }
 
-bool UBlackboardComponent::InitializeBlackboard(UBlackboardData* NewAsset)
+bool UBlackboardComponent::InitializeBlackboard(UBlackboardData& NewAsset)
 {
 	// if we re-initialize with the same asset then there's no point
 	// in reseting, since we'd loose all the accumulated knowledge
-	if (NewAsset == BlackboardAsset)
+	if (&NewAsset == BlackboardAsset)
+	{
+		return false;
+	}
+
+	UAISystem* AISystem = UAISystem::GetCurrentSafe(GetWorld());
+	if (AISystem == nullptr)
 	{
 		return false;
 	}
 
-	BlackboardAsset = NewAsset;
+	if (BlackboardAsset)
+	{
+		AISystem->UnregisterBlackboardComponent(*BlackboardAsset, *this);
+	}
+
+	BlackboardAsset = &NewAsset;
 	ValueMemory.Reset();
 	ValueOffsets.Reset();
+	bSynchronizedKeyPopulated = false;
 
 	bool bSuccess = true;
 	
-	if (BlackboardAsset)
+	if (BlackboardAsset->IsValid())
 	{
-		if (BlackboardAsset->IsValid())
-		{
-			InitializeParentChain(BlackboardAsset);
+		InitializeParentChain(BlackboardAsset);
 
-			TArray<FBlackboardInitializationData> InitList;
-			const int32 NumKeys = BlackboardAsset->GetNumKeys();
-			InitList.Reserve(NumKeys);
-			ValueOffsets.AddZeroed(NumKeys);
+		TArray<FBlackboardInitializationData> InitList;
+		const int32 NumKeys = BlackboardAsset->GetNumKeys();
+		InitList.Reserve(NumKeys);
+		ValueOffsets.AddZeroed(NumKeys);
 
-			for (UBlackboardData* It = BlackboardAsset; It; It = It->Parent)
+		for (UBlackboardData* It = BlackboardAsset; It; It = It->Parent)
+		{
+			for (int32 KeyIndex = 0; KeyIndex < It->Keys.Num(); KeyIndex++)
 			{
-				for (int32 KeyIndex = 0; KeyIndex < It->Keys.Num(); KeyIndex++)
+				if (It->Keys[KeyIndex].KeyType)
 				{
-					if (It->Keys[KeyIndex].KeyType)
-					{
-						InitList.Add(FBlackboardInitializationData(KeyIndex + It->GetFirstKeyID(), It->Keys[KeyIndex].KeyType->GetValueSize()));
-					}
+					InitList.Add(FBlackboardInitializationData(KeyIndex + It->GetFirstKeyID(), It->Keys[KeyIndex].KeyType->GetValueSize()));
 				}
 			}
+		}
 
-			// sort key values by memory size, so they can be packed better
-			// it still won't protect against structures, that are internally misaligned (-> uint8, uint32)
-			// but since all Engine level keys are good... 
-			InitList.Sort(FBlackboardInitializationData::FMemorySort());
-			uint16 MemoryOffset = 0;
-			for (int32 Index = 0; Index < InitList.Num(); Index++)
-			{
-				ValueOffsets[InitList[Index].KeyID] = MemoryOffset;
-				MemoryOffset += InitList[Index].DataSize;
-			}
+		// sort key values by memory size, so they can be packed better
+		// it still won't protect against structures, that are internally misaligned (-> uint8, uint32)
+		// but since all Engine level keys are good... 
+		InitList.Sort(FBlackboardInitializationData::FMemorySort());
+		uint16 MemoryOffset = 0;
+		for (int32 Index = 0; Index < InitList.Num(); Index++)
+		{
+			ValueOffsets[InitList[Index].KeyID] = MemoryOffset;
+			MemoryOffset += InitList[Index].DataSize;
+		}
 
-			ValueMemory.AddZeroed(MemoryOffset);
+		ValueMemory.AddZeroed(MemoryOffset);
 
-			// initialize memory
-			for (int32 Index = 0; Index < InitList.Num(); Index++)
-			{
-				const FBlackboardEntry* KeyData = BlackboardAsset->GetKey(InitList[Index].KeyID);
-				uint8* RawData = GetKeyRawData(InitList[Index].KeyID);
+		// initialize memory
+		for (int32 Index = 0; Index < InitList.Num(); Index++)
+		{
+			const FBlackboardEntry* KeyData = BlackboardAsset->GetKey(InitList[Index].KeyID);
+			uint8* RawData = GetKeyRawData(InitList[Index].KeyID);
 
-				KeyData->KeyType->Initialize(RawData);
-			}
+			KeyData->KeyType->Initialize(RawData);
 		}
-		else
+	
+		// naive initial synchronization with one of already instantiated blackboards using the same BB asset
+		if (BlackboardAsset->HasSyncronizedKeys())
 		{
-			bSuccess = false;
-			UE_LOG(LogBehaviorTree, Error, TEXT("Blackboard asset (%s) has errors and can't be used!"), *GetNameSafe(BlackboardAsset));
+			PopulateSynchronizedKeys();
 		}
 	}
+	else
+	{
+		bSuccess = false;
+		UE_LOG(LogBehaviorTree, Error, TEXT("Blackboard asset (%s) has errors and can't be used!"), *GetNameSafe(BlackboardAsset));
+	}
 
 	return bSuccess;
 }
 
+void UBlackboardComponent::PopulateSynchronizedKeys()
+{
+	ensure(bSynchronizedKeyPopulated == false);
+	
+	UAISystem* AISystem = UAISystem::GetCurrentSafe(GetWorld());
+	check(AISystem);
+	AISystem->RegisterBlackboardComponent(*BlackboardAsset, *this);
+
+	for (auto Iter = AISystem->CreateBlackboardDataToComponentsIterator(*BlackboardAsset); Iter; ++Iter)
+	{
+		UBlackboardComponent* OtherBlackboard = Iter.Value();
+		if (OtherBlackboard != nullptr && ShouldSyncWithBlackboard(*OtherBlackboard))
+		{
+			for (const auto& Key : BlackboardAsset->Keys)
+			{
+				if (Key.bInstanceSynced)
+				{
+					const int32 KeyID = BlackboardAsset->GetKeyID(Key.EntryName);
+					check(BlackboardAsset->GetKeyID(Key.EntryName) == OtherBlackboard->BlackboardAsset->GetKeyID(Key.EntryName));
+					check(BlackboardAsset->GetKeyType(KeyID) == OtherBlackboard->BlackboardAsset->GetKeyType(KeyID));
+
+					const TSubclassOf<UBlackboardKeyType> ValueType = BlackboardAsset->GetKeyType(KeyID);
+					const UBlackboardKeyType* ValueTypeInstance = GetDefault<UBlackboardKeyType>(*ValueType);
+					ValueTypeInstance->CopyValue(GetKeyRawData(KeyID), OtherBlackboard->GetKeyRawData(KeyID));
+				}
+			}
+			break;
+		}
+	}
+
+	bSynchronizedKeyPopulated = true;
+}
+
+bool UBlackboardComponent::ShouldSyncWithBlackboard(UBlackboardComponent& OtherBlackboardComponent) const
+{
+	return &OtherBlackboardComponent != this && (
+		(BrainComp == nullptr || (BrainComp->GetAIOwner() != nullptr && BrainComp->GetAIOwner()->ShouldSyncBlackboardWith(OtherBlackboardComponent) == true))
+		|| (OtherBlackboardComponent.BrainComp == nullptr || (OtherBlackboardComponent.BrainComp->GetAIOwner() != nullptr && OtherBlackboardComponent.BrainComp->GetAIOwner()->ShouldSyncBlackboardWith(OtherBlackboardComponent) == true)));
+}
+
 UBrainComponent* UBlackboardComponent::GetBrainComponent() const
 {
 	return BrainComp;
@@ -162,6 +234,11 @@ TSubclassOf<UBlackboardKeyType> UBlackboardComponent::GetKeyType(FBlackboard::FK
 	return BlackboardAsset ? BlackboardAsset->GetKeyType(KeyID) : NULL;
 }
 
+bool UBlackboardComponent::IsKeyInstanceSynced(FBlackboard::FKey KeyID) const
+{
+	return BlackboardAsset ? BlackboardAsset->IsKeyInstanceSynced(KeyID) : false;
+}
+
 int32 UBlackboardComponent::GetNumKeys() const
 {
 	return BlackboardAsset ? BlackboardAsset->GetNumKeys() : 0;
@@ -518,267 +595,109 @@ FRotator UBlackboardComponent::GetValueAsRotator(uint8 KeyID) const
 void UBlackboardComponent::SetValueAsObject(const FName& KeyName, UObject* ObjectValue)
 {
 	const FBlackboard::FKey KeyID = GetKeyID(KeyName);
-	SetValueAsObject(KeyID, ObjectValue);
-}
-
-void UBlackboardComponent::SetValueAsObject(FBlackboard::FKey KeyID, UObject* ObjectValue)
-{
-	if ( GetKeyType(KeyID) != UBlackboardKeyType_Object::StaticClass() )
-	{
-		return;
-	}
-
-	uint8* RawData = GetKeyRawData(KeyID);
-	if (RawData)
-	{
-		const bool bChanged = UBlackboardKeyType_Object::SetValue(RawData, ObjectValue);
-		if (bChanged)
-		{
-			NotifyObservers(KeyID);
-		}
-	}
-}
+	SetValue<UBlackboardKeyType_Object>(KeyID, ObjectValue);
+}
+
+//void UBlackboardComponent::SetValueAsObject(FBlackboard::FKey KeyID, UObject* ObjectValue)
+//{
+//	if (GetKeyType(KeyID) != UBlackboardKeyType_Object::StaticClass())
+//	{
+//		return;
+//	}
+//
+//	uint8* RawData = GetKeyRawData(KeyID);
+//	if (RawData)
+//	{
+//		const bool bChanged = UBlackboardKeyType_Object::SetValue(RawData, ObjectValue);
+//		if (bChanged)
+//		{
+//			NotifyObservers(KeyID);
+//			if (bHasSyncedKeys && IsKeyInstanceSynced(KeyID))
+//			{
+//				SynchronizeKeys(KeyID, ObjectValue);
+//			}
+//		}
+//	}
+//}
+//
+////void UBlackboardData::SynchronizeKeys(FBlackboard::FKey KeyID, UObject* ObjectValue)
+////{
+////	check(BlackboardAsset);
+////
+////	if (BrainComp && BrainComp->GetAIOwner())
+////	{
+////		UAISystem* AISystem = UAISystem::GetCurrentSafe(GetWorld());
+////		for (auto Iter = AISystem->CreateBlackboardDataToComponentsIterator(*BlackboardAsset); Iter; ++Iter)
+////		{
+////			UBlackboardComponent* OtherBlackboard = Iter.Value();
+////			if (OtherBlackboard && OtherBlackboard != this && OtherBlackboard->BrainComp && OtherBlackboard->BrainComp->GetAIOwner()
+////				&& OtherBlackboard->BrainComp->GetAIOwner()->ShouldInstanceSyncUpdatedBlackboardField(*OtherBlackboard, KeyID, *this, *BrainComp->GetAIOwner()))
+////			{
+////				UBlackboardKeyType_Object::SetValue(OtherBlackboard->GetKeyRawData(KeyID), ObjectValue);
+////				OtherBlackboard->NotifyObservers(KeyID);
+////			}
+////		}
+////	}
+////}
 
 void UBlackboardComponent::SetValueAsClass(const FName& KeyName, UClass* ClassValue)
 {
 	const FBlackboard::FKey KeyID = GetKeyID(KeyName);
-	SetValueAsClass(KeyID, ClassValue);
-}
-
-void UBlackboardComponent::SetValueAsClass(FBlackboard::FKey KeyID, UClass* ClassValue)
-{
-	if (GetKeyType(KeyID) != UBlackboardKeyType_Class::StaticClass())
-	{
-		return;
-	}
-
-	uint8* RawData = GetKeyRawData(KeyID);
-	if (RawData)
-	{
-		const bool bChanged = UBlackboardKeyType_Class::SetValue(RawData, ClassValue);
-		if (bChanged)
-		{
-			NotifyObservers(KeyID);
-		}
-	}
+	SetValue<UBlackboardKeyType_Class>(KeyID, ClassValue);
 }
 
 void UBlackboardComponent::SetValueAsEnum(const FName& KeyName, uint8 EnumValue)
 {
 	const FBlackboard::FKey KeyID = GetKeyID(KeyName);
-	SetValueAsEnum(KeyID, EnumValue);
-}
-
-void UBlackboardComponent::SetValueAsEnum(FBlackboard::FKey KeyID, uint8 EnumValue)
-{
-	if (GetKeyType(KeyID) != UBlackboardKeyType_Enum::StaticClass() &&
-		GetKeyType(KeyID) != UBlackboardKeyType_NativeEnum::StaticClass())
-	{
-		return;
-	}
-
-	uint8* RawData = GetKeyRawData(KeyID);
-	if (RawData)
-	{
-		const bool bChanged = UBlackboardKeyType_Enum::SetValue(RawData, EnumValue);
-		if (bChanged)
-		{
-			NotifyObservers(KeyID);
-		}
-	}
+	SetValue<UBlackboardKeyType_Enum>(KeyID, EnumValue);
 }
 
 void UBlackboardComponent::SetValueAsInt(const FName& KeyName, int32 IntValue)
 {
 	const FBlackboard::FKey KeyID = GetKeyID(KeyName);
-	SetValueAsInt(KeyID, IntValue);
-}
-
-void UBlackboardComponent::SetValueAsInt(FBlackboard::FKey KeyID, int32 IntValue)
-{
-	if (GetKeyType(KeyID) != UBlackboardKeyType_Int::StaticClass())
-	{
-		return;
-	}
-
-	uint8* RawData = GetKeyRawData(KeyID);
-	if (RawData)
-	{
-		const bool bChanged = UBlackboardKeyType_Int::SetValue(RawData, IntValue);
-		if (bChanged)
-		{
-			NotifyObservers(KeyID);
-		}
-	}
+	SetValue<UBlackboardKeyType_Int>(KeyID, IntValue);
 }
 
 void UBlackboardComponent::SetValueAsFloat(const FName& KeyName, float FloatValue)
 {
 	const FBlackboard::FKey KeyID = GetKeyID(KeyName);
-	SetValueAsFloat(KeyID, FloatValue);
-}
-
-void UBlackboardComponent::SetValueAsFloat(FBlackboard::FKey KeyID, float FloatValue)
-{
-	if (GetKeyType(KeyID) != UBlackboardKeyType_Float::StaticClass())
-	{
-		return;
-	}
-
-	uint8* RawData = GetKeyRawData(KeyID);
-	if (RawData)
-	{
-		const bool bChanged = UBlackboardKeyType_Float::SetValue(RawData, FloatValue);
-		if (bChanged)
-		{
-			NotifyObservers(KeyID);
-		}
-	}
+	SetValue<UBlackboardKeyType_Float>(KeyID, FloatValue);
 }
 
 void UBlackboardComponent::SetValueAsBool(const FName& KeyName, bool BoolValue)
 {
 	const FBlackboard::FKey KeyID = GetKeyID(KeyName);
-	SetValueAsBool(KeyID, BoolValue);
-}
-
-void UBlackboardComponent::SetValueAsBool(FBlackboard::FKey KeyID, bool BoolValue)
-{
-	if (GetKeyType(KeyID) != UBlackboardKeyType_Bool::StaticClass())
-	{
-		return;
-	}
-
-	uint8* RawData = GetKeyRawData(KeyID);
-	if (RawData)
-	{
-		const bool bChanged = UBlackboardKeyType_Bool::SetValue(RawData, BoolValue);
-		if (bChanged)
-		{
-			NotifyObservers(KeyID);
-		}
-	}
+	SetValue<UBlackboardKeyType_Bool>(KeyID, BoolValue);
 }
 
 void UBlackboardComponent::SetValueAsString(const FName& KeyName, FString StringValue)
 {
 	const FBlackboard::FKey KeyID = GetKeyID(KeyName);
-	SetValueAsString(KeyID, StringValue);
-}
-
-void UBlackboardComponent::SetValueAsString(FBlackboard::FKey KeyID, const FString& StringValue)
-{
-	if (GetKeyType(KeyID) != UBlackboardKeyType_String::StaticClass())
-	{
-		return;
-	}
-
-	uint8* RawData = GetKeyRawData(KeyID);
-	if (RawData)
-	{
-		const bool bChanged = UBlackboardKeyType_String::SetValue(RawData, StringValue);
-		if (bChanged)
-		{
-			NotifyObservers(KeyID);
-		}
-	}
+	SetValue<UBlackboardKeyType_String>(KeyID, StringValue);
 }
 
 void UBlackboardComponent::SetValueAsName(const FName& KeyName, FName NameValue)
 {
 	const FBlackboard::FKey KeyID = GetKeyID(KeyName);
-	SetValueAsName(KeyID, NameValue);
-}
-
-void UBlackboardComponent::SetValueAsName(FBlackboard::FKey KeyID, const FName& NameValue)
-{
-	if (GetKeyType(KeyID) != UBlackboardKeyType_Name::StaticClass())
-	{
-		return;
-	}
-
-	uint8* RawData = GetKeyRawData(KeyID);
-	if (RawData)
-	{
-		const bool bChanged = UBlackboardKeyType_Name::SetValue(RawData, NameValue);
-		if (bChanged)
-		{
-			NotifyObservers(KeyID);
-		}
-	}
+	SetValue<UBlackboardKeyType_Name>(KeyID, NameValue);
 }
 
 void UBlackboardComponent::SetValueAsVector(const FName& KeyName, FVector VectorValue)
 {
 	const FBlackboard::FKey KeyID = GetKeyID(KeyName);
-	SetValueAsVector(KeyID, VectorValue);
+	SetValue<UBlackboardKeyType_Vector>(KeyID, VectorValue);
 }
 
-void UBlackboardComponent::SetValueAsVector(FBlackboard::FKey KeyID, const FVector& VectorValue)
-{
-	if (GetKeyType(KeyID) != UBlackboardKeyType_Vector::StaticClass())
-	{
-		return;
-	}
-
-	uint8* RawData = GetKeyRawData(KeyID);
-	if (RawData)
-	{
-		const bool bChanged = UBlackboardKeyType_Vector::SetValue(RawData, VectorValue);
-		if (bChanged)
-		{
-			NotifyObservers(KeyID);
-		}
-	}
-}
-
-
 void UBlackboardComponent::SetValueAsRotator(const FName& KeyName, FRotator RotatorValue)
 {
 	const FBlackboard::FKey KeyID = GetKeyID(KeyName);
-	SetValueAsRotator(KeyID, RotatorValue);
-}
-
-void UBlackboardComponent::SetValueAsRotator(uint8 KeyID, const FRotator& RotatorValue)
-{
-	if (GetKeyType(KeyID) != UBlackboardKeyType_Rotator::StaticClass())
-	{
-		return;
-	}
-
-	uint8* RawData = GetKeyRawData(KeyID);
-	if (RawData)
-	{
-		const bool bChanged = UBlackboardKeyType_Rotator::SetValue(RawData, RotatorValue);
-		if (bChanged)
-		{
-			NotifyObservers(KeyID);
-		}
-	}
+	SetValue<UBlackboardKeyType_Rotator>(KeyID, RotatorValue);
 }
 
 void UBlackboardComponent::ClearValueAsVector(const FName& KeyName)
 {
 	const FBlackboard::FKey KeyID = GetKeyID(KeyName);
-	ClearValueAsVector(KeyID);
-}
-
-void UBlackboardComponent::ClearValueAsVector(FBlackboard::FKey KeyID)
-{
-	if (GetKeyType(KeyID) != UBlackboardKeyType_Vector::StaticClass())
-	{
-		return;
-	}
-
-	uint8* RawData = GetKeyRawData(KeyID);
-	if (RawData)
-	{
-		const bool bChanged = UBlackboardKeyType_Vector::SetValue(RawData, FAISystem::InvalidLocation);
-		if (bChanged)
-		{
-			NotifyObservers(KeyID);
-		}
-	}
+	ClearValue(KeyID);
 }
 
 bool UBlackboardComponent::IsVectorValueSet(const FName& KeyName) const
@@ -830,17 +749,35 @@ void UBlackboardComponent::ClearValue(const FName& KeyName)
 
 void UBlackboardComponent::ClearValue(FBlackboard::FKey KeyID)
 {
+	check(BlackboardAsset);
+
 	uint8* RawData = GetKeyRawData(KeyID);
 	if (RawData)
 	{
 		TSubclassOf<UBlackboardKeyType> ValueType = GetKeyType(KeyID);
-		if (ValueType)
+		check(ValueType);
+		
+		const UBlackboardKeyType* ValueTypeInstance = GetDefault<UBlackboardKeyType>(*ValueType);
+
+		const bool bChanged = ValueTypeInstance->Clear(RawData);
+		if (bChanged)
 		{
-			// c-cast is ok here since ValueType is of at least UBlackboardKeyType and it's default object has to be UBlackboardKeyType instance
-			const bool bChanged = GetDefault<UBlackboardKeyType>()->Clear(RawData);
-			if (bChanged)
+			NotifyObservers(KeyID);
+
+			if (BlackboardAsset->HasSyncronizedKeys() && IsKeyInstanceSynced(KeyID))
 			{
-				NotifyObservers(KeyID);
+				// grab the value set and apply the same to synchronized keys
+				// to avoid virtual function call overhead
+				UAISystem* AISystem = UAISystem::GetCurrentSafe(GetWorld());
+				for (auto Iter = AISystem->CreateBlackboardDataToComponentsIterator(*BlackboardAsset); Iter; ++Iter)
+				{
+					UBlackboardComponent* OtherBlackboard = Iter.Value();
+					if (OtherBlackboard != nullptr && ShouldSyncWithBlackboard(*OtherBlackboard))
+					{
+						ValueTypeInstance->CopyValue(OtherBlackboard->GetKeyRawData(KeyID), RawData);
+						OtherBlackboard->NotifyObservers(KeyID);
+					}
+				}
 			}
 		}
 	}
diff --git a/Engine/Source/Runtime/AIModule/Private/BehaviorTree/BlackboardData.cpp b/Engine/Source/Runtime/AIModule/Private/BehaviorTree/BlackboardData.cpp
--- a/Engine/Source/Runtime/AIModule/Private/BehaviorTree/BlackboardData.cpp
+++ b/Engine/Source/Runtime/AIModule/Private/BehaviorTree/BlackboardData.cpp
@@ -46,6 +46,12 @@ TSubclassOf<UBlackboardKeyType> UBlackboardData::GetKeyType(FBlackboard::FKey Ke
 	return KeyEntry ? KeyEntry->KeyType->GetClass() : NULL;
 }
 
+bool UBlackboardData::IsKeyInstanceSynced(FBlackboard::FKey KeyID) const
+{
+	const FBlackboardEntry* KeyEntry = GetKey(KeyID);
+	return KeyEntry ? KeyEntry->bInstanceSynced : false;
+}
+
 const FBlackboardEntry* UBlackboardData::GetKey(FBlackboard::FKey KeyID) const
 {
 	if (KeyID != FBlackboard::InvalidKey)
@@ -101,31 +107,56 @@ bool UBlackboardData::IsValid() const
 	return true;
 }
 
+void UBlackboardData::UpdateIfHasSynchronizedKeys()
+{
+	bHasSyncronizedKeys = Parent != nullptr && Parent->bHasSyncronizedKeys;
+	for (int32 KeyIndex = 0; KeyIndex < Keys.Num() && bHasSyncronizedKeys == false; ++KeyIndex)
+	{
+		bHasSyncronizedKeys |= Keys[KeyIndex].bInstanceSynced;
+	}
+}
+
 void UBlackboardData::PostLoad()
 {
 	Super::PostLoad();
 
 	UpdateParentKeys();
+	UpdateIfHasSynchronizedKeys();
 }
 
 #if WITH_EDITOR
+void UBlackboardData::PreEditChange(class FEditPropertyChain& PropertyAboutToChange)
+{
+	Super::PreEditChange(PropertyAboutToChange);
+}
 
 void UBlackboardData::PostEditChangeProperty(FPropertyChangedEvent& PropertyChangedEvent)
 {
+	static const FName NAME_Parent = GET_MEMBER_NAME_CHECKED(UBlackboardData, Parent);
+	static const FName NAME_InstanceSynced = GET_MEMBER_NAME_CHECKED(FBlackboardEntry, bInstanceSynced);
+
 	Super::PostEditChangeProperty(PropertyChangedEvent);
 
-	if (PropertyChangedEvent.Property && PropertyChangedEvent.Property->GetFName() == GET_MEMBER_NAME_CHECKED(UBlackboardData, Parent))
+	if (PropertyChangedEvent.Property)
 	{
-		// look for cycles
-		if (Parent && Parent->HasParent(this))
+		if (PropertyChangedEvent.Property->GetFName() == NAME_Parent)
 		{
-			UE_LOG(LogBehaviorTree, Warning, TEXT("Blackboard asset (%s) has (%s) in parent chain! Clearing value to avoid cycle."),
-				*GetNameSafe(Parent), *GetNameSafe(this));
+			// look for cycles
+			if (Parent && Parent->HasParent(this))
+			{
+				UE_LOG(LogBehaviorTree, Warning, TEXT("Blackboard asset (%s) has (%s) in parent chain! Clearing value to avoid cycle."),
+					*GetNameSafe(Parent), *GetNameSafe(this));
 
-			Parent = NULL;
+				Parent = NULL;
+			}
+
+			UpdateParentKeys();
 		}
 
-		UpdateParentKeys();
+		if (PropertyChangedEvent.Property->GetFName() == NAME_InstanceSynced || PropertyChangedEvent.Property->GetFName() == NAME_Parent)
+		{
+			UpdateIfHasSynchronizedKeys();
+		}
 	}
 }
 
diff --git a/Engine/Source/Runtime/AIModule/Private/EnvironmentQuery/Items/EnvQueryItemType_ActorBase.cpp b/Engine/Source/Runtime/AIModule/Private/EnvironmentQuery/Items/EnvQueryItemType_ActorBase.cpp
--- a/Engine/Source/Runtime/AIModule/Private/EnvironmentQuery/Items/EnvQueryItemType_ActorBase.cpp
+++ b/Engine/Source/Runtime/AIModule/Private/EnvironmentQuery/Items/EnvQueryItemType_ActorBase.cpp
@@ -21,7 +21,7 @@ bool UEnvQueryItemType_ActorBase::StoreInBlackboard(FBlackboardKeySelector& KeyS
 	if (!bStored && KeySelector.SelectedKeyType == UBlackboardKeyType_Object::StaticClass())
 	{
 		UObject* MyObject = GetActor(RawData);
-		Blackboard->SetValueAsObject(KeySelector.GetSelectedKeyID(), MyObject);
+		Blackboard->SetValue<UBlackboardKeyType_Object>(KeySelector.GetSelectedKeyID(), MyObject);
 
 		bStored = true;
 	}
diff --git a/Engine/Source/Runtime/AIModule/Private/EnvironmentQuery/Items/EnvQueryItemType_VectorBase.cpp b/Engine/Source/Runtime/AIModule/Private/EnvironmentQuery/Items/EnvQueryItemType_VectorBase.cpp
--- a/Engine/Source/Runtime/AIModule/Private/EnvironmentQuery/Items/EnvQueryItemType_VectorBase.cpp
+++ b/Engine/Source/Runtime/AIModule/Private/EnvironmentQuery/Items/EnvQueryItemType_VectorBase.cpp
@@ -20,7 +20,7 @@ bool UEnvQueryItemType_VectorBase::StoreInBlackboard(FBlackboardKeySelector& Key
 	if (KeySelector.SelectedKeyType == UBlackboardKeyType_Vector::StaticClass())
 	{
 		const FVector MyLocation = GetItemLocation(RawData);
-		Blackboard->SetValueAsVector(KeySelector.GetSelectedKeyID(), MyLocation);
+		Blackboard->SetValue<UBlackboardKeyType_Vector>(KeySelector.GetSelectedKeyID(), MyLocation);
 
 		bStored = true;
 	}
