Compilation fixes after blackboard key-syncing related changes #UE4

diff --git a/Engine/Source/Developer/AITestSuite/Private/BehaviorTree/TestBTTask_SetFlag.cpp b/Engine/Source/Developer/AITestSuite/Private/BehaviorTree/TestBTTask_SetFlag.cpp
--- a/Engine/Source/Developer/AITestSuite/Private/BehaviorTree/TestBTTask_SetFlag.cpp
+++ b/Engine/Source/Developer/AITestSuite/Private/BehaviorTree/TestBTTask_SetFlag.cpp
@@ -14,6 +14,6 @@ UTestBTTask_SetFlag::UTestBTTask_SetFlag(const FObjectInitializer& ObjectInitial
 
 EBTNodeResult::Type UTestBTTask_SetFlag::ExecuteTask(UBehaviorTreeComponent& OwnerComp, uint8* NodeMemory)
 {
-	OwnerComp.GetBlackboardComponent()->SetValueAsBool(KeyName, bValue);
+	OwnerComp.GetBlackboardComponent()->SetValue<UBlackboardKeyType_Bool>(KeyName, bValue);
 	return TaskResult;
 }
diff --git a/Engine/Source/Runtime/AIModule/Classes/BehaviorTree/Blackboard/BlackboardKeyType_Bool.h b/Engine/Source/Runtime/AIModule/Classes/BehaviorTree/Blackboard/BlackboardKeyType_Bool.h
--- a/Engine/Source/Runtime/AIModule/Classes/BehaviorTree/Blackboard/BlackboardKeyType_Bool.h
+++ b/Engine/Source/Runtime/AIModule/Classes/BehaviorTree/Blackboard/BlackboardKeyType_Bool.h
@@ -9,6 +9,7 @@ class AIMODULE_API UBlackboardKeyType_Bool : public UBlackboardKeyType
 	GENERATED_UCLASS_BODY()
 
 	typedef bool FDataType;
+	static const FDataType InvalidValue;
 
 	static bool GetValue(const uint8* RawData);
 	static bool SetValue(uint8* RawData, bool bValue);
diff --git a/Engine/Source/Runtime/AIModule/Classes/BehaviorTree/Blackboard/BlackboardKeyType_Class.h b/Engine/Source/Runtime/AIModule/Classes/BehaviorTree/Blackboard/BlackboardKeyType_Class.h
--- a/Engine/Source/Runtime/AIModule/Classes/BehaviorTree/Blackboard/BlackboardKeyType_Class.h
+++ b/Engine/Source/Runtime/AIModule/Classes/BehaviorTree/Blackboard/BlackboardKeyType_Class.h
@@ -9,6 +9,7 @@ class AIMODULE_API UBlackboardKeyType_Class : public UBlackboardKeyType
 	GENERATED_UCLASS_BODY()
 
 	typedef UClass* FDataType;
+	static const FDataType InvalidValue;
 
 	UPROPERTY(Category=Blackboard, EditDefaultsOnly, meta=(AllowAbstract="1"))
 	UClass* BaseClass;
diff --git a/Engine/Source/Runtime/AIModule/Classes/BehaviorTree/Blackboard/BlackboardKeyType_Enum.h b/Engine/Source/Runtime/AIModule/Classes/BehaviorTree/Blackboard/BlackboardKeyType_Enum.h
--- a/Engine/Source/Runtime/AIModule/Classes/BehaviorTree/Blackboard/BlackboardKeyType_Enum.h
+++ b/Engine/Source/Runtime/AIModule/Classes/BehaviorTree/Blackboard/BlackboardKeyType_Enum.h
@@ -9,6 +9,7 @@ class AIMODULE_API UBlackboardKeyType_Enum : public UBlackboardKeyType
 	GENERATED_UCLASS_BODY()
 
 	typedef uint8 FDataType;
+	static const FDataType InvalidValue;
 
 	UPROPERTY(Category=Blackboard, EditDefaultsOnly)
 	UEnum* EnumType;
diff --git a/Engine/Source/Runtime/AIModule/Classes/BehaviorTree/Blackboard/BlackboardKeyType_Float.h b/Engine/Source/Runtime/AIModule/Classes/BehaviorTree/Blackboard/BlackboardKeyType_Float.h
--- a/Engine/Source/Runtime/AIModule/Classes/BehaviorTree/Blackboard/BlackboardKeyType_Float.h
+++ b/Engine/Source/Runtime/AIModule/Classes/BehaviorTree/Blackboard/BlackboardKeyType_Float.h
@@ -9,6 +9,7 @@ class AIMODULE_API UBlackboardKeyType_Float : public UBlackboardKeyType
 	GENERATED_UCLASS_BODY()
 
 	typedef float FDataType;
+	static const FDataType InvalidValue;
 
 	static float GetValue(const uint8* RawData);
 	static bool SetValue(uint8* RawData, float Value);
diff --git a/Engine/Source/Runtime/AIModule/Classes/BehaviorTree/Blackboard/BlackboardKeyType_Int.h b/Engine/Source/Runtime/AIModule/Classes/BehaviorTree/Blackboard/BlackboardKeyType_Int.h
--- a/Engine/Source/Runtime/AIModule/Classes/BehaviorTree/Blackboard/BlackboardKeyType_Int.h
+++ b/Engine/Source/Runtime/AIModule/Classes/BehaviorTree/Blackboard/BlackboardKeyType_Int.h
@@ -9,6 +9,7 @@ class AIMODULE_API UBlackboardKeyType_Int : public UBlackboardKeyType
 	GENERATED_UCLASS_BODY()
 
 	typedef int32 FDataType;
+	static const FDataType InvalidValue;
 
 	static int32 GetValue(const uint8* RawData);
 	static bool SetValue(uint8* RawData, int32 Value);
diff --git a/Engine/Source/Runtime/AIModule/Classes/BehaviorTree/Blackboard/BlackboardKeyType_Name.h b/Engine/Source/Runtime/AIModule/Classes/BehaviorTree/Blackboard/BlackboardKeyType_Name.h
--- a/Engine/Source/Runtime/AIModule/Classes/BehaviorTree/Blackboard/BlackboardKeyType_Name.h
+++ b/Engine/Source/Runtime/AIModule/Classes/BehaviorTree/Blackboard/BlackboardKeyType_Name.h
@@ -9,6 +9,7 @@ class AIMODULE_API UBlackboardKeyType_Name : public UBlackboardKeyType
 	GENERATED_UCLASS_BODY()
 
 	typedef FName FDataType;
+	static const FDataType InvalidValue;
 
 	static FName GetValue(const uint8* RawData);
 	static bool SetValue(uint8* RawData, const FName& Value);
diff --git a/Engine/Source/Runtime/AIModule/Classes/BehaviorTree/Blackboard/BlackboardKeyType_NativeEnum.h b/Engine/Source/Runtime/AIModule/Classes/BehaviorTree/Blackboard/BlackboardKeyType_NativeEnum.h
--- a/Engine/Source/Runtime/AIModule/Classes/BehaviorTree/Blackboard/BlackboardKeyType_NativeEnum.h
+++ b/Engine/Source/Runtime/AIModule/Classes/BehaviorTree/Blackboard/BlackboardKeyType_NativeEnum.h
@@ -10,6 +10,7 @@ class AIMODULE_API UBlackboardKeyType_NativeEnum : public UBlackboardKeyType
 	GENERATED_UCLASS_BODY()
 
 	typedef uint8 FDataType;
+	static const FDataType InvalidValue;
 
 	UPROPERTY(Category=Blackboard, EditDefaultsOnly)
 	FString EnumName;
diff --git a/Engine/Source/Runtime/AIModule/Classes/BehaviorTree/Blackboard/BlackboardKeyType_Object.h b/Engine/Source/Runtime/AIModule/Classes/BehaviorTree/Blackboard/BlackboardKeyType_Object.h
--- a/Engine/Source/Runtime/AIModule/Classes/BehaviorTree/Blackboard/BlackboardKeyType_Object.h
+++ b/Engine/Source/Runtime/AIModule/Classes/BehaviorTree/Blackboard/BlackboardKeyType_Object.h
@@ -9,6 +9,7 @@ class AIMODULE_API UBlackboardKeyType_Object : public UBlackboardKeyType
 	GENERATED_UCLASS_BODY()
 
 	typedef UObject* FDataType;
+	static const FDataType InvalidValue;
 
 	UPROPERTY(Category=Blackboard, EditDefaultsOnly, meta=(AllowAbstract="1"))
 	UClass* BaseClass;
diff --git a/Engine/Source/Runtime/AIModule/Classes/BehaviorTree/Blackboard/BlackboardKeyType_Rotator.h b/Engine/Source/Runtime/AIModule/Classes/BehaviorTree/Blackboard/BlackboardKeyType_Rotator.h
--- a/Engine/Source/Runtime/AIModule/Classes/BehaviorTree/Blackboard/BlackboardKeyType_Rotator.h
+++ b/Engine/Source/Runtime/AIModule/Classes/BehaviorTree/Blackboard/BlackboardKeyType_Rotator.h
@@ -8,7 +8,8 @@ class AIMODULE_API UBlackboardKeyType_Rotator : public UBlackboardKeyType
 	GENERATED_UCLASS_BODY()
 	
 	typedef FRotator FDataType; 
-	
+	static const FDataType InvalidValue;
+
 	static FRotator GetValue(const uint8* RawData);
 	static bool SetValue(uint8* RawData, const FRotator& Value);
 
diff --git a/Engine/Source/Runtime/AIModule/Classes/BehaviorTree/Blackboard/BlackboardKeyType_String.h b/Engine/Source/Runtime/AIModule/Classes/BehaviorTree/Blackboard/BlackboardKeyType_String.h
--- a/Engine/Source/Runtime/AIModule/Classes/BehaviorTree/Blackboard/BlackboardKeyType_String.h
+++ b/Engine/Source/Runtime/AIModule/Classes/BehaviorTree/Blackboard/BlackboardKeyType_String.h
@@ -9,6 +9,7 @@ class AIMODULE_API UBlackboardKeyType_String : public UBlackboardKeyType
 	GENERATED_UCLASS_BODY()
 
 	typedef FString FDataType;
+	static const FDataType InvalidValue;
 
 	static FString GetValue(const uint8* RawData);
 	static bool SetValue(uint8* RawData, const FString& Value);
diff --git a/Engine/Source/Runtime/AIModule/Classes/BehaviorTree/Blackboard/BlackboardKeyType_Vector.h b/Engine/Source/Runtime/AIModule/Classes/BehaviorTree/Blackboard/BlackboardKeyType_Vector.h
--- a/Engine/Source/Runtime/AIModule/Classes/BehaviorTree/Blackboard/BlackboardKeyType_Vector.h
+++ b/Engine/Source/Runtime/AIModule/Classes/BehaviorTree/Blackboard/BlackboardKeyType_Vector.h
@@ -9,6 +9,7 @@ class AIMODULE_API UBlackboardKeyType_Vector : public UBlackboardKeyType
 	GENERATED_UCLASS_BODY()
 
 	typedef FVector FDataType; 
+	static const FDataType InvalidValue;
 	
 	static FVector GetValue(const uint8* RawData);
 	static bool SetValue(uint8* RawData, const FVector& Value);
diff --git a/Engine/Source/Runtime/AIModule/Classes/BehaviorTree/BlackboardComponent.h b/Engine/Source/Runtime/AIModule/Classes/BehaviorTree/BlackboardComponent.h
--- a/Engine/Source/Runtime/AIModule/Classes/BehaviorTree/BlackboardComponent.h
+++ b/Engine/Source/Runtime/AIModule/Classes/BehaviorTree/BlackboardComponent.h
@@ -15,6 +15,7 @@
 #include "BehaviorTreeTypes.h"
 #include "BehaviorTree/Blackboard/BlackboardKeyType.h"
 #include "BlackboardData.h"
+#include "AISystem.h"
 #include "BlackboardComponent.generated.h"
 
 class UBlackboardData;
@@ -92,47 +93,37 @@ public:
 
 	UFUNCTION(BlueprintCallable, Category="AI|Components|Blackboard")
 	UObject* GetValueAsObject(const FName& KeyName) const;
-	UObject* GetValueAsObject(FBlackboard::FKey KeyID) const;
 
 	UFUNCTION(BlueprintCallable, Category="AI|Components|Blackboard")
 	UClass* GetValueAsClass(const FName& KeyName) const;
-	UClass* GetValueAsClass(FBlackboard::FKey KeyID) const;
 
 	UFUNCTION(BlueprintCallable, Category="AI|Components|Blackboard")
 	uint8 GetValueAsEnum(const FName& KeyName) const;
-	uint8 GetValueAsEnum(FBlackboard::FKey KeyID) const;
 
 	UFUNCTION(BlueprintCallable, Category="AI|Components|Blackboard")
 	int32 GetValueAsInt(const FName& KeyName) const;
-	int32 GetValueAsInt(FBlackboard::FKey KeyID) const;
 
 	UFUNCTION(BlueprintCallable, Category="AI|Components|Blackboard")
 	float GetValueAsFloat(const FName& KeyName) const;
-	float GetValueAsFloat(FBlackboard::FKey KeyID) const;
 
 	UFUNCTION(BlueprintCallable, Category="AI|Components|Blackboard")
 	bool GetValueAsBool(const FName& KeyName) const;
-	bool GetValueAsBool(FBlackboard::FKey KeyID) const;
 
 	UFUNCTION(BlueprintCallable, Category="AI|Components|Blackboard")
 	FString GetValueAsString(const FName& KeyName) const;
-	FString GetValueAsString(FBlackboard::FKey KeyID) const;
 	
 	UFUNCTION(BlueprintCallable, Category="AI|Components|Blackboard")
 	FName GetValueAsName(const FName& KeyName) const;
-	FName GetValueAsName(FBlackboard::FKey KeyID) const;
 
 	UFUNCTION(BlueprintCallable, Category="AI|Components|Blackboard")
 	FVector GetValueAsVector(const FName& KeyName) const;
-	FVector GetValueAsVector(FBlackboard::FKey KeyID) const;
 
 	UFUNCTION(BlueprintCallable, Category="AI|Components|Blackboard")
 	FRotator GetValueAsRotator(const FName& KeyName) const;
-	FRotator GetValueAsRotator(FBlackboard::FKey KeyID) const;
 
 	UFUNCTION(BlueprintCallable, Category="AI|Components|Blackboard")
 	void SetValueAsObject(const FName& KeyName, UObject* ObjectValue);
-
+	
 	UFUNCTION(BlueprintCallable, Category="AI|Components|Blackboard")
 	void SetValueAsClass(const FName& KeyName, UClass* ClassValue);
 
@@ -157,23 +148,14 @@ public:
 	UFUNCTION(BlueprintCallable, Category="AI|Components|Blackboard")
 	void SetValueAsVector(const FName& KeyName, FVector VectorValue);
 
-	UFUNCTION(BlueprintCallable, Category="AI|Components|Blackboard")
-	void ClearValueAsVector(const FName& KeyName);
-	void ClearValueAsVector(FBlackboard::FKey KeyID);
+	UFUNCTION(BlueprintCallable, Category = "AI|Components|Blackboard")
+	void SetValueAsRotator(const FName& KeyName, FRotator VectorValue);
 
 	UFUNCTION(BlueprintCallable, Category="AI|Components|Blackboard", Meta=(
 		Tooltip="If the vector value has been set (and not cleared), this function returns true (indicating that the value should be valid).  If it's not set, the vector value is invalid and this function will return false.  (Also returns false if the key specified does not hold a vector.)"))
 	bool IsVectorValueSet(const FName& KeyName) const;
 	bool IsVectorValueSet(FBlackboard::FKey KeyID) const;
-
-	UFUNCTION(BlueprintCallable, Category="AI|Components|Blackboard")
-	void SetValueAsRotator(const FName& KeyName, FRotator VectorValue);
-	void SetValueAsRotator(FBlackboard::FKey KeyID, const FRotator& VectorValue);
-
-	UFUNCTION(BlueprintCallable, Category="AI|Components|Blackboard")
-	void ClearValueAsRotator(const FName& KeyName);
-	void ClearValueAsRotator(FBlackboard::FKey KeyID);
-
+		
 	/** return false if call failed (most probably no such entry in BB) */
 	UFUNCTION(BlueprintCallable, Category="AI|Components|Blackboard")
 	bool GetLocationFromEntry(const FName& KeyName, FVector& ResultLocation) const;
@@ -188,9 +170,18 @@ public:
 	void ClearValue(const FName& KeyName);
 	void ClearValue(FBlackboard::FKey KeyID);
 
+	template<class TDataClass>
+	bool SetValue(const FName& KeyName, typename TDataClass::FDataType Value);
+
 	template<class TDataClass>
 	bool SetValue(FBlackboard::FKey KeyID, typename TDataClass::FDataType Value);
 
+	template<class TDataClass>
+	typename TDataClass::FDataType GetValue(const FName& KeyName) const;
+
+	template<class TDataClass>
+	typename TDataClass::FDataType GetValue(FBlackboard::FKey KeyID) const;
+
 	/** get pointer to raw data for given key */
 	FORCEINLINE uint8* GetKeyRawData(const FName& KeyName) { return GetKeyRawData(GetKeyID(KeyName)); }
 	FORCEINLINE uint8* GetKeyRawData(FBlackboard::FKey KeyID) { return ValueMemory.Num() && ValueOffsets.IsValidIndex(KeyID) ? (ValueMemory.GetData() + ValueOffsets[KeyID]) : NULL; }
@@ -213,7 +204,57 @@ public:
 	/** prepare blackboard snapshot for logs */
 	virtual void DescribeSelfToVisLog(struct FVisualLogEntry* Snapshot) const;
 #endif
-	
+
+	//----------------------------------------------------------------------//
+	// DEPRECATED
+	//----------------------------------------------------------------------//
+	DEPRECATED(4.7, "This function is deprecated. Please use ClearValue() instead.")
+	void ClearValueAsRotator(FBlackboard::FKey KeyID);
+	DEPRECATED(4.7, "This function is deprecated. Please use ClearValue() instead.")
+	void ClearValueAsVector(FBlackboard::FKey KeyID);
+	DEPRECATED(4.7, "This function is deprecated. Please use SetValue<UBlackboardKeyType_Object>() instead.")
+	void SetValueAsObject(FBlackboard::FKey KeyID, UObject* ObjectValue);
+	DEPRECATED(4.7, "This function is deprecated. Please use SetValue<UBlackboardKeyType_Class>() instead.")
+	void SetValueAsClass(FBlackboard::FKey KeyID, UClass* ClassValue);
+	DEPRECATED(4.7, "This function is deprecated. Please use SetValue<UBlackboardKeyType_Enum>() instead.")
+	void SetValueAsEnum(FBlackboard::FKey KeyID, uint8 EnumValue);
+	DEPRECATED(4.7, "This function is deprecated. Please use SetValue<UBlackboardKeyType_Int>() instead.")
+	void SetValueAsInt(FBlackboard::FKey KeyID, int32 IntValue);
+	DEPRECATED(4.7, "This function is deprecated. Please use SetValue<UBlackboardKeyType_Float>() instead.")
+	void SetValueAsFloat(FBlackboard::FKey KeyID, float FloatValue);
+	DEPRECATED(4.7, "This function is deprecated. Please use SetValue<UBlackboardKeyType_Bool>() instead.")
+	void SetValueAsBool(FBlackboard::FKey KeyID, bool BoolValue);
+	DEPRECATED(4.7, "This function is deprecated. Please use SetValue<UBlackboardKeyType_String>() instead.")
+	void SetValueAsString(FBlackboard::FKey KeyID, FString StringValue);
+	DEPRECATED(4.7, "This function is deprecated. Please use SetValue<UBlackboardKeyType_Name>() instead.")
+	void SetValueAsName(FBlackboard::FKey KeyID, FName NameValue);
+	DEPRECATED(4.7, "This function is deprecated. Please use SetValue<UBlackboardKeyType_Vector>() instead.")
+	void SetValueAsVector(FBlackboard::FKey KeyID, FVector VectorValue);
+	DEPRECATED(4.7, "This function is deprecated. Please use GetValue<UBlackboardKeyType_Object>() instead.")
+	UObject* GetValueAsObject(FBlackboard::FKey KeyID) const;
+	DEPRECATED(4.7, "This function is deprecated. Please use GetValue<UBlackboardKeyType_Class>() instead.")
+	UClass* GetValueAsClass(FBlackboard::FKey KeyID) const;
+	DEPRECATED(4.7, "This function is deprecated. Please use GetValue<UBlackboardKeyType_Enum>() instead.")
+	uint8 GetValueAsEnum(FBlackboard::FKey KeyID) const;
+	DEPRECATED(4.7, "This function is deprecated. Please use GetValue<UBlackboardKeyType_Int>() instead.")
+	int32 GetValueAsInt(FBlackboard::FKey KeyID) const;
+	DEPRECATED(4.7, "This function is deprecated. Please use GetValue<UBlackboardKeyType_Float>() instead.")
+	float GetValueAsFloat(FBlackboard::FKey KeyID) const;
+	DEPRECATED(4.7, "This function is deprecated. Please use GetValue<UBlackboardKeyType_Bool>() instead.")
+	bool GetValueAsBool(FBlackboard::FKey KeyID) const;
+	DEPRECATED(4.7, "This function is deprecated. Please use GetValue<UBlackboardKeyType_String>() instead.")
+	FString GetValueAsString(FBlackboard::FKey KeyID) const;
+	DEPRECATED(4.7, "This function is deprecated. Please use GetValue<UBlackboardKeyType_Name>() instead.")
+	FName GetValueAsName(FBlackboard::FKey KeyID) const;
+	DEPRECATED(4.7, "This function is deprecated. Please use GetValue<UBlackboardKeyType_Vector>() instead.")
+	FVector GetValueAsVector(FBlackboard::FKey KeyID) const;
+	DEPRECATED(4.7, "This function is deprecated. Please use GetValue<UBlackboardKeyType_Rotator>() instead.")
+	FRotator GetValueAsRotator(FBlackboard::FKey KeyID) const;
+	UFUNCTION(BlueprintCallable, Category = "AI|Components|Blackboard", Meta = (DeprecatedFunction, DeprecationMessage = "Please use ClearValue instead"))
+	void ClearValueAsVector(const FName& KeyName);
+	UFUNCTION(BlueprintCallable, Category = "AI|Components|Blackboard", Meta = (DeprecatedFunction, DeprecationMessage = "Please use ClearValue instead"))
+	void ClearValueAsRotator(const FName& KeyName);
+
 protected:
 
 	/** cached behavior tree component */
@@ -262,6 +303,13 @@ FORCEINLINE bool UBlackboardComponent::HasValidAsset() const
 	return BlackboardAsset && BlackboardAsset->IsValid();
 }
 
+template<class TDataClass>
+bool UBlackboardComponent::SetValue(const FName& KeyName, typename TDataClass::FDataType Value)
+{
+	const FBlackboard::FKey KeyID = GetKeyID(KeyName);
+	return SetValue<TDataClass>(KeyID, Value);
+}
+
 template<class TDataClass>
 bool UBlackboardComponent::SetValue(FBlackboard::FKey KeyID, typename TDataClass::FDataType Value)
 {
@@ -296,4 +344,20 @@ bool UBlackboardComponent::SetValue(FBlackboard::FKey KeyID, typename TDataClass
 	}
 
 	return false;
-}
\ No newline at end of file
+}
+
+template<class TDataClass>
+typename TDataClass::FDataType UBlackboardComponent::GetValue(const FName& KeyName) const
+{
+	const FBlackboard::FKey KeyID = GetKeyID(KeyName);
+	return (GetKeyType(KeyID) != TDataClass::StaticClass())
+		? TDataClass::InvalidValue
+		: GetValue<TDataClass>(KeyID);
+}
+
+template<class TDataClass>
+typename TDataClass::FDataType UBlackboardComponent::GetValue(FBlackboard::FKey KeyID) const
+{
+	const uint8* RawData = GetKeyRawData(KeyID);
+	return RawData ? TDataClass::GetValue(RawData) : TDataClass::InvalidValue;
+}
diff --git a/Engine/Source/Runtime/AIModule/Private/BehaviorTree/BTFunctionLibrary.cpp b/Engine/Source/Runtime/AIModule/Private/BehaviorTree/BTFunctionLibrary.cpp
--- a/Engine/Source/Runtime/AIModule/Private/BehaviorTree/BTFunctionLibrary.cpp
+++ b/Engine/Source/Runtime/AIModule/Private/BehaviorTree/BTFunctionLibrary.cpp
@@ -123,7 +123,7 @@ void UBTFunctionLibrary::SetBlackboardValueAsObject(UBTNode* NodeOwner, const FB
 	UBlackboardComponent* BlackboardComp = GetBlackboard(*NodeOwner);
 	if (BlackboardComp != NULL)
 	{
-		BlackboardComp->SetValueAsObject(Key.SelectedKeyName, Value);
+		BlackboardComp->SetValue<UBlackboardKeyType_Object>(Key.SelectedKeyName, Value);
 	}
 }
 
@@ -133,7 +133,7 @@ void UBTFunctionLibrary::SetBlackboardValueAsClass(UBTNode* NodeOwner, const FBl
 	UBlackboardComponent* BlackboardComp = GetBlackboard(*NodeOwner);
 	if (BlackboardComp != NULL)
 	{
-		BlackboardComp->SetValueAsClass(Key.SelectedKeyName, Value);
+		BlackboardComp->SetValue<UBlackboardKeyType_Class>(Key.SelectedKeyName, Value);
 	}
 }
 
@@ -143,7 +143,7 @@ void UBTFunctionLibrary::SetBlackboardValueAsEnum(UBTNode* NodeOwner, const FBla
 	UBlackboardComponent* BlackboardComp = GetBlackboard(*NodeOwner);
 	if (BlackboardComp != NULL)
 	{
-		BlackboardComp->SetValueAsEnum(Key.SelectedKeyName, Value);
+		BlackboardComp->SetValue<UBlackboardKeyType_Enum>(Key.SelectedKeyName, Value);
 	}
 }
 
@@ -153,7 +153,7 @@ void UBTFunctionLibrary::SetBlackboardValueAsInt(UBTNode* NodeOwner, const FBlac
 	UBlackboardComponent* BlackboardComp = GetBlackboard(*NodeOwner);
 	if (BlackboardComp != NULL)
 	{
-		BlackboardComp->SetValueAsInt(Key.SelectedKeyName, Value);
+		BlackboardComp->SetValue<UBlackboardKeyType_Int>(Key.SelectedKeyName, Value);
 	}
 }
 
@@ -163,7 +163,7 @@ void UBTFunctionLibrary::SetBlackboardValueAsFloat(UBTNode* NodeOwner, const FBl
 	UBlackboardComponent* BlackboardComp = GetBlackboard(*NodeOwner);
 	if (BlackboardComp != NULL)
 	{
-		BlackboardComp->SetValueAsFloat(Key.SelectedKeyName, Value);
+		BlackboardComp->SetValue<UBlackboardKeyType_Float>(Key.SelectedKeyName, Value);
 	}
 }
 
@@ -173,7 +173,7 @@ void UBTFunctionLibrary::SetBlackboardValueAsBool(UBTNode* NodeOwner, const FBla
 	UBlackboardComponent* BlackboardComp = GetBlackboard(*NodeOwner);
 	if (BlackboardComp != NULL)
 	{
-		BlackboardComp->SetValueAsBool(Key.SelectedKeyName, Value);
+		BlackboardComp->SetValue<UBlackboardKeyType_Bool>(Key.SelectedKeyName, Value);
 	}
 }
 
@@ -183,7 +183,7 @@ void UBTFunctionLibrary::SetBlackboardValueAsString(UBTNode* NodeOwner, const FB
 	UBlackboardComponent* BlackboardComp = GetBlackboard(*NodeOwner);
 	if (BlackboardComp != NULL)
 	{
-		BlackboardComp->SetValueAsString(Key.SelectedKeyName, Value);
+		BlackboardComp->SetValue<UBlackboardKeyType_String>(Key.SelectedKeyName, Value);
 	}
 }
 
@@ -193,7 +193,7 @@ void UBTFunctionLibrary::SetBlackboardValueAsName(UBTNode* NodeOwner, const FBla
 	UBlackboardComponent* BlackboardComp = GetBlackboard(*NodeOwner);
 	if (BlackboardComp != NULL)
 	{
-		BlackboardComp->SetValueAsName(Key.SelectedKeyName, Value);
+		BlackboardComp->SetValue<UBlackboardKeyType_Name>(Key.SelectedKeyName, Value);
 	}
 }
 
@@ -203,7 +203,7 @@ void UBTFunctionLibrary::SetBlackboardValueAsVector(UBTNode* NodeOwner, const FB
 	UBlackboardComponent* BlackboardComp = GetBlackboard(*NodeOwner);
 	if (BlackboardComp != NULL)
 	{
-		BlackboardComp->SetValueAsVector(Key.SelectedKeyName, Value);
+		BlackboardComp->SetValue<UBlackboardKeyType_Vector>(Key.SelectedKeyName, Value);
 	}
 }
 
@@ -213,7 +213,7 @@ void UBTFunctionLibrary::ClearBlackboardValueAsVector(UBTNode* NodeOwner, const
 	UBlackboardComponent* BlackboardComp = GetBlackboard(*NodeOwner);
 	if (BlackboardComp != NULL)
 	{
-		BlackboardComp->ClearValueAsVector(Key.SelectedKeyName);
+		BlackboardComp->ClearValue(Key.SelectedKeyName);
 	}
 }
 
@@ -223,7 +223,7 @@ void UBTFunctionLibrary::SetBlackboardValueAsRotator(UBTNode* NodeOwner, const F
 	UBlackboardComponent* BlackboardComp = GetBlackboard(*NodeOwner);
 	if (BlackboardComp != NULL)
 	{
-		BlackboardComp->SetValueAsRotator(Key.SelectedKeyName, Value);
+		BlackboardComp->SetValue<UBlackboardKeyType_Rotator>(Key.SelectedKeyName, Value);
 	}
 }
 
diff --git a/Engine/Source/Runtime/AIModule/Private/BehaviorTree/Blackboard/BlackboardKeyType_Bool.cpp b/Engine/Source/Runtime/AIModule/Private/BehaviorTree/Blackboard/BlackboardKeyType_Bool.cpp
--- a/Engine/Source/Runtime/AIModule/Private/BehaviorTree/Blackboard/BlackboardKeyType_Bool.cpp
+++ b/Engine/Source/Runtime/AIModule/Private/BehaviorTree/Blackboard/BlackboardKeyType_Bool.cpp
@@ -3,6 +3,8 @@
 #include "AIModulePrivate.h"
 #include "BehaviorTree/Blackboard/BlackboardKeyType_Bool.h"
 
+const UBlackboardKeyType_Bool::FDataType UBlackboardKeyType_Bool::InvalidValue = false;
+
 UBlackboardKeyType_Bool::UBlackboardKeyType_Bool(const FObjectInitializer& ObjectInitializer) : Super(ObjectInitializer)
 {
 	ValueSize = sizeof(uint8);
diff --git a/Engine/Source/Runtime/AIModule/Private/BehaviorTree/Blackboard/BlackboardKeyType_Class.cpp b/Engine/Source/Runtime/AIModule/Private/BehaviorTree/Blackboard/BlackboardKeyType_Class.cpp
--- a/Engine/Source/Runtime/AIModule/Private/BehaviorTree/Blackboard/BlackboardKeyType_Class.cpp
+++ b/Engine/Source/Runtime/AIModule/Private/BehaviorTree/Blackboard/BlackboardKeyType_Class.cpp
@@ -3,6 +3,8 @@
 #include "AIModulePrivate.h"
 #include "BehaviorTree/Blackboard/BlackboardKeyType_Class.h"
 
+const UBlackboardKeyType_Class::FDataType UBlackboardKeyType_Class::InvalidValue = nullptr;
+
 UBlackboardKeyType_Class::UBlackboardKeyType_Class(const FObjectInitializer& ObjectInitializer) : Super(ObjectInitializer)
 {
 	ValueSize = sizeof(TWeakObjectPtr<UClass>);
diff --git a/Engine/Source/Runtime/AIModule/Private/BehaviorTree/Blackboard/BlackboardKeyType_Enum.cpp b/Engine/Source/Runtime/AIModule/Private/BehaviorTree/Blackboard/BlackboardKeyType_Enum.cpp
--- a/Engine/Source/Runtime/AIModule/Private/BehaviorTree/Blackboard/BlackboardKeyType_Enum.cpp
+++ b/Engine/Source/Runtime/AIModule/Private/BehaviorTree/Blackboard/BlackboardKeyType_Enum.cpp
@@ -3,6 +3,8 @@
 #include "AIModulePrivate.h"
 #include "BehaviorTree/Blackboard/BlackboardKeyType_Enum.h"
 
+const UBlackboardKeyType_Enum::FDataType UBlackboardKeyType_Enum::InvalidValue = UBlackboardKeyType_Enum::FDataType(0);
+
 UBlackboardKeyType_Enum::UBlackboardKeyType_Enum(const FObjectInitializer& ObjectInitializer) : Super(ObjectInitializer)
 {
 	ValueSize = sizeof(uint8);
diff --git a/Engine/Source/Runtime/AIModule/Private/BehaviorTree/Blackboard/BlackboardKeyType_Float.cpp b/Engine/Source/Runtime/AIModule/Private/BehaviorTree/Blackboard/BlackboardKeyType_Float.cpp
--- a/Engine/Source/Runtime/AIModule/Private/BehaviorTree/Blackboard/BlackboardKeyType_Float.cpp
+++ b/Engine/Source/Runtime/AIModule/Private/BehaviorTree/Blackboard/BlackboardKeyType_Float.cpp
@@ -3,6 +3,8 @@
 #include "AIModulePrivate.h"
 #include "BehaviorTree/Blackboard/BlackboardKeyType_Float.h"
 
+const UBlackboardKeyType_Float::FDataType UBlackboardKeyType_Float::InvalidValue = UBlackboardKeyType_Float::FDataType(0);
+
 UBlackboardKeyType_Float::UBlackboardKeyType_Float(const FObjectInitializer& ObjectInitializer) : Super(ObjectInitializer)
 {
 	ValueSize = sizeof(float);
diff --git a/Engine/Source/Runtime/AIModule/Private/BehaviorTree/Blackboard/BlackboardKeyType_Int.cpp b/Engine/Source/Runtime/AIModule/Private/BehaviorTree/Blackboard/BlackboardKeyType_Int.cpp
--- a/Engine/Source/Runtime/AIModule/Private/BehaviorTree/Blackboard/BlackboardKeyType_Int.cpp
+++ b/Engine/Source/Runtime/AIModule/Private/BehaviorTree/Blackboard/BlackboardKeyType_Int.cpp
@@ -3,6 +3,8 @@
 #include "AIModulePrivate.h"
 #include "BehaviorTree/Blackboard/BlackboardKeyType_Int.h"
 
+const UBlackboardKeyType_Int::FDataType UBlackboardKeyType_Int::InvalidValue = UBlackboardKeyType_Int::FDataType(0);
+
 UBlackboardKeyType_Int::UBlackboardKeyType_Int(const FObjectInitializer& ObjectInitializer) : Super(ObjectInitializer)
 {
 	ValueSize = sizeof(int32);
diff --git a/Engine/Source/Runtime/AIModule/Private/BehaviorTree/Blackboard/BlackboardKeyType_Name.cpp b/Engine/Source/Runtime/AIModule/Private/BehaviorTree/Blackboard/BlackboardKeyType_Name.cpp
--- a/Engine/Source/Runtime/AIModule/Private/BehaviorTree/Blackboard/BlackboardKeyType_Name.cpp
+++ b/Engine/Source/Runtime/AIModule/Private/BehaviorTree/Blackboard/BlackboardKeyType_Name.cpp
@@ -3,6 +3,8 @@
 #include "AIModulePrivate.h"
 #include "BehaviorTree/Blackboard/BlackboardKeyType_Name.h"
 
+const UBlackboardKeyType_Name::FDataType UBlackboardKeyType_Name::InvalidValue = NAME_None;
+
 UBlackboardKeyType_Name::UBlackboardKeyType_Name(const FObjectInitializer& ObjectInitializer) : Super(ObjectInitializer)
 {
 	ValueSize = sizeof(FName);
diff --git a/Engine/Source/Runtime/AIModule/Private/BehaviorTree/Blackboard/BlackboardKeyType_NativeEnum.cpp b/Engine/Source/Runtime/AIModule/Private/BehaviorTree/Blackboard/BlackboardKeyType_NativeEnum.cpp
--- a/Engine/Source/Runtime/AIModule/Private/BehaviorTree/Blackboard/BlackboardKeyType_NativeEnum.cpp
+++ b/Engine/Source/Runtime/AIModule/Private/BehaviorTree/Blackboard/BlackboardKeyType_NativeEnum.cpp
@@ -3,6 +3,8 @@
 #include "AIModulePrivate.h"
 #include "BehaviorTree/Blackboard/BlackboardKeyType_NativeEnum.h"
 
+const UBlackboardKeyType_NativeEnum::FDataType UBlackboardKeyType_NativeEnum::InvalidValue = UBlackboardKeyType_NativeEnum::FDataType(0);
+
 UBlackboardKeyType_NativeEnum::UBlackboardKeyType_NativeEnum(const FObjectInitializer& ObjectInitializer) : Super(ObjectInitializer)
 {
 	ValueSize = sizeof(uint8);
diff --git a/Engine/Source/Runtime/AIModule/Private/BehaviorTree/Blackboard/BlackboardKeyType_Object.cpp b/Engine/Source/Runtime/AIModule/Private/BehaviorTree/Blackboard/BlackboardKeyType_Object.cpp
--- a/Engine/Source/Runtime/AIModule/Private/BehaviorTree/Blackboard/BlackboardKeyType_Object.cpp
+++ b/Engine/Source/Runtime/AIModule/Private/BehaviorTree/Blackboard/BlackboardKeyType_Object.cpp
@@ -3,6 +3,8 @@
 #include "AIModulePrivate.h"
 #include "BehaviorTree/Blackboard/BlackboardKeyType_Object.h"
 
+const UBlackboardKeyType_Object::FDataType UBlackboardKeyType_Object::InvalidValue = nullptr;
+
 UBlackboardKeyType_Object::UBlackboardKeyType_Object(const FObjectInitializer& ObjectInitializer) : Super(ObjectInitializer)
 {
 	ValueSize = sizeof(FWeakObjectPtr);
diff --git a/Engine/Source/Runtime/AIModule/Private/BehaviorTree/Blackboard/BlackboardKeyType_Rotator.cpp b/Engine/Source/Runtime/AIModule/Private/BehaviorTree/Blackboard/BlackboardKeyType_Rotator.cpp
--- a/Engine/Source/Runtime/AIModule/Private/BehaviorTree/Blackboard/BlackboardKeyType_Rotator.cpp
+++ b/Engine/Source/Runtime/AIModule/Private/BehaviorTree/Blackboard/BlackboardKeyType_Rotator.cpp
@@ -3,6 +3,8 @@
 #include "AIModulePrivate.h"
 #include "BehaviorTree/Blackboard/BlackboardKeyType_Rotator.h"
 
+const UBlackboardKeyType_Rotator::FDataType UBlackboardKeyType_Rotator::InvalidValue = FAISystem::InvalidRotation;
+
 UBlackboardKeyType_Rotator::UBlackboardKeyType_Rotator(const FObjectInitializer& ObjectInitializer) : Super(ObjectInitializer)
 {
 	ValueSize = sizeof(FRotator);
diff --git a/Engine/Source/Runtime/AIModule/Private/BehaviorTree/Blackboard/BlackboardKeyType_String.cpp b/Engine/Source/Runtime/AIModule/Private/BehaviorTree/Blackboard/BlackboardKeyType_String.cpp
--- a/Engine/Source/Runtime/AIModule/Private/BehaviorTree/Blackboard/BlackboardKeyType_String.cpp
+++ b/Engine/Source/Runtime/AIModule/Private/BehaviorTree/Blackboard/BlackboardKeyType_String.cpp
@@ -3,6 +3,8 @@
 #include "AIModulePrivate.h"
 #include "BehaviorTree/Blackboard/BlackboardKeyType_String.h"
 
+const UBlackboardKeyType_String::FDataType UBlackboardKeyType_String::InvalidValue = FString();
+
 UBlackboardKeyType_String::UBlackboardKeyType_String(const FObjectInitializer& ObjectInitializer) : Super(ObjectInitializer)
 {
 	ValueSize = sizeof(FString);
diff --git a/Engine/Source/Runtime/AIModule/Private/BehaviorTree/Blackboard/BlackboardKeyType_Vector.cpp b/Engine/Source/Runtime/AIModule/Private/BehaviorTree/Blackboard/BlackboardKeyType_Vector.cpp
--- a/Engine/Source/Runtime/AIModule/Private/BehaviorTree/Blackboard/BlackboardKeyType_Vector.cpp
+++ b/Engine/Source/Runtime/AIModule/Private/BehaviorTree/Blackboard/BlackboardKeyType_Vector.cpp
@@ -3,6 +3,8 @@
 #include "AIModulePrivate.h"
 #include "BehaviorTree/Blackboard/BlackboardKeyType_Vector.h"
 
+const UBlackboardKeyType_Vector::FDataType UBlackboardKeyType_Vector::InvalidValue = FAISystem::InvalidLocation;
+
 UBlackboardKeyType_Vector::UBlackboardKeyType_Vector(const FObjectInitializer& ObjectInitializer) : Super(ObjectInitializer)
 {
 	ValueSize = sizeof(FVector);
diff --git a/Engine/Source/Runtime/AIModule/Private/BehaviorTree/BlackboardComponent.cpp b/Engine/Source/Runtime/AIModule/Private/BehaviorTree/BlackboardComponent.cpp
--- a/Engine/Source/Runtime/AIModule/Private/BehaviorTree/BlackboardComponent.cpp
+++ b/Engine/Source/Runtime/AIModule/Private/BehaviorTree/BlackboardComponent.cpp
@@ -422,36 +422,12 @@ void UBlackboardComponent::DescribeSelfToVisLog(FVisualLogEntry* Snapshot) const
 
 UObject* UBlackboardComponent::GetValueAsObject(const FName& KeyName) const
 {
-	FBlackboard::FKey KeyID = GetKeyID(KeyName);
-	if (GetKeyType(KeyID) != UBlackboardKeyType_Object::StaticClass())
-	{
-		KeyID = FBlackboard::InvalidKey;
-	}
-
-	return GetValueAsObject(KeyID);
-}
-
-UObject* UBlackboardComponent::GetValueAsObject(FBlackboard::FKey KeyID) const
-{
-	const uint8* RawData = GetKeyRawData(KeyID);
-	return RawData ? UBlackboardKeyType_Object::GetValue(RawData) : NULL;
+	return GetValue<UBlackboardKeyType_Object>(KeyName);
 }
 
 UClass* UBlackboardComponent::GetValueAsClass(const FName& KeyName) const
 {
-	FBlackboard::FKey KeyID = GetKeyID(KeyName);
-	if (GetKeyType(KeyID) != UBlackboardKeyType_Class::StaticClass())
-	{
-		KeyID = FBlackboard::InvalidKey;
-	}
-
-	return GetValueAsClass(KeyID);
-}
-
-UClass* UBlackboardComponent::GetValueAsClass(FBlackboard::FKey KeyID) const
-{
-	const uint8* RawData = GetKeyRawData(KeyID);
-	return RawData ? UBlackboardKeyType_Class::GetValue(RawData) : NULL;
+	return GetValue<UBlackboardKeyType_Class>(KeyName);
 }
 
 uint8 UBlackboardComponent::GetValueAsEnum(const FName& KeyName) const
@@ -460,136 +436,45 @@ uint8 UBlackboardComponent::GetValueAsEnum(const FName& KeyName) const
 	if (GetKeyType(KeyID) != UBlackboardKeyType_Enum::StaticClass() &&
 		GetKeyType(KeyID) != UBlackboardKeyType_NativeEnum::StaticClass())
 	{
-		KeyID = FBlackboard::InvalidKey;
+		return UBlackboardKeyType_Enum::InvalidValue;
 	}
 
-	return GetValueAsEnum(KeyID);
-}
-
-uint8 UBlackboardComponent::GetValueAsEnum(FBlackboard::FKey KeyID) const
-{
-	const uint8* RawData = GetKeyRawData(KeyID);
-	return RawData ? UBlackboardKeyType_Enum::GetValue(RawData) : 0;
+	return GetValue<UBlackboardKeyType_Enum>(KeyName);
 }
 
 int32 UBlackboardComponent::GetValueAsInt(const FName& KeyName) const
 {
-	FBlackboard::FKey KeyID = GetKeyID(KeyName);
-	if (GetKeyType(KeyID) != UBlackboardKeyType_Int::StaticClass())
-	{
-		KeyID = FBlackboard::InvalidKey;
-	}
-
-	return GetValueAsInt(KeyID);
-}
-
-int32 UBlackboardComponent::GetValueAsInt(FBlackboard::FKey KeyID) const
-{
-	const uint8* RawData = GetKeyRawData(KeyID);
-	return RawData ? UBlackboardKeyType_Int::GetValue(RawData) : 0;
+	return GetValue<UBlackboardKeyType_Int>(KeyName);
 }
 
 float UBlackboardComponent::GetValueAsFloat(const FName& KeyName) const
 {
-	FBlackboard::FKey KeyID = GetKeyID(KeyName);
-	if (GetKeyType(KeyID) != UBlackboardKeyType_Float::StaticClass())
-	{
-		KeyID = FBlackboard::InvalidKey;
-	}
-
-	return GetValueAsFloat(KeyID);
-}
-
-float UBlackboardComponent::GetValueAsFloat(FBlackboard::FKey KeyID) const
-{
-	const uint8* RawData = GetKeyRawData(KeyID);
-	return RawData ? UBlackboardKeyType_Float::GetValue(RawData) : 0.0f;
+	return GetValue<UBlackboardKeyType_Float>(KeyName);
 }
 
 bool UBlackboardComponent::GetValueAsBool(const FName& KeyName) const
 {
-	FBlackboard::FKey KeyID = GetKeyID(KeyName);
-	if (GetKeyType(KeyID) != UBlackboardKeyType_Bool::StaticClass())
-	{
-		KeyID = FBlackboard::InvalidKey;
-	}
-
-	return GetValueAsBool(KeyID);
-}
-
-bool UBlackboardComponent::GetValueAsBool(FBlackboard::FKey KeyID) const
-{
-	const uint8* RawData = GetKeyRawData(KeyID);
-	return RawData ? UBlackboardKeyType_Bool::GetValue(RawData) : false;
+	return GetValue<UBlackboardKeyType_Bool>(KeyName);
 }
 
 FString UBlackboardComponent::GetValueAsString(const FName& KeyName) const
 {
-	FBlackboard::FKey KeyID = GetKeyID(KeyName);
-	if (GetKeyType(KeyID) != UBlackboardKeyType_String::StaticClass())
-	{
-		KeyID = FBlackboard::InvalidKey;
-	}
-
-	return GetValueAsString(KeyID);
-}
-
-FString UBlackboardComponent::GetValueAsString(FBlackboard::FKey KeyID) const
-{
-	const uint8* RawData = GetKeyRawData(KeyID);
-	return RawData ? UBlackboardKeyType_String::GetValue(RawData) : FString();
+	return GetValue<UBlackboardKeyType_String>(KeyName);
 }
 
 FName UBlackboardComponent::GetValueAsName(const FName& KeyName) const
 {
-	FBlackboard::FKey KeyID = GetKeyID(KeyName);
-	if (GetKeyType(KeyID) != UBlackboardKeyType_Name::StaticClass())
-	{
-		KeyID = FBlackboard::InvalidKey;
-	}
-
-	return GetValueAsName(KeyID);
-}
-
-FName UBlackboardComponent::GetValueAsName(FBlackboard::FKey KeyID) const
-{
-	const uint8* RawData = GetKeyRawData(KeyID);
-	return RawData ? UBlackboardKeyType_Name::GetValue(RawData) : NAME_None;
+	return GetValue<UBlackboardKeyType_Name>(KeyName);
 }
 
 FVector UBlackboardComponent::GetValueAsVector(const FName& KeyName) const
 {
-	FBlackboard::FKey KeyID = GetKeyID(KeyName);
-	if (GetKeyType(KeyID) != UBlackboardKeyType_Vector::StaticClass())
-	{
-		KeyID = FBlackboard::InvalidKey;
-	}
-
-	return GetValueAsVector(KeyID);
+	return GetValue<UBlackboardKeyType_Vector>(KeyName);
 }
 
-FVector UBlackboardComponent::GetValueAsVector(FBlackboard::FKey KeyID) const
-{
-	const uint8* RawData = GetKeyRawData(KeyID);
-	return RawData ? UBlackboardKeyType_Vector::GetValue(RawData) : FAISystem::InvalidLocation;
-}
-
-
 FRotator UBlackboardComponent::GetValueAsRotator(const FName& KeyName) const
 {
-	uint8 KeyID = GetKeyID(KeyName);
-	if (GetKeyType(KeyID) != UBlackboardKeyType_Rotator::StaticClass())
-	{
-		KeyID = INDEX_NONE;
-	}
-
-	return GetValueAsRotator(KeyID);
-}
-
-FRotator UBlackboardComponent::GetValueAsRotator(uint8 KeyID) const
-{
-	const uint8* RawData = GetKeyRawData(KeyID);
-	return RawData ? UBlackboardKeyType_Rotator::GetValue(RawData) : FRotator::ZeroRotator;
+	return GetValue<UBlackboardKeyType_Rotator>(KeyName);
 }
 
 void UBlackboardComponent::SetValueAsObject(const FName& KeyName, UObject* ObjectValue)
@@ -598,48 +483,6 @@ void UBlackboardComponent::SetValueAsObject(const FName& KeyName, UObject* Objec
 	SetValue<UBlackboardKeyType_Object>(KeyID, ObjectValue);
 }
 
-//void UBlackboardComponent::SetValueAsObject(FBlackboard::FKey KeyID, UObject* ObjectValue)
-//{
-//	if (GetKeyType(KeyID) != UBlackboardKeyType_Object::StaticClass())
-//	{
-//		return;
-//	}
-//
-//	uint8* RawData = GetKeyRawData(KeyID);
-//	if (RawData)
-//	{
-//		const bool bChanged = UBlackboardKeyType_Object::SetValue(RawData, ObjectValue);
-//		if (bChanged)
-//		{
-//			NotifyObservers(KeyID);
-//			if (bHasSyncedKeys && IsKeyInstanceSynced(KeyID))
-//			{
-//				SynchronizeKeys(KeyID, ObjectValue);
-//			}
-//		}
-//	}
-//}
-//
-////void UBlackboardData::SynchronizeKeys(FBlackboard::FKey KeyID, UObject* ObjectValue)
-////{
-////	check(BlackboardAsset);
-////
-////	if (BrainComp && BrainComp->GetAIOwner())
-////	{
-////		UAISystem* AISystem = UAISystem::GetCurrentSafe(GetWorld());
-////		for (auto Iter = AISystem->CreateBlackboardDataToComponentsIterator(*BlackboardAsset); Iter; ++Iter)
-////		{
-////			UBlackboardComponent* OtherBlackboard = Iter.Value();
-////			if (OtherBlackboard && OtherBlackboard != this && OtherBlackboard->BrainComp && OtherBlackboard->BrainComp->GetAIOwner()
-////				&& OtherBlackboard->BrainComp->GetAIOwner()->ShouldInstanceSyncUpdatedBlackboardField(*OtherBlackboard, KeyID, *this, *BrainComp->GetAIOwner()))
-////			{
-////				UBlackboardKeyType_Object::SetValue(OtherBlackboard->GetKeyRawData(KeyID), ObjectValue);
-////				OtherBlackboard->NotifyObservers(KeyID);
-////			}
-////		}
-////	}
-////}
-
 void UBlackboardComponent::SetValueAsClass(const FName& KeyName, UClass* ClassValue)
 {
 	const FBlackboard::FKey KeyID = GetKeyID(KeyName);
@@ -700,6 +543,12 @@ void UBlackboardComponent::ClearValueAsVector(const FName& KeyName)
 	ClearValue(KeyID);
 }
 
+void UBlackboardComponent::ClearValueAsVector(FBlackboard::FKey KeyID)
+{
+	ensure(false && "This function is deprecated. Use more generic ClearValue() instead");
+	ClearValue(KeyID);
+}
+
 bool UBlackboardComponent::IsVectorValueSet(const FName& KeyName) const
 {
 	const FBlackboard::FKey KeyID = GetKeyID(KeyName);
@@ -708,37 +557,20 @@ bool UBlackboardComponent::IsVectorValueSet(const FName& KeyName) const
 
 bool UBlackboardComponent::IsVectorValueSet(FBlackboard::FKey KeyID) const
 {
-	FVector VectorValue = GetValueAsVector(KeyID);
-	if (VectorValue == FAISystem::InvalidLocation)
-	{
-		return false;
-	}
-
-	return true;
+	FVector VectorValue = GetValue<UBlackboardKeyType_Vector>(KeyID);
+	return (VectorValue != FAISystem::InvalidLocation);
 }
 
 void UBlackboardComponent::ClearValueAsRotator(const FName& KeyName)
 {
 	const FBlackboard::FKey KeyID = GetKeyID(KeyName);
-	ClearValueAsRotator(KeyID);
+	ClearValue(KeyID);
 }
 
 void UBlackboardComponent::ClearValueAsRotator(uint8 KeyID)
 {
-	if (GetKeyType(KeyID) != UBlackboardKeyType_Rotator::StaticClass())
-	{
-		return;
-	}
-
-	uint8* RawData = GetKeyRawData(KeyID);
-	if (RawData)
-	{
-		const bool bChanged = UBlackboardKeyType_Rotator::SetValue(RawData, FAISystem::InvalidRotation);
-		if (bChanged)
-		{
-			NotifyObservers(KeyID);
-		}
-	}
+	ensure(false && "This function is deprecated. Use more generic ClearValue() instead");
+	ClearValue(KeyID);
 }
 
 void UBlackboardComponent::ClearValue(const FName& KeyName)
@@ -824,3 +656,101 @@ bool UBlackboardComponent::GetRotationFromEntry(FBlackboard::FKey KeyID, FRotato
 
 	return false;
 }
+
+//----------------------------------------------------------------------//
+// Deprecated
+//----------------------------------------------------------------------//
+void UBlackboardComponent::SetValueAsObject(FBlackboard::FKey KeyID, UObject* ObjectValue)
+{
+	SetValue<UBlackboardKeyType_Object>(KeyID, ObjectValue);
+}
+
+void UBlackboardComponent::SetValueAsClass(FBlackboard::FKey KeyID, UClass* ClassValue)
+{
+	SetValue<UBlackboardKeyType_Class>(KeyID, ClassValue);
+}
+
+void UBlackboardComponent::SetValueAsEnum(FBlackboard::FKey KeyID, uint8 EnumValue)
+{
+	SetValue<UBlackboardKeyType_Enum>(KeyID, EnumValue);
+}
+
+void UBlackboardComponent::SetValueAsInt(FBlackboard::FKey KeyID, int32 IntValue)
+{
+	SetValue<UBlackboardKeyType_Int>(KeyID, IntValue);
+}
+
+void UBlackboardComponent::SetValueAsFloat(FBlackboard::FKey KeyID, float FloatValue)
+{
+	SetValue<UBlackboardKeyType_Float>(KeyID, FloatValue);
+}
+
+void UBlackboardComponent::SetValueAsBool(FBlackboard::FKey KeyID, bool BoolValue)
+{
+	SetValue<UBlackboardKeyType_Bool>(KeyID, BoolValue);
+}
+
+void UBlackboardComponent::SetValueAsString(FBlackboard::FKey KeyID, FString StringValue)
+{
+	SetValue<UBlackboardKeyType_String>(KeyID, StringValue);
+}
+
+void UBlackboardComponent::SetValueAsName(FBlackboard::FKey KeyID, FName NameValue)
+{
+	SetValue<UBlackboardKeyType_Name>(KeyID, NameValue);
+}
+
+void UBlackboardComponent::SetValueAsVector(FBlackboard::FKey KeyID, FVector VectorValue)
+{
+	SetValue<UBlackboardKeyType_Vector>(KeyID, VectorValue);
+}
+
+UObject* UBlackboardComponent::GetValueAsObject(FBlackboard::FKey KeyID) const
+{
+	return GetValue<UBlackboardKeyType_Object>(KeyID);
+}
+
+UClass* UBlackboardComponent::GetValueAsClass(FBlackboard::FKey KeyID) const
+{
+	return GetValue<UBlackboardKeyType_Class>(KeyID);
+}
+
+uint8 UBlackboardComponent::GetValueAsEnum(FBlackboard::FKey KeyID) const
+{
+	return GetValue<UBlackboardKeyType_Enum>(KeyID);
+}
+
+int32 UBlackboardComponent::GetValueAsInt(FBlackboard::FKey KeyID) const
+{
+	return GetValue<UBlackboardKeyType_Int>(KeyID);
+}
+
+float UBlackboardComponent::GetValueAsFloat(FBlackboard::FKey KeyID) const
+{
+	return GetValue<UBlackboardKeyType_Float>(KeyID);
+}
+
+bool UBlackboardComponent::GetValueAsBool(FBlackboard::FKey KeyID) const
+{
+	return GetValue<UBlackboardKeyType_Bool>(KeyID);
+}
+
+FString UBlackboardComponent::GetValueAsString(FBlackboard::FKey KeyID) const
+{
+	return GetValue<UBlackboardKeyType_String>(KeyID);
+}
+
+FName UBlackboardComponent::GetValueAsName(FBlackboard::FKey KeyID) const
+{
+	return GetValue<UBlackboardKeyType_Name>(KeyID);
+}
+
+FVector UBlackboardComponent::GetValueAsVector(FBlackboard::FKey KeyID) const
+{
+	return GetValue<UBlackboardKeyType_Vector>(KeyID);
+}
+
+FRotator UBlackboardComponent::GetValueAsRotator(FBlackboard::FKey KeyID) const
+{
+	return GetValue<UBlackboardKeyType_Rotator>(KeyID);
+}
diff --git a/Engine/Source/Runtime/AIModule/Private/BehaviorTree/Decorators/BTDecorator_CheckGameplayTagsOnActor.cpp b/Engine/Source/Runtime/AIModule/Private/BehaviorTree/Decorators/BTDecorator_CheckGameplayTagsOnActor.cpp
--- a/Engine/Source/Runtime/AIModule/Private/BehaviorTree/Decorators/BTDecorator_CheckGameplayTagsOnActor.cpp
+++ b/Engine/Source/Runtime/AIModule/Private/BehaviorTree/Decorators/BTDecorator_CheckGameplayTagsOnActor.cpp
@@ -29,7 +29,7 @@ bool UBTDecorator_CheckGameplayTagsOnActor::CalculateRawConditionValue(UBehavior
 		return false;
 	}
 
-	IGameplayTagAssetInterface* GameplayTagAssetInterface = Cast<IGameplayTagAssetInterface>(BlackboardComp->GetValueAsObject(ActorToCheck.GetSelectedKeyID()));
+	IGameplayTagAssetInterface* GameplayTagAssetInterface = Cast<IGameplayTagAssetInterface>(BlackboardComp->GetValue<UBlackboardKeyType_Object>(ActorToCheck.GetSelectedKeyID()));
 	if (GameplayTagAssetInterface == NULL)
 	{
 		return false;
diff --git a/Engine/Source/Runtime/AIModule/Private/BehaviorTree/Services/BTService_DefaultFocus.cpp b/Engine/Source/Runtime/AIModule/Private/BehaviorTree/Services/BTService_DefaultFocus.cpp
--- a/Engine/Source/Runtime/AIModule/Private/BehaviorTree/Services/BTService_DefaultFocus.cpp
+++ b/Engine/Source/Runtime/AIModule/Private/BehaviorTree/Services/BTService_DefaultFocus.cpp
@@ -34,7 +34,7 @@ void UBTService_DefaultFocus::OnBecomeRelevant(UBehaviorTreeComponent& OwnerComp
 	{
 		if (BlackboardKey.SelectedKeyType == UBlackboardKeyType_Object::StaticClass())
 		{
-			UObject* KeyValue = MyBlackboard->GetValueAsObject(BlackboardKey.GetSelectedKeyID());
+			UObject* KeyValue = MyBlackboard->GetValue<UBlackboardKeyType_Object>(BlackboardKey.GetSelectedKeyID());
 			AActor* TargetActor = Cast<AActor>(KeyValue);
 			if (TargetActor)
 			{
@@ -45,7 +45,7 @@ void UBTService_DefaultFocus::OnBecomeRelevant(UBehaviorTreeComponent& OwnerComp
 		}
 		else
 		{
-			const FVector FocusLocation = MyBlackboard->GetValueAsVector(BlackboardKey.GetSelectedKeyID());
+			const FVector FocusLocation = MyBlackboard->GetValue<UBlackboardKeyType_Vector>(BlackboardKey.GetSelectedKeyID());
 			OwnerController->SetFocalPoint(FocusLocation, EAIFocusPriority::Default);
 			MyMemory->FocusLocationSet = FocusLocation;
 		}
diff --git a/Engine/Source/Runtime/AIModule/Private/BehaviorTree/Tasks/BTTask_MoveDirectlyToward.cpp b/Engine/Source/Runtime/AIModule/Private/BehaviorTree/Tasks/BTTask_MoveDirectlyToward.cpp
--- a/Engine/Source/Runtime/AIModule/Private/BehaviorTree/Tasks/BTTask_MoveDirectlyToward.cpp
+++ b/Engine/Source/Runtime/AIModule/Private/BehaviorTree/Tasks/BTTask_MoveDirectlyToward.cpp
@@ -32,7 +32,7 @@ EBTNodeResult::Type UBTTask_MoveDirectlyToward::ExecuteTask(UBehaviorTreeCompone
 
 		if (BlackboardKey.SelectedKeyType == UBlackboardKeyType_Object::StaticClass())
 		{
-			UObject* KeyValue = MyBlackboard->GetValueAsObject(BlackboardKey.GetSelectedKeyID());
+			UObject* KeyValue = MyBlackboard->GetValue<UBlackboardKeyType_Object>(BlackboardKey.GetSelectedKeyID());
 			AActor* TargetActor = Cast<AActor>(KeyValue);
 			if (TargetActor)
 			{
@@ -43,7 +43,7 @@ EBTNodeResult::Type UBTTask_MoveDirectlyToward::ExecuteTask(UBehaviorTreeCompone
 		}
 		else if (BlackboardKey.SelectedKeyType == UBlackboardKeyType_Vector::StaticClass())
 		{
-			const FVector TargetLocation = MyBlackboard->GetValueAsVector(BlackboardKey.GetSelectedKeyID());
+			const FVector TargetLocation = MyBlackboard->GetValue<UBlackboardKeyType_Vector>(BlackboardKey.GetSelectedKeyID());
 			RequestResult = MyController->MoveToLocation(TargetLocation, AcceptableRadius, /*bStopOnOverlap=*/true, /*bUsePathfinding=*/false, /*bProjectDestinationToNavigation=*/bProjectVectorGoalToNavigation, bAllowStrafe);
 		}
 
diff --git a/Engine/Source/Runtime/AIModule/Private/BehaviorTree/Tasks/BTTask_MoveTo.cpp b/Engine/Source/Runtime/AIModule/Private/BehaviorTree/Tasks/BTTask_MoveTo.cpp
--- a/Engine/Source/Runtime/AIModule/Private/BehaviorTree/Tasks/BTTask_MoveTo.cpp
+++ b/Engine/Source/Runtime/AIModule/Private/BehaviorTree/Tasks/BTTask_MoveTo.cpp
@@ -54,7 +54,7 @@ EBTNodeResult::Type UBTTask_MoveTo::PerformMoveTask(UBehaviorTreeComponent& Owne
 
 		if (BlackboardKey.SelectedKeyType == UBlackboardKeyType_Object::StaticClass())
 		{
-			UObject* KeyValue = MyBlackboard->GetValueAsObject(BlackboardKey.GetSelectedKeyID());
+			UObject* KeyValue = MyBlackboard->GetValue<UBlackboardKeyType_Object>(BlackboardKey.GetSelectedKeyID());
 			AActor* TargetActor = Cast<AActor>(KeyValue);
 			if (TargetActor)
 			{
@@ -67,7 +67,7 @@ EBTNodeResult::Type UBTTask_MoveTo::PerformMoveTask(UBehaviorTreeComponent& Owne
 		}
 		else if (BlackboardKey.SelectedKeyType == UBlackboardKeyType_Vector::StaticClass())
 		{
-			const FVector TargetLocation = MyBlackboard->GetValueAsVector(BlackboardKey.GetSelectedKeyID());
+			const FVector TargetLocation = MyBlackboard->GetValue<UBlackboardKeyType_Vector>(BlackboardKey.GetSelectedKeyID());
 			RequestResult = MyController->MoveToLocation(TargetLocation, AcceptableRadius, true, true, false, bAllowStrafe, FilterClass);
 		}
 
diff --git a/Engine/Source/Runtime/AIModule/Private/BehaviorTree/Tasks/BTTask_RotateToFaceBBEntry.cpp b/Engine/Source/Runtime/AIModule/Private/BehaviorTree/Tasks/BTTask_RotateToFaceBBEntry.cpp
--- a/Engine/Source/Runtime/AIModule/Private/BehaviorTree/Tasks/BTTask_RotateToFaceBBEntry.cpp
+++ b/Engine/Source/Runtime/AIModule/Private/BehaviorTree/Tasks/BTTask_RotateToFaceBBEntry.cpp
@@ -60,7 +60,7 @@ EBTNodeResult::Type UBTTask_RotateToFaceBBEntry::ExecuteTask(UBehaviorTreeCompon
 
 	if (BlackboardKey.SelectedKeyType == UBlackboardKeyType_Object::StaticClass())
 	{
-		UObject* KeyValue = MyBlackboard->GetValueAsObject(BlackboardKey.GetSelectedKeyID());
+		UObject* KeyValue = MyBlackboard->GetValue<UBlackboardKeyType_Object>(BlackboardKey.GetSelectedKeyID());
 		AActor* ActorValue = Cast<AActor>(KeyValue);
 
 		if (ActorValue != NULL)
@@ -83,7 +83,7 @@ EBTNodeResult::Type UBTTask_RotateToFaceBBEntry::ExecuteTask(UBehaviorTreeCompon
 	}
 	else if (BlackboardKey.SelectedKeyType == UBlackboardKeyType_Vector::StaticClass())
 	{
-		const FVector KeyValue = MyBlackboard->GetValueAsVector(BlackboardKey.GetSelectedKeyID());
+		const FVector KeyValue = MyBlackboard->GetValue<UBlackboardKeyType_Vector>(BlackboardKey.GetSelectedKeyID());
 		
 		if (FAISystem::IsValidLocation(KeyValue))
 		{
@@ -104,7 +104,7 @@ EBTNodeResult::Type UBTTask_RotateToFaceBBEntry::ExecuteTask(UBehaviorTreeCompon
 	}
 	else if (BlackboardKey.SelectedKeyType == UBlackboardKeyType_Rotator::StaticClass())
 	{
-		const FRotator KeyValue = MyBlackboard->GetValueAsRotator(BlackboardKey.GetSelectedKeyID());
+		const FRotator KeyValue = MyBlackboard->GetValue<UBlackboardKeyType_Rotator>(BlackboardKey.GetSelectedKeyID());
 
 		if (FAISystem::IsValidRotation(KeyValue))
 		{
diff --git a/Engine/Source/Runtime/AIModule/Private/BehaviorTree/Tasks/BTTask_WaitBlackboardTime.cpp b/Engine/Source/Runtime/AIModule/Private/BehaviorTree/Tasks/BTTask_WaitBlackboardTime.cpp
--- a/Engine/Source/Runtime/AIModule/Private/BehaviorTree/Tasks/BTTask_WaitBlackboardTime.cpp
+++ b/Engine/Source/Runtime/AIModule/Private/BehaviorTree/Tasks/BTTask_WaitBlackboardTime.cpp
@@ -24,7 +24,7 @@ EBTNodeResult::Type UBTTask_WaitBlackboardTime::ExecuteTask(UBehaviorTreeCompone
 	const UBlackboardComponent* MyBlackboard = OwnerComp.GetBlackboardComponent();
 	if (MyBlackboard && BlackboardKey.SelectedKeyType == UBlackboardKeyType_Float::StaticClass())
 	{
-		WaitTime = MyBlackboard->GetValueAsFloat(BlackboardKey.GetSelectedKeyID());
+		WaitTime = MyBlackboard->GetValue<UBlackboardKeyType_Float>(BlackboardKey.GetSelectedKeyID());
 	}
 	
 	return Super::ExecuteTask(OwnerComp, NodeMemory);
