Compilation fixes #UE4

diff --git a/Engine/Source/Runtime/AIModule/Classes/BehaviorTree/BlackboardData.h b/Engine/Source/Runtime/AIModule/Classes/BehaviorTree/BlackboardData.h
--- a/Engine/Source/Runtime/AIModule/Classes/BehaviorTree/BlackboardData.h
+++ b/Engine/Source/Runtime/AIModule/Classes/BehaviorTree/BlackboardData.h
@@ -2,6 +2,7 @@
 #pragma once
 
 #include "Engine/DataAsset.h"
+#include "BehaviorTree/BehaviorTreeTypes.h"
 #include "BlackboardData.generated.h"
 
 class UBlackboardData;
diff --git a/Engine/Source/Runtime/AIModule/Private/AISystem.cpp b/Engine/Source/Runtime/AIModule/Private/AISystem.cpp
--- a/Engine/Source/Runtime/AIModule/Private/AISystem.cpp
+++ b/Engine/Source/Runtime/AIModule/Private/AISystem.cpp
@@ -4,10 +4,12 @@
 #include "EngineUtils.h"
 #include "StringClassReference.h"
 #include "BehaviorTree/BehaviorTreeManager.h"
+#include "BehaviorTree/BlackboardComponent.h"
 #include "EnvironmentQuery/EnvQueryManager.h"
 #include "Perception/AIPerceptionSystem.h"
 #include "GameFramework/PlayerController.h"
 #include "HotSpots/AIHotSpotManager.h"
+#include "BehaviorTree/BlackboardData.h"
 #include "AISystem.h"
 
 
