Force integration to fix issues with P4 history caused due to CL 2385683

diff --git a/Engine/Source/Runtime/AIModule/Classes/BehaviorTree/Tasks/BTTask_PushPawnAction.h b/Engine/Source/Runtime/AIModule/Classes/BehaviorTree/Tasks/BTTask_PushPawnAction.h
new file mode 100644
--- /dev/null
+++ b/Engine/Source/Runtime/AIModule/Classes/BehaviorTree/Tasks/BTTask_PushPawnAction.h
@@ -0,0 +1,25 @@
+// Copyright 1998-2015 Epic Games, Inc. All Rights Reserved.
+
+#pragma once
+
+#include "BehaviorTree/Tasks/BTTask_PawnActionBase.h"
+#include "BTTask_PushPawnAction.generated.h"
+
+class UPawnAction;
+
+/**
+ * Action task node.
+ * Push pawn action to controller.
+ */
+UCLASS()
+class AIMODULE_API UBTTask_PushPawnAction : public UBTTask_PawnActionBase
+{
+	GENERATED_UCLASS_BODY()
+
+protected:
+	UPROPERTY(EditAnywhere, Instanced, Category = Action)
+	UPawnAction* Action;
+
+	virtual EBTNodeResult::Type ExecuteTask(UBehaviorTreeComponent* OwnerComp, uint8* NodeMemory) override;
+	virtual FString GetStaticDescription() const override;
+};
diff --git a/Engine/Source/Runtime/AIModule/Private/BehaviorTree/Tasks/BTTask_PushPawnAction.cpp b/Engine/Source/Runtime/AIModule/Private/BehaviorTree/Tasks/BTTask_PushPawnAction.cpp
new file mode 100644
--- /dev/null
+++ b/Engine/Source/Runtime/AIModule/Private/BehaviorTree/Tasks/BTTask_PushPawnAction.cpp
@@ -0,0 +1,27 @@
+// Copyright 1998-2015 Epic Games, Inc. All Rights Reserved.
+
+#include "AIModulePrivate.h"
+#include "AIController.h"
+#include "Actions/PawnActionsComponent.h"
+#include "BehaviorTree/Tasks/BTTask_PushPawnAction.h"
+
+UBTTask_PushPawnAction::UBTTask_PushPawnAction(const FObjectInitializer& ObjectInitializer) : Super(ObjectInitializer)
+{
+	NodeName = "Push PawnAction";
+}
+
+EBTNodeResult::Type UBTTask_PushPawnAction::ExecuteTask(UBehaviorTreeComponent* OwnerComp, uint8* NodeMemory)
+{
+	UPawnAction* ActionCopy = Action ? DuplicateObject<UPawnAction>(Action, OwnerComp) : nullptr;
+	if (ActionCopy == nullptr)
+	{
+		return EBTNodeResult::Failed;
+	}
+
+	return PushAction(OwnerComp, ActionCopy);
+}
+
+FString UBTTask_PushPawnAction::GetStaticDescription() const
+{
+	return FString::Printf(TEXT("Push Action: %s"), Action ? *Action->GetDisplayName() : TEXT("None"));
+}
