BTTask_PushPawnAction compilation issues fix #UE4

diff --git a/Engine/Source/Runtime/AIModule/Classes/BehaviorTree/Tasks/BTTask_PushPawnAction.h b/Engine/Source/Runtime/AIModule/Classes/BehaviorTree/Tasks/BTTask_PushPawnAction.h
--- a/Engine/Source/Runtime/AIModule/Classes/BehaviorTree/Tasks/BTTask_PushPawnAction.h
+++ b/Engine/Source/Runtime/AIModule/Classes/BehaviorTree/Tasks/BTTask_PushPawnAction.h
@@ -20,6 +20,6 @@ protected:
 	UPROPERTY(EditAnywhere, Instanced, Category = Action)
 	UPawnAction* Action;
 
-	virtual EBTNodeResult::Type ExecuteTask(UBehaviorTreeComponent* OwnerComp, uint8* NodeMemory) override;
+	virtual EBTNodeResult::Type ExecuteTask(UBehaviorTreeComponent& OwnerComp, uint8* NodeMemory) override;
 	virtual FString GetStaticDescription() const override;
 };
diff --git a/Engine/Source/Runtime/AIModule/Private/BehaviorTree/Tasks/BTTask_PushPawnAction.cpp b/Engine/Source/Runtime/AIModule/Private/BehaviorTree/Tasks/BTTask_PushPawnAction.cpp
--- a/Engine/Source/Runtime/AIModule/Private/BehaviorTree/Tasks/BTTask_PushPawnAction.cpp
+++ b/Engine/Source/Runtime/AIModule/Private/BehaviorTree/Tasks/BTTask_PushPawnAction.cpp
@@ -10,9 +10,9 @@ UBTTask_PushPawnAction::UBTTask_PushPawnAction(const FObjectInitializer& ObjectI
 	NodeName = "Push PawnAction";
 }
 
-EBTNodeResult::Type UBTTask_PushPawnAction::ExecuteTask(UBehaviorTreeComponent* OwnerComp, uint8* NodeMemory)
+EBTNodeResult::Type UBTTask_PushPawnAction::ExecuteTask(UBehaviorTreeComponent& OwnerComp, uint8* NodeMemory)
 {
-	UPawnAction* ActionCopy = Action ? DuplicateObject<UPawnAction>(Action, OwnerComp) : nullptr;
+	UPawnAction* ActionCopy = Action ? DuplicateObject<UPawnAction>(Action, &OwnerComp) : nullptr;
 	if (ActionCopy == nullptr)
 	{
 		return EBTNodeResult::Failed;
@@ -23,5 +23,6 @@ EBTNodeResult::Type UBTTask_PushPawnAction::ExecuteTask(UBehaviorTreeComponent*
 
 FString UBTTask_PushPawnAction::GetStaticDescription() const
 {
-	return FString::Printf(TEXT("Push Action: %s"), Action ? *Action->GetDisplayName() : TEXT("None"));
+	//return FString::Printf(TEXT("Push Action: %s"), Action ? *Action->GetDisplayName() : TEXT("None"));
+	return FString::Printf(TEXT("Push Action: %s"), *GetNameSafe(Action));
 }
