More compilation issues fix #UE4

diff --git a/Engine/Source/Runtime/AIModule/Private/BehaviorTree/Tasks/BTTask_PushPawnAction.cpp b/Engine/Source/Runtime/AIModule/Private/BehaviorTree/Tasks/BTTask_PushPawnAction.cpp
--- a/Engine/Source/Runtime/AIModule/Private/BehaviorTree/Tasks/BTTask_PushPawnAction.cpp
+++ b/Engine/Source/Runtime/AIModule/Private/BehaviorTree/Tasks/BTTask_PushPawnAction.cpp
@@ -18,7 +18,7 @@ EBTNodeResult::Type UBTTask_PushPawnAction::ExecuteTask(UBehaviorTreeComponent&
 		return EBTNodeResult::Failed;
 	}
 
-	return PushAction(OwnerComp, ActionCopy);
+	return PushAction(OwnerComp, *ActionCopy);
 }
 
 FString UBTTask_PushPawnAction::GetStaticDescription() const
