disabled loop decorator for simple parallel node
#ue4

diff --git a/Engine/Source/Runtime/AIModule/Classes/BehaviorTree/BTCompositeNode.h b/Engine/Source/Runtime/AIModule/Classes/BehaviorTree/BTCompositeNode.h
--- a/Engine/Source/Runtime/AIModule/Classes/BehaviorTree/BTCompositeNode.h
+++ b/Engine/Source/Runtime/AIModule/Classes/BehaviorTree/BTCompositeNode.h
@@ -145,7 +145,7 @@ class AIMODULE_API UBTCompositeNode : public UBTNode
 	uint16 GetLastExecutionIndex() const;
 
 	/** set override for next child index */
-	void SetChildOverride(FBehaviorTreeSearchData& SearchData, int8 Index) const;
+	virtual void SetChildOverride(FBehaviorTreeSearchData& SearchData, int8 Index) const;
 
 	/** gathers description of all runtime parameters */
 	virtual void DescribeRuntimeValues(const UBehaviorTreeComponent& OwnerComp, uint8* NodeMemory, EBTDescriptionVerbosity::Type Verbosity, TArray<FString>& Values) const override;
diff --git a/Engine/Source/Runtime/AIModule/Classes/BehaviorTree/Composites/BTComposite_SimpleParallel.h b/Engine/Source/Runtime/AIModule/Classes/BehaviorTree/Composites/BTComposite_SimpleParallel.h
--- a/Engine/Source/Runtime/AIModule/Classes/BehaviorTree/Composites/BTComposite_SimpleParallel.h
+++ b/Engine/Source/Runtime/AIModule/Classes/BehaviorTree/Composites/BTComposite_SimpleParallel.h
@@ -59,6 +59,7 @@ class AIMODULE_API UBTComposite_SimpleParallel : public UBTCompositeNode
 	virtual void NotifyChildExecution(UBehaviorTreeComponent& OwnerComp, uint8* NodeMemory, int32 ChildIdx, EBTNodeResult::Type& NodeResult) const override;
 	virtual void NotifyNodeDeactivation(FBehaviorTreeSearchData& SearchData, EBTNodeResult::Type& NodeResult) const override;
 	virtual bool CanPushSubtree(UBehaviorTreeComponent& OwnerComp, uint8* NodeMemory, int32 ChildIdx) const override;
+	virtual void SetChildOverride(FBehaviorTreeSearchData& SearchData, int8 Index) const override;
 	virtual uint16 GetInstanceMemorySize() const override;
 	virtual FString GetStaticDescription() const override;
 	virtual void DescribeRuntimeValues(const UBehaviorTreeComponent& OwnerComp, uint8* NodeMemory, EBTDescriptionVerbosity::Type Verbosity, TArray<FString>& Values) const override;
diff --git a/Engine/Source/Runtime/AIModule/Private/BehaviorTree/Composites/BTComposite_SimpleParallel.cpp b/Engine/Source/Runtime/AIModule/Private/BehaviorTree/Composites/BTComposite_SimpleParallel.cpp
--- a/Engine/Source/Runtime/AIModule/Private/BehaviorTree/Composites/BTComposite_SimpleParallel.cpp
+++ b/Engine/Source/Runtime/AIModule/Private/BehaviorTree/Composites/BTComposite_SimpleParallel.cpp
@@ -134,6 +134,11 @@ bool UBTComposite_SimpleParallel::CanPushSubtree(UBehaviorTreeComponent& OwnerCo
 	return (ChildIdx != EBTParallelChild::MainTask);
 }
 
+void UBTComposite_SimpleParallel::SetChildOverride(FBehaviorTreeSearchData& SearchData, int8 Index) const
+{
+	// don't allow any external mechanic to change child execution order (e.g. loop decorators)
+}
+
 FString UBTComposite_SimpleParallel::DescribeFinishMode(EBTParallelMode::Type Mode)
 {
 	static FString FinishDesc[] = { TEXT("AbortBackground"), TEXT("WaitForBackground") };
