merged fix for infinite loop in parallel BT node causing editor to freeze
#ue4

diff --git a/Engine/Source/Runtime/AIModule/Classes/BehaviorTree/BehaviorTreeTypes.h b/Engine/Source/Runtime/AIModule/Classes/BehaviorTree/BehaviorTreeTypes.h
--- a/Engine/Source/Runtime/AIModule/Classes/BehaviorTree/BehaviorTreeTypes.h
+++ b/Engine/Source/Runtime/AIModule/Classes/BehaviorTree/BehaviorTreeTypes.h
@@ -51,6 +51,7 @@ namespace BTSpecialChild
 {
 	const int32 NotInitialized = -1;	// special value for child indices: needs to be initialized
 	const int32 ReturnToParent = -2;	// special value for child indices: return to parent node
+	const int32 PostponeSearch = -3;	// special value for child indices: search again in next tick
 }
 
 UENUM(BlueprintType)
diff --git a/Engine/Source/Runtime/AIModule/Classes/BehaviorTree/Composites/BTComposite_SimpleParallel.h b/Engine/Source/Runtime/AIModule/Classes/BehaviorTree/Composites/BTComposite_SimpleParallel.h
--- a/Engine/Source/Runtime/AIModule/Classes/BehaviorTree/Composites/BTComposite_SimpleParallel.h
+++ b/Engine/Source/Runtime/AIModule/Classes/BehaviorTree/Composites/BTComposite_SimpleParallel.h
@@ -30,6 +30,9 @@ namespace EBTParallelMode
 
 struct FBTParallelMemory : public FBTCompositeMemory
 {
+	/** last Id of search, detect infinite loops when there isn't any valid task in background tree */
+	int32 LastSearchId;
+
 	/** finish result of main task */
 	TEnumAsByte<EBTNodeResult::Type> MainTaskResult;
 
diff --git a/Engine/Source/Runtime/AIModule/Private/BehaviorTree/BTCompositeNode.cpp b/Engine/Source/Runtime/AIModule/Private/BehaviorTree/BTCompositeNode.cpp
--- a/Engine/Source/Runtime/AIModule/Private/BehaviorTree/BTCompositeNode.cpp
+++ b/Engine/Source/Runtime/AIModule/Private/BehaviorTree/BTCompositeNode.cpp
@@ -23,9 +23,8 @@ int32 UBTCompositeNode::FindChildToExecute(FBehaviorTreeSearchData& SearchData,
 
 	if (Children.Num())
 	{
-		for (int32 ChildIdx = GetNextChild(SearchData, NodeMemory->CurrentChild, LastResult);
-			Children.IsValidIndex(ChildIdx);
-			ChildIdx = GetNextChild(SearchData, ChildIdx, LastResult))
+		int32 ChildIdx = GetNextChild(SearchData, NodeMemory->CurrentChild, LastResult);
+		while (Children.IsValidIndex(ChildIdx))
 		{
 			// check decorators
 			if (DoDecoratorsAllowExecution(SearchData.OwnerComp, SearchData.OwnerComp.ActiveInstanceIdx, ChildIdx))
@@ -39,6 +38,12 @@ int32 UBTCompositeNode::FindChildToExecute(FBehaviorTreeSearchData& SearchData,
 				LastResult = EBTNodeResult::Failed;
 				NotifyDecoratorsOnFailedActivation(SearchData, ChildIdx, LastResult);
 			}
+
+			ChildIdx = GetNextChild(SearchData, ChildIdx, LastResult);
+			if (ChildIdx == BTSpecialChild::PostponeSearch)
+			{
+				RetIdx = ChildIdx;
+			}
 		}
 	}
 
diff --git a/Engine/Source/Runtime/AIModule/Private/BehaviorTree/BehaviorTreeComponent.cpp b/Engine/Source/Runtime/AIModule/Private/BehaviorTree/BehaviorTreeComponent.cpp
--- a/Engine/Source/Runtime/AIModule/Private/BehaviorTree/BehaviorTreeComponent.cpp
+++ b/Engine/Source/Runtime/AIModule/Private/BehaviorTree/BehaviorTreeComponent.cpp
@@ -885,6 +885,7 @@ void UBehaviorTreeComponent::ProcessExecutionRequest()
 
 	const uint16 PrevActiveInstanceIdx = ActiveInstanceIdx;
 	bool bIsSearchValid = true;
+	bool bResetRequest = true;
 
 	EBTNodeResult::Type NodeResult = ExecutionRequest.ContinueWithResult;
 	UBTTaskNode* NextTask = NULL;
@@ -984,7 +985,16 @@ void UBehaviorTreeComponent::ProcessExecutionRequest()
 					TestNode->OnChildDeactivation(SearchData, *ChildNode, NodeResult);
 				}
 			}
-			else
+			else if (ChildBranchIdx == BTSpecialChild::PostponeSearch)
+			{
+				// break out of current search loop
+				TestNode = NULL;
+				bIsSearchValid = false;
+
+				// keep current request for next tick
+				bResetRequest = false;
+			}
+			else if (TestNode->Children.IsValidIndex(ChildBranchIdx))
 			{
 				// was new task found?
 				NextTask = TestNode->Children[ChildBranchIdx].ChildTask;
@@ -1022,7 +1032,15 @@ void UBehaviorTreeComponent::ProcessExecutionRequest()
 		// finish timer scope
 	}
 
-	ExecutionRequest = FBTNodeExecutionInfo();
+	if (bResetRequest)
+	{
+		ExecutionRequest = FBTNodeExecutionInfo();
+	}
+	else
+	{
+		ScheduleExecutionUpdate();
+	}
+
 	if (bIsSearchValid)
 	{
 		// abort task if needed
diff --git a/Engine/Source/Runtime/AIModule/Private/BehaviorTree/Composites/BTComposite_SimpleParallel.cpp b/Engine/Source/Runtime/AIModule/Private/BehaviorTree/Composites/BTComposite_SimpleParallel.cpp
--- a/Engine/Source/Runtime/AIModule/Private/BehaviorTree/Composites/BTComposite_SimpleParallel.cpp
+++ b/Engine/Source/Runtime/AIModule/Private/BehaviorTree/Composites/BTComposite_SimpleParallel.cpp
@@ -27,8 +27,15 @@ int32 UBTComposite_SimpleParallel::GetNextChildHandler(FBehaviorTreeSearchData&
 		// unless search is from abort from background tree - return to parent (and break from search when node gets deactivated)
 		NextChildIdx = EBTParallelChild::BackgroundTree;
 		MyMemory->bForceBackgroundTree = false;
+
+		if ((PrevChild == EBTParallelChild::BackgroundTree) && (MyMemory->LastSearchId == SearchData.SearchId))
+		{
+			// going to background tree again within the same search - possible infinite loop
+			NextChildIdx = BTSpecialChild::PostponeSearch;
+		}
 	}
 
+	MyMemory->LastSearchId = SearchData.SearchId;
 	return NextChildIdx;
 }
 
@@ -37,6 +44,8 @@ void UBTComposite_SimpleParallel::NotifyChildExecution(UBehaviorTreeComponent& O
 	FBTParallelMemory* MyMemory = (FBTParallelMemory*)NodeMemory;
 	if (ChildIdx == EBTParallelChild::MainTask)
 	{
+		MyMemory->MainTaskResult = NodeResult;
+
 		if (NodeResult == EBTNodeResult::InProgress)
 		{
 			EBTTaskStatus::Type Status = OwnerComp.GetTaskStatus(Children[EBTParallelChild::MainTask].ChildTask);
@@ -53,7 +62,6 @@ void UBTComposite_SimpleParallel::NotifyChildExecution(UBehaviorTreeComponent& O
 		}
 		else if (MyMemory->bMainTaskIsActive)
 		{
-			MyMemory->MainTaskResult = NodeResult;
 			MyMemory->bMainTaskIsActive = false;
 
 			const int32 MyInstanceIdx = OwnerComp.FindInstanceContainingNode(this);
