Added versions of UEnvQueryTest::GetItemActor (and other similar) that take FEnvQueryInstance::ItemIterator reference #UE4

The reason behind it being it was possible to get wrong results by not using * operator on the iterator, resulting in it being cast to bool first, then to int, which in most use cases resulted in simply getting results for actor at index 1

diff --git a/Engine/Source/Runtime/AIModule/Classes/EnvironmentQuery/EnvQueryTest.h b/Engine/Source/Runtime/AIModule/Classes/EnvironmentQuery/EnvQueryTest.h
--- a/Engine/Source/Runtime/AIModule/Classes/EnvironmentQuery/EnvQueryTest.h
+++ b/Engine/Source/Runtime/AIModule/Classes/EnvironmentQuery/EnvQueryTest.h
@@ -5,7 +5,6 @@
 #include "EnvironmentQuery/EnvQueryTypes.h"
 #include "EnvQueryTest.generated.h"
 
-struct FEnvQueryInstance;
 class UEnvQueryItemType;
 class UEnvQueryContext;
 #if WITH_EDITOR
@@ -41,7 +40,7 @@ class AIMODULE_API UEnvQueryTest : public UObject
 	UPROPERTY(EditDefaultsOnly, Category=Filter)
 	TEnumAsByte<EEnvTestFilterType::Type> FilterType;
 
-	/** Boolean value that must be matched in order for scoring to occur or filtering test to pass. */
+	/** Desired boolean value of the test for scoring to occur or filtering test to pass. */
 	UPROPERTY(EditDefaultsOnly, Category=Filter)
 	FAIDataProviderBoolValue BoolValue;
 
@@ -135,11 +134,29 @@ public:
 	/** helper: get location of item */
 	FVector GetItemLocation(FEnvQueryInstance& QueryInstance, int32 ItemIndex) const;
 
+	/** helper: get location of item */
+	FORCEINLINE FVector GetItemLocation(FEnvQueryInstance& QueryInstance, const FEnvQueryInstance::ItemIterator& Iterator) const
+	{
+		return GetItemLocation(QueryInstance, *Iterator);
+	}
+
 	/** helper: get location of item */
 	FRotator GetItemRotation(FEnvQueryInstance& QueryInstance, int32 ItemIndex) const;
 
+	/** helper: get location of item */
+	FORCEINLINE FRotator GetItemRotation(FEnvQueryInstance& QueryInstance, const FEnvQueryInstance::ItemIterator& Iterator) const
+	{
+		return GetItemRotation(QueryInstance, *Iterator);
+	}
+
 	/** helper: get actor from item */
 	AActor* GetItemActor(FEnvQueryInstance& QueryInstance, int32 ItemIndex) const;
+		
+	/** helper: get actor from item */
+	FORCEINLINE AActor* GetItemActor(FEnvQueryInstance& QueryInstance, const FEnvQueryInstance::ItemIterator& Iterator) const
+	{
+		return GetItemActor(QueryInstance, *Iterator);
+	}
 
 	/** normalize scores in range */
 	void NormalizeItemScores(FEnvQueryInstance& QueryInstance);
diff --git a/Engine/Source/Runtime/AIModule/Private/EnvironmentQuery/Tests/EnvQueryTest_GameplayTags.cpp b/Engine/Source/Runtime/AIModule/Private/EnvironmentQuery/Tests/EnvQueryTest_GameplayTags.cpp
--- a/Engine/Source/Runtime/AIModule/Private/EnvironmentQuery/Tests/EnvQueryTest_GameplayTags.cpp
+++ b/Engine/Source/Runtime/AIModule/Private/EnvironmentQuery/Tests/EnvQueryTest_GameplayTags.cpp
@@ -58,7 +58,7 @@ void UEnvQueryTest_GameplayTags::RunTest(FEnvQueryInstance& QueryInstance) const
 	// loop through all items
 	for (FEnvQueryInstance::ItemIterator It(this, QueryInstance); It; ++It)
 	{
-		AActor* ItemActor = GetItemActor(QueryInstance, It);
+		AActor* ItemActor = GetItemActor(QueryInstance, *It);
 		IGameplayTagAssetInterface* GameplayTagAssetInterface = Cast<IGameplayTagAssetInterface>(ItemActor);
 		if (GameplayTagAssetInterface != NULL)
 		{
