loop support for simple parallel behavior tree node
#ue4

diff --git a/Engine/Source/Runtime/AIModule/Classes/BehaviorTree/BehaviorTreeComponent.h b/Engine/Source/Runtime/AIModule/Classes/BehaviorTree/BehaviorTreeComponent.h
--- a/Engine/Source/Runtime/AIModule/Classes/BehaviorTree/BehaviorTreeComponent.h
+++ b/Engine/Source/Runtime/AIModule/Classes/BehaviorTree/BehaviorTreeComponent.h
@@ -336,5 +336,5 @@ FORCEINLINE uint16 UBehaviorTreeComponent::GetActiveInstanceIdx() const
 
 FORCEINLINE bool UBehaviorTreeComponent::IsRestartPending() const
 {
-	return ExecutionRequest.ExecuteNode && ExecutionRequest.bIsRestart;
+	return ExecutionRequest.ExecuteNode && !ExecutionRequest.bTryNextChild;
 }
diff --git a/Engine/Source/Runtime/AIModule/Classes/BehaviorTree/Composites/BTComposite_SimpleParallel.h b/Engine/Source/Runtime/AIModule/Classes/BehaviorTree/Composites/BTComposite_SimpleParallel.h
--- a/Engine/Source/Runtime/AIModule/Classes/BehaviorTree/Composites/BTComposite_SimpleParallel.h
+++ b/Engine/Source/Runtime/AIModule/Classes/BehaviorTree/Composites/BTComposite_SimpleParallel.h
@@ -13,7 +13,6 @@ namespace EBTParallelChild
 	{
 		MainTask,
 		BackgroundTree,
-		SearchIndicator,
 	};
 }
 
@@ -41,6 +40,9 @@ struct FBTParallelMemory : public FBTCompositeMemory
 
 	/** try running background tree task even if main task has finished */
 	uint8 bForceBackgroundTree : 1;
+
+	/** set when main task needs to be repeated */
+	uint8 bRepeatMainTask : 1;
 };
 
 /**
diff --git a/Engine/Source/Runtime/AIModule/Private/BehaviorTree/Composites/BTComposite_SimpleParallel.cpp b/Engine/Source/Runtime/AIModule/Private/BehaviorTree/Composites/BTComposite_SimpleParallel.cpp
--- a/Engine/Source/Runtime/AIModule/Private/BehaviorTree/Composites/BTComposite_SimpleParallel.cpp
+++ b/Engine/Source/Runtime/AIModule/Private/BehaviorTree/Composites/BTComposite_SimpleParallel.cpp
@@ -20,6 +20,7 @@ int32 UBTComposite_SimpleParallel::GetNextChildHandler(FBehaviorTreeSearchData&
 	{
 		NextChildIdx = EBTParallelChild::MainTask;
 		MyMemory->MainTaskResult = EBTNodeResult::Failed;
+		MyMemory->bRepeatMainTask = false;
 	}
 	else if ((MyMemory->bMainTaskIsActive || MyMemory->bForceBackgroundTree) && !SearchData.OwnerComp.IsRestartPending())
 	{
@@ -27,12 +28,18 @@ int32 UBTComposite_SimpleParallel::GetNextChildHandler(FBehaviorTreeSearchData&
 		// unless search is from abort from background tree - return to parent (and break from search when node gets deactivated)
 		NextChildIdx = EBTParallelChild::BackgroundTree;
 		MyMemory->bForceBackgroundTree = false;
+	}
+	else if (MyMemory->bRepeatMainTask)
+	{
+		UE_VLOG(SearchData.OwnerComp.GetOwner(), LogBehaviorTree, Verbose, TEXT("Repeating main task of %s"), *UBehaviorTreeTypes::DescribeNodeHelper(this));
+		NextChildIdx = EBTParallelChild::MainTask;
+		MyMemory->bRepeatMainTask = false;
+	}
 
-		if ((PrevChild == EBTParallelChild::BackgroundTree) && (MyMemory->LastSearchId == SearchData.SearchId))
-		{
-			// going to background tree again within the same search - possible infinite loop
-			NextChildIdx = BTSpecialChild::PostponeSearch;
-		}
+	if ((PrevChild == NextChildIdx) && (MyMemory->LastSearchId == SearchData.SearchId))
+	{
+		// retrying the same branch again within the same search - possible infinite loop
+		NextChildIdx = BTSpecialChild::PostponeSearch;
 	}
 
 	MyMemory->LastSearchId = SearchData.SearchId;
@@ -67,7 +74,7 @@ void UBTComposite_SimpleParallel::NotifyChildExecution(UBehaviorTreeComponent& O
 			const int32 MyInstanceIdx = OwnerComp.FindInstanceContainingNode(this);
 
 			OwnerComp.UnregisterParallelTask(Children[EBTParallelChild::MainTask].ChildTask, MyInstanceIdx);
-			if (NodeResult != EBTNodeResult::Aborted)
+			if (NodeResult != EBTNodeResult::Aborted && !MyMemory->bRepeatMainTask)
 			{
 				// check if subtree should be aborted when task finished with success/failed result
 				if (FinishMode == EBTParallelMode::AbortBackground)
@@ -144,7 +151,15 @@ bool UBTComposite_SimpleParallel::CanPushSubtree(UBehaviorTreeComponent& OwnerCo
 
 void UBTComposite_SimpleParallel::SetChildOverride(FBehaviorTreeSearchData& SearchData, int8 Index) const
 {
-	// don't allow any external mechanic to change child execution order (e.g. loop decorators)
+	// store information about repeating main task and use it when possible
+	// ignore changes to other child indices = background tree, as it's already being auto repeated
+	if (Index == EBTParallelChild::MainTask)
+	{
+		FBTParallelMemory* MyMemory = GetNodeMemory<FBTParallelMemory>(SearchData);
+		MyMemory->bRepeatMainTask = true;
+
+		UE_VLOG(SearchData.OwnerComp.GetOwner(), LogBehaviorTree, Log, TEXT("Main task of %s will be repeated."), *UBehaviorTreeTypes::DescribeNodeHelper(this));
+	}
 }
 
 FString UBTComposite_SimpleParallel::DescribeFinishMode(EBTParallelMode::Type Mode)
diff --git a/Engine/Source/Runtime/AIModule/Private/BehaviorTree/Decorators/BTDecorator_Loop.cpp b/Engine/Source/Runtime/AIModule/Private/BehaviorTree/Decorators/BTDecorator_Loop.cpp
--- a/Engine/Source/Runtime/AIModule/Private/BehaviorTree/Decorators/BTDecorator_Loop.cpp
+++ b/Engine/Source/Runtime/AIModule/Private/BehaviorTree/Decorators/BTDecorator_Loop.cpp
@@ -2,6 +2,7 @@
 
 #include "AIModulePrivate.h"
 #include "BehaviorTree/BTCompositeNode.h"
+#include "BehaviorTree/Composites/BTComposite_SimpleParallel.h"
 #include "BehaviorTree/Decorators/BTDecorator_Loop.h"
 
 UBTDecorator_Loop::UBTDecorator_Loop(const FObjectInitializer& ObjectInitializer) : Super(ObjectInitializer)
@@ -19,7 +20,10 @@ void UBTDecorator_Loop::OnNodeActivation(FBehaviorTreeSearchData& SearchData)
 {
 	FBTLoopDecoratorMemory* DecoratorMemory = GetNodeMemory<FBTLoopDecoratorMemory>(SearchData);
 	FBTCompositeMemory* ParentMemory = GetParentNode()->GetNodeMemory<FBTCompositeMemory>(SearchData);
-	if (ParentMemory->CurrentChild != ChildIndex)
+	const bool bIsSpecialNode = GetParentNode()->IsA(UBTComposite_SimpleParallel::StaticClass());
+
+	if ((bIsSpecialNode && ParentMemory->CurrentChild == BTSpecialChild::NotInitialized) ||
+		(!bIsSpecialNode && ParentMemory->CurrentChild != ChildIndex))
 	{
 		// initialize counter if it's first activation
 		DecoratorMemory->RemainingExecutions = NumLoops;
