#UE4 - Added Strength parameter to FAITeamStimulusEvent

diff --git a/Engine/Source/Runtime/AIModule/Classes/Perception/AISense_Team.h b/Engine/Source/Runtime/AIModule/Classes/Perception/AISense_Team.h
--- a/Engine/Source/Runtime/AIModule/Classes/Perception/AISense_Team.h
+++ b/Engine/Source/Runtime/AIModule/Classes/Perception/AISense_Team.h
@@ -22,6 +22,7 @@ public:
 	float RangeSq;
 	float InformationAge;
 	FGenericTeamId TeamIdentifier;
+	float Strength;
 private:
 	UPROPERTY()
 	AActor* Broadcaster;
@@ -30,7 +31,7 @@ public:
 	AActor* Enemy;
 		
 	FAITeamStimulusEvent(){}	
-	FAITeamStimulusEvent(AActor* InBroadcaster, AActor* InEnemy, const FVector& InLastKnowLocation, float EventRange, float PassedInfoAge = 0.f);
+	FAITeamStimulusEvent(AActor* InBroadcaster, AActor* InEnemy, const FVector& InLastKnowLocation, float EventRange, float PassedInfoAge = 0.f, float InStrength = 1.f);
 
 	FORCEINLINE void CacheBroadcastLocation()
 	{
diff --git a/Engine/Source/Runtime/AIModule/Private/Perception/AISense_Team.cpp b/Engine/Source/Runtime/AIModule/Private/Perception/AISense_Team.cpp
--- a/Engine/Source/Runtime/AIModule/Private/Perception/AISense_Team.cpp
+++ b/Engine/Source/Runtime/AIModule/Private/Perception/AISense_Team.cpp
@@ -6,8 +6,8 @@
 //----------------------------------------------------------------------//
 // 
 //----------------------------------------------------------------------//
-FAITeamStimulusEvent::FAITeamStimulusEvent(AActor* InBroadcaster, AActor* InEnemy, const FVector& InLastKnowLocation, float EventRange, float PassedInfoAge)
-	: LastKnowLocation(InLastKnowLocation), RangeSq(FMath::Square(EventRange)), InformationAge(PassedInfoAge), Broadcaster(InBroadcaster), Enemy(InEnemy)
+FAITeamStimulusEvent::FAITeamStimulusEvent(AActor* InBroadcaster, AActor* InEnemy, const FVector& InLastKnowLocation, float EventRange, float PassedInfoAge, float InStrength)
+	: LastKnowLocation(InLastKnowLocation), RangeSq(FMath::Square(EventRange)), InformationAge(PassedInfoAge), Strength(InStrength), Broadcaster(InBroadcaster), Enemy(InEnemy)
 {
 	CacheBroadcastLocation();
 
@@ -49,7 +49,7 @@ float UAISense_Team::Update()
 				continue;
 			}
 			
-			Listener.RegisterStimulus(Event.Enemy, FAIStimulus(*this, 1.f, Event.LastKnowLocation, Event.GetBroadcastLocation(), FAIStimulus::SensingSucceeded).SetStimulusAge(Event.InformationAge));
+			Listener.RegisterStimulus(Event.Enemy, FAIStimulus(*this, Event.Strength, Event.LastKnowLocation, Event.GetBroadcastLocation(), FAIStimulus::SensingSucceeded).SetStimulusAge(Event.InformationAge));
 		}
 	}
 
