removed initial Tick event from behavior tree service - this should be implemented using OnSearchStart,
fixed task search using outdated info when more important request was accepted
#ue4

diff --git a/Engine/Source/Runtime/AIModule/Classes/BehaviorTree/BTService.h b/Engine/Source/Runtime/AIModule/Classes/BehaviorTree/BTService.h
--- a/Engine/Source/Runtime/AIModule/Classes/BehaviorTree/BTService.h
+++ b/Engine/Source/Runtime/AIModule/Classes/BehaviorTree/BTService.h
@@ -9,14 +9,15 @@
 /** 
  *	Behavior Tree service nodes is designed to perform "background" tasks that update AI's knowledge.
  *
- *  First tick can be called when parent node is activated during task search (default).
- *  It's meant to be used by services that periodically check condition, which can affect
- *  execution flow in underlying branch. Check OnSearchStart() for details.
- *  Keep in mind that any checks performed there have to be instantaneous!
+ *  Services are being executed when underlying branch of behavior tree becomes active,
+ *  but unlike tasks they don't return any results and can't directly affect execution flow.
+ *
+ *  Usually they perform periodical checks (see TickNode) and often store results in blackboard.
+ *  If any decorator node below requires results of check beforehand, use OnSearchStart function.
+ *   Keep in mind that any checks performed there have to be instantaneous!
  *	
- *	BT service should never modify the parent BT's state or flow. From BT point of view
- *  services are fire-and-forget calls and no service-state-checking is being performed. 
- *  If you need this data create a BTDecorator to check this information.
+ *  Other typical use case is creating a marker when specific branch is being executed
+ *  (see OnBecomeRelevant, OnCeaseRelevant), by setting a flag in blackboard.
  *
  *  Because some of them can be instanced for specific AI, following virtual functions are not marked as const:
  *   - OnBecomeRelevant (from UBTAuxiliaryNode)
@@ -27,7 +28,6 @@
  *  If your node is not being instanced (default behavior), DO NOT change any properties of object within those functions!
  *  Template nodes are shared across all behavior tree components using the same tree asset and must store
  *  their runtime properties in provided NodeMemory block (allocation size determined by GetInstanceMemorySize() )
- *
  */
  
 UCLASS(Abstract)
@@ -51,10 +51,14 @@ protected:
 	UPROPERTY(Category=Service, EditAnywhere, meta=(ClampMin="0.001"))
 	float Interval;
 
-	/** allows adding random time to service's Interval */
+	/** adds random range to service's Interval */
 	UPROPERTY(Category=Service, EditAnywhere, meta=(ClampMin="0.0"))
 	float RandomDeviation;
 
+	/** call Tick event when task search enters this node (SearchStart will be called as well) */
+	UPROPERTY(Category = Service, EditAnywhere, AdvancedDisplay)
+	uint32 bCallTickOnSearchStart : 1;
+
 	/** if set, service will be notified about search entering underlying branch */
 	uint32 bNotifyOnSearch : 1;
 
@@ -62,10 +66,14 @@ protected:
 	 * this function should be considered as const (don't modify state of object) if node is not instanced! */
 	virtual void TickNode(UBehaviorTreeComponent& OwnerComp, uint8* NodeMemory, float DeltaSeconds) override;
 
-	/** called when search enters underlying branch */
+	/** called when search enters underlying branch
+	 * this function should be considered as const (don't modify state of object) if node is not instanced! */
 	virtual void OnSearchStart(FBehaviorTreeSearchData& SearchData);
 
 #if WITH_EDITOR
 	virtual FName GetNodeIconName() const override;
 #endif // WITH_EDITOR
+
+	/** set next tick time */
+	void ScheduleNextTick(uint8* NodeMemory);
 };
diff --git a/Engine/Source/Runtime/AIModule/Classes/BehaviorTree/BehaviorTreeTypes.h b/Engine/Source/Runtime/AIModule/Classes/BehaviorTree/BehaviorTreeTypes.h
--- a/Engine/Source/Runtime/AIModule/Classes/BehaviorTree/BehaviorTreeTypes.h
+++ b/Engine/Source/Runtime/AIModule/Classes/BehaviorTree/BehaviorTreeTypes.h
@@ -51,7 +51,6 @@ namespace BTSpecialChild
 {
 	const int32 NotInitialized = -1;	// special value for child indices: needs to be initialized
 	const int32 ReturnToParent = -2;	// special value for child indices: return to parent node
-	const int32 PostponeSearch = -3;	// special value for child indices: search again in next tick
 }
 
 UENUM(BlueprintType)
@@ -380,13 +379,20 @@ struct FBehaviorTreeSearchData
 	/** search unique number */
 	int32 SearchId;
 
+	/** if set, current search will be restarted in next tick */
+	uint32 bPostponeSearch : 1;
+
+	/** set when task search is in progress */
+	uint32 bSearchInProgress : 1;
+
 	/** adds update info to PendingUpdates array, removing all previous updates for this node */
 	void AddUniqueUpdate(const FBehaviorTreeSearchUpdate& UpdateInfo);
 
 	/** assign unique Id number */
 	void AssignSearchId();
 
-	FBehaviorTreeSearchData(UBehaviorTreeComponent& InOwnerComp) : OwnerComp(InOwnerComp)
+	FBehaviorTreeSearchData(UBehaviorTreeComponent& InOwnerComp) 
+		: OwnerComp(InOwnerComp), bPostponeSearch(false), bSearchInProgress(false)
 	{}
 
 private:
diff --git a/Engine/Source/Runtime/AIModule/Classes/BehaviorTree/Services/BTService_BlueprintBase.h b/Engine/Source/Runtime/AIModule/Classes/BehaviorTree/Services/BTService_BlueprintBase.h
--- a/Engine/Source/Runtime/AIModule/Classes/BehaviorTree/Services/BTService_BlueprintBase.h
+++ b/Engine/Source/Runtime/AIModule/Classes/BehaviorTree/Services/BTService_BlueprintBase.h
@@ -49,6 +49,10 @@ protected:
 	UPROPERTY(EditInstanceOnly, Category=Description)
 	uint32 bShowPropertyDetails : 1;
 
+	/** show detailed information about implemented events */
+	UPROPERTY(EditInstanceOnly, Category = Description)
+	uint32 bShowEventDetails : 1;
+
 	/** set if ReceiveTick is implemented by blueprint */
 	uint32 ReceiveTickImplementations : 2;
 	
@@ -72,7 +76,7 @@ protected:
 	UFUNCTION(BlueprintImplementableEvent)
 	virtual void ReceiveTick(AActor* OwnerActor, float DeltaSeconds);
 
-	/** task search enters branch of tree, initial ReceiveTick won't be called if you implement this event!
+	/** task search enters branch of tree
 	 *	@Note that if both generic and AI event versions are implemented only the more
 	 *	suitable one will be called, meaning the AI version if called for AI, generic one otherwise */
 	UFUNCTION(BlueprintImplementableEvent)
diff --git a/Engine/Source/Runtime/AIModule/Private/BehaviorTree/BTCompositeNode.cpp b/Engine/Source/Runtime/AIModule/Private/BehaviorTree/BTCompositeNode.cpp
--- a/Engine/Source/Runtime/AIModule/Private/BehaviorTree/BTCompositeNode.cpp
+++ b/Engine/Source/Runtime/AIModule/Private/BehaviorTree/BTCompositeNode.cpp
@@ -24,7 +24,7 @@ int32 UBTCompositeNode::FindChildToExecute(FBehaviorTreeSearchData& SearchData,
 	if (Children.Num())
 	{
 		int32 ChildIdx = GetNextChild(SearchData, NodeMemory->CurrentChild, LastResult);
-		while (Children.IsValidIndex(ChildIdx))
+		while (Children.IsValidIndex(ChildIdx) && !SearchData.bPostponeSearch)
 		{
 			// check decorators
 			if (DoDecoratorsAllowExecution(SearchData.OwnerComp, SearchData.OwnerComp.ActiveInstanceIdx, ChildIdx))
@@ -40,10 +40,6 @@ int32 UBTCompositeNode::FindChildToExecute(FBehaviorTreeSearchData& SearchData,
 			}
 
 			ChildIdx = GetNextChild(SearchData, ChildIdx, LastResult);
-			if (ChildIdx == BTSpecialChild::PostponeSearch)
-			{
-				RetIdx = ChildIdx;
-			}
 		}
 	}
 
diff --git a/Engine/Source/Runtime/AIModule/Private/BehaviorTree/BTService.cpp b/Engine/Source/Runtime/AIModule/Private/BehaviorTree/BTService.cpp
--- a/Engine/Source/Runtime/AIModule/Private/BehaviorTree/BTService.cpp
+++ b/Engine/Source/Runtime/AIModule/Private/BehaviorTree/BTService.cpp
@@ -8,6 +8,7 @@ UBTService::UBTService(const FObjectInitializer& ObjectInitializer) : Super(Obje
 	bNotifyTick = true;
 	bNotifyOnSearch = true;
 	bTickIntervals = true;
+	bCallTickOnSearchStart = false;
 
 	Interval = 0.5f;
 	RandomDeviation = 0.1f;
@@ -15,24 +16,38 @@ UBTService::UBTService(const FObjectInitializer& ObjectInitializer) : Super(Obje
 
 void UBTService::TickNode(UBehaviorTreeComponent& OwnerComp, uint8* NodeMemory, float DeltaSeconds)
 {
-	const float NextTickTime = FMath::FRandRange(FMath::Max(0.0f, Interval - RandomDeviation), (Interval + RandomDeviation));
-	SetNextTickTime(NodeMemory, NextTickTime);
+	ScheduleNextTick(NodeMemory);
 }
 
 void UBTService::OnSearchStart(FBehaviorTreeSearchData& SearchData)
 {
-	uint8* NodeMemory = GetNodeMemory<uint8>(SearchData);
-	TickNode(SearchData.OwnerComp, NodeMemory, 0.0f);
+	// empty in base class
 }
 
 void UBTService::NotifyParentActivation(FBehaviorTreeSearchData& SearchData)
 {
-	if (bNotifyOnSearch)
+	if (bNotifyOnSearch || bNotifyTick)
 	{
 		UBTNode* NodeOb = bCreateNodeInstance ? GetNodeInstance(SearchData) : this;
 		if (NodeOb)
 		{
-			((UBTService*)NodeOb)->OnSearchStart(SearchData);
+			UBTService* ServiceNodeOb = (UBTService*)NodeOb;
+			uint8* NodeMemory = GetNodeMemory<uint8>(SearchData);
+
+			if (bNotifyTick)
+			{
+				ServiceNodeOb->ScheduleNextTick(NodeMemory);
+			}
+
+			if (bNotifyOnSearch)
+			{
+				ServiceNodeOb->OnSearchStart(SearchData);
+			}
+
+			if (bCallTickOnSearchStart)
+			{
+				ServiceNodeOb->TickNode(SearchData.OwnerComp, NodeMemory, 0.0f);
+			}
 		}
 	}
 }
@@ -64,3 +79,9 @@ FName UBTService::GetNodeIconName() const
 }
 
 #endif // WITH_EDITOR
+
+void UBTService::ScheduleNextTick(uint8* NodeMemory)
+{
+	const float NextTickTime = FMath::FRandRange(FMath::Max(0.0f, Interval - RandomDeviation), (Interval + RandomDeviation));
+	SetNextTickTime(NodeMemory, NextTickTime);
+}
diff --git a/Engine/Source/Runtime/AIModule/Private/BehaviorTree/BehaviorTreeComponent.cpp b/Engine/Source/Runtime/AIModule/Private/BehaviorTree/BehaviorTreeComponent.cpp
--- a/Engine/Source/Runtime/AIModule/Private/BehaviorTree/BehaviorTreeComponent.cpp
+++ b/Engine/Source/Runtime/AIModule/Private/BehaviorTree/BehaviorTreeComponent.cpp
@@ -687,6 +687,14 @@ void UBehaviorTreeComponent::RequestExecution(UBTCompositeNode* RequestedOn, int
 	ExecutionRequest.bTryNextChild = !bSwitchToHigherPriority;
 	ExecutionRequest.bIsRestart = (RequestedBy != GetActiveNode());
 	PendingExecution = FBTPendingExecutionInfo();
+	
+	// break out of current search if new request is more important than currently processed one
+	// no point in starting new task just to abandon it in next tick
+	if (SearchData.bSearchInProgress)
+	{
+		UE_VLOG(GetOwner(), LogBehaviorTree, Log, TEXT("> aborting current task search!"));
+		SearchData.bPostponeSearch = true;
+	}
 
 	ScheduleExecutionUpdate();
 }
@@ -812,7 +820,11 @@ void UBehaviorTreeComponent::ApplyDiscardedSearch()
 		}
 	}
 
+	// apply new observing decorators
 	ApplySearchUpdates(UpdateList, 0);
+
+	// remove everything else
+	SearchData.PendingUpdates.Reset();
 }
 
 void UBehaviorTreeComponent::TickComponent(float DeltaTime, enum ELevelTick TickType, FActorComponentTickFunction *ThisTickFunction)
@@ -885,7 +897,6 @@ void UBehaviorTreeComponent::ProcessExecutionRequest()
 
 	const uint16 PrevActiveInstanceIdx = ActiveInstanceIdx;
 	bool bIsSearchValid = true;
-	bool bResetRequest = true;
 
 	EBTNodeResult::Type NodeResult = ExecutionRequest.ContinueWithResult;
 	UBTTaskNode* NextTask = NULL;
@@ -909,6 +920,8 @@ void UBehaviorTreeComponent::ProcessExecutionRequest()
 		FBehaviorTreeInstance& ActiveInstance = InstanceStack[ActiveInstanceIdx];
 		UBTCompositeNode* TestNode = ExecutionRequest.ExecuteNode;
 		SearchData.AssignSearchId();
+		SearchData.bPostponeSearch = false;
+		SearchData.bSearchInProgress = true;
 
 		// activate root node if needed (can't be handled by parent composite...)
 		if (ActiveInstance.ActiveNode == NULL)
@@ -955,7 +968,13 @@ void UBehaviorTreeComponent::ProcessExecutionRequest()
 			const int32 ChildBranchIdx = TestNode->FindChildToExecute(SearchData, NodeResult);
 			UBTNode* StoreNode = TestNode;
 
-			if (ChildBranchIdx == BTSpecialChild::ReturnToParent)
+			if (SearchData.bPostponeSearch)
+			{
+				// break out of current search loop
+				TestNode = NULL;
+				bIsSearchValid = false;
+			}
+			else if (ChildBranchIdx == BTSpecialChild::ReturnToParent)
 			{
 				UBTCompositeNode* ChildNode = TestNode;
 				TestNode = TestNode->GetParentNode();
@@ -985,15 +1004,6 @@ void UBehaviorTreeComponent::ProcessExecutionRequest()
 					TestNode->OnChildDeactivation(SearchData, *ChildNode, NodeResult);
 				}
 			}
-			else if (ChildBranchIdx == BTSpecialChild::PostponeSearch)
-			{
-				// break out of current search loop
-				TestNode = NULL;
-				bIsSearchValid = false;
-
-				// keep current request for next tick
-				bResetRequest = false;
-			}
 			else if (TestNode->Children.IsValidIndex(ChildBranchIdx))
 			{
 				// was new task found?
@@ -1029,16 +1039,17 @@ void UBehaviorTreeComponent::ProcessExecutionRequest()
 			CopyInstanceMemoryFromPersistent();
 		}
 
+		SearchData.bSearchInProgress = false;
 		// finish timer scope
 	}
 
-	if (bResetRequest)
+	if (SearchData.bPostponeSearch)
 	{
-		ExecutionRequest = FBTNodeExecutionInfo();
+		ScheduleExecutionUpdate();
 	}
 	else
 	{
-		ScheduleExecutionUpdate();
+		ExecutionRequest = FBTNodeExecutionInfo();
 	}
 
 	if (bIsSearchValid)
diff --git a/Engine/Source/Runtime/AIModule/Private/BehaviorTree/Composites/BTComposite_SimpleParallel.cpp b/Engine/Source/Runtime/AIModule/Private/BehaviorTree/Composites/BTComposite_SimpleParallel.cpp
--- a/Engine/Source/Runtime/AIModule/Private/BehaviorTree/Composites/BTComposite_SimpleParallel.cpp
+++ b/Engine/Source/Runtime/AIModule/Private/BehaviorTree/Composites/BTComposite_SimpleParallel.cpp
@@ -39,7 +39,7 @@ int32 UBTComposite_SimpleParallel::GetNextChildHandler(FBehaviorTreeSearchData&
 	if ((PrevChild == NextChildIdx) && (MyMemory->LastSearchId == SearchData.SearchId))
 	{
 		// retrying the same branch again within the same search - possible infinite loop
-		NextChildIdx = BTSpecialChild::PostponeSearch;
+		SearchData.bPostponeSearch = true;
 	}
 
 	MyMemory->LastSearchId = SearchData.SearchId;
diff --git a/Engine/Source/Runtime/AIModule/Private/BehaviorTree/Services/BTService_BlueprintBase.cpp b/Engine/Source/Runtime/AIModule/Private/BehaviorTree/Services/BTService_BlueprintBase.cpp
--- a/Engine/Source/Runtime/AIModule/Private/BehaviorTree/Services/BTService_BlueprintBase.cpp
+++ b/Engine/Source/Runtime/AIModule/Private/BehaviorTree/Services/BTService_BlueprintBase.cpp
@@ -14,7 +14,7 @@ UBTService_BlueprintBase::UBTService_BlueprintBase(const FObjectInitializer& Obj
 
 	bNotifyBecomeRelevant = ReceiveActivationImplementations != 0;
 	bNotifyCeaseRelevant = bNotifyBecomeRelevant;
-	bNotifyOnSearch = ReceiveTickImplementations != 0 || ReceiveSearchStartImplementations != 0;
+	bNotifyOnSearch = ReceiveSearchStartImplementations != 0;
 	bNotifyTick = ReceiveTickImplementations != 0;
 	bShowPropertyDetails = true;
 
@@ -85,7 +85,7 @@ void UBTService_BlueprintBase::OnCeaseRelevant(UBehaviorTreeComponent& OwnerComp
 
 void UBTService_BlueprintBase::OnSearchStart(FBehaviorTreeSearchData& SearchData)
 {
-	// skip flag, will be handled by bNotifyOnSearch
+	Super::OnSearchStart(SearchData);
 
 	if (ReceiveSearchStartImplementations != 0)
 	{
@@ -97,14 +97,6 @@ void UBTService_BlueprintBase::OnSearchStart(FBehaviorTreeSearchData& SearchData
 		{
 			ReceiveSearchStart(ActorOwner);
 		}
-
-		const float NextTickTime = FMath::FRandRange(FMath::Max(0.0f, Interval - RandomDeviation), (Interval + RandomDeviation));
-		uint8* NodeMemory = GetNodeMemory<uint8>(SearchData);
-		SetNextTickTime(NodeMemory, NextTickTime);
-	}
-	else
-	{
-		Super::OnSearchStart(SearchData);
 	}
 }
 
@@ -136,12 +128,20 @@ FString UBTService_BlueprintBase::GetStaticServiceDescription() const
 	UBTService_BlueprintBase* CDO = (UBTService_BlueprintBase*)(GetClass()->GetDefaultObject());
 	if (CDO)
 	{
-		ReturnDesc = FString::Printf(TEXT("%s, %s, %s, %s\n"),
-			ReceiveTickImplementations != 0 ? *GetStaticTickIntervalDescription() : TEXT("No tick"),
-			ReceiveActivationImplementations != 0 ? TEXT("Activation") : TEXT("No Activation"),
-			ReceiveDeactivationImplementations != 0 ? TEXT("Deactivation") : TEXT("No Deactivation"),
-			ReceiveSearchStartImplementations != 0 ? TEXT("Search Start") : TEXT("No Search Start"));
-								
+		if (bShowEventDetails)
+		{
+			ReturnDesc = FString::Printf(TEXT("%s, %s, %s, %s\n"),
+				ReceiveTickImplementations != 0 ? *GetStaticTickIntervalDescription() : TEXT("No tick"),
+				ReceiveActivationImplementations != 0 ? TEXT("Activation") : TEXT("No Activation"),
+				ReceiveDeactivationImplementations != 0 ? TEXT("Deactivation") : TEXT("No Deactivation"),
+				ReceiveSearchStartImplementations != 0 ? TEXT("Search Start") : TEXT("No Search Start"));
+		}
+		else
+		{
+			ReturnDesc = Super::GetStaticServiceDescription();
+			ReturnDesc += TEXT('\n');
+		}
+						
 		if (bShowPropertyDetails)
 		{
 			UClass* StopAtClass = UBTService_BlueprintBase::StaticClass();
