#ue4 - Avoid casts in Navigation update when result is not being used.

diff --git a/Engine/Source/Runtime/AIModule/Private/Navigation/CrowdManager.cpp b/Engine/Source/Runtime/AIModule/Private/Navigation/CrowdManager.cpp
--- a/Engine/Source/Runtime/AIModule/Private/Navigation/CrowdManager.cpp
+++ b/Engine/Source/Runtime/AIModule/Private/Navigation/CrowdManager.cpp
@@ -768,7 +768,7 @@ void UCrowdManager::UpdateAgentPaths()
 		FCrowdAgentData& AgentData = It.Value();
 		if (AgentData.bIsSimulated && AgentData.IsValid())
 		{
-			UCrowdFollowingComponent* CrowdComponent = (UCrowdFollowingComponent*)Cast<const UCrowdFollowingComponent>(It.Key());
+			UCrowdFollowingComponent* CrowdComponent = nullptr;
 
 			const dtCrowdAgent* Agent = DetourCrowd->getAgent(AgentData.AgentIndex);
 			dtPolyRef AgentPolyRef = Agent->corridor.getFirstPoly();
@@ -793,6 +793,7 @@ void UCrowdManager::UpdateAgentPaths()
 						DetourCrowd->resetAgentVelocity(AgentData.AgentIndex);
 
 						// start using smart link
+						CrowdComponent = (CrowdComponent ? CrowdComponent : (UCrowdFollowingComponent*)Cast<const UCrowdFollowingComponent>(It.Key()));
 						if (CrowdComponent)
 						{
 							CrowdComponent->StartUsingCustomLink(CustomLink, EndPt);
@@ -802,10 +803,14 @@ void UCrowdManager::UpdateAgentPaths()
 			}
 
 			// look for poly updates
-			if (AgentPolyRef != AgentData.PrevPoly && CrowdComponent)
+			if (AgentPolyRef != AgentData.PrevPoly)
 			{
-				CrowdComponent->OnNavNodeChanged(AgentPolyRef, AgentData.PrevPoly, Agent->corridor.getPathCount());
-				AgentData.PrevPoly = AgentPolyRef;
+				CrowdComponent = (CrowdComponent ? CrowdComponent : (UCrowdFollowingComponent*)Cast<const UCrowdFollowingComponent>(It.Key()));
+				if (CrowdComponent)
+				{
+					CrowdComponent->OnNavNodeChanged(AgentPolyRef, AgentData.PrevPoly, Agent->corridor.getPathCount());
+					AgentData.PrevPoly = AgentPolyRef;
+				}
 			}
 		}
 	}
