added OnTaskFinished event to behavior tree task nodes,
always calling OnCeaseRelevant and OnTaskFinished on active nodes when tree finishes (restart / end play)

diff --git a/Engine/Source/Runtime/AIModule/Classes/BehaviorTree/BTTaskNode.h b/Engine/Source/Runtime/AIModule/Classes/BehaviorTree/BTTaskNode.h
--- a/Engine/Source/Runtime/AIModule/Classes/BehaviorTree/BTTaskNode.h
+++ b/Engine/Source/Runtime/AIModule/Classes/BehaviorTree/BTTaskNode.h
@@ -54,6 +54,9 @@ class AIMODULE_API UBTTaskNode : public UBTNode
 	/** wrapper for node instancing: TickTask */
 	void WrappedTickTask(UBehaviorTreeComponent& OwnerComp, uint8* NodeMemory, float DeltaSeconds) const;
 
+	/** wrapper for node instancing: OnTaskFinished */
+	void WrappedOnTaskFinished(UBehaviorTreeComponent& OwnerComp, uint8* NodeMemory, EBTNodeResult::Type TaskResult) const;
+
 	/** helper function: finish latent executing */
 	void FinishLatentTask(UBehaviorTreeComponent& OwnerComp, EBTNodeResult::Type TaskResult) const;
 
@@ -65,6 +68,9 @@ protected:
 	/** if set, TickTask will be called */
 	uint8 bNotifyTick : 1;
 
+	/** if set, OnTaskFinished will be called */
+	uint8 bNotifyTaskFinished : 1;
+
 	/** ticks this task 
 	 * this function should be considered as const (don't modify state of object) if node is not instanced! */
 	virtual void TickTask(UBehaviorTreeComponent& OwnerComp, uint8* NodeMemory, float DeltaSeconds);
@@ -73,6 +79,10 @@ protected:
 	 * this function should be considered as const (don't modify state of object) if node is not instanced! */
 	virtual void OnMessage(UBehaviorTreeComponent& OwnerComp, uint8* NodeMemory, FName Message, int32 RequestID, bool bSuccess);
 
+	/** called when task execution is finished
+	 * this function should be considered as const (don't modify state of object) if node is not instanced! */
+	virtual void OnTaskFinished(UBehaviorTreeComponent& OwnerComp, uint8* NodeMemory, EBTNodeResult::Type TaskResult);
+
 	/** register message observer */
 	void WaitForMessage(UBehaviorTreeComponent& OwnerComp, FName MessageType) const;
 	void WaitForMessage(UBehaviorTreeComponent& OwnerComp, FName MessageType, int32 RequestID) const;
diff --git a/Engine/Source/Runtime/AIModule/Classes/BehaviorTree/BehaviorTreeComponent.h b/Engine/Source/Runtime/AIModule/Classes/BehaviorTree/BehaviorTreeComponent.h
--- a/Engine/Source/Runtime/AIModule/Classes/BehaviorTree/BehaviorTreeComponent.h
+++ b/Engine/Source/Runtime/AIModule/Classes/BehaviorTree/BehaviorTreeComponent.h
@@ -77,11 +77,12 @@ protected:
 public:
 	virtual bool IsRunning() const override;
 	virtual bool IsPaused() const override;
+	virtual void Cleanup() override;
 	// End UBrainComponent overrides
 
-	// Begin UObject overrides
-	virtual void BeginDestroy() override;
-	// End UObject overrides
+	// Begin UActorComponent overrides
+	virtual void UninitializeComponent() override;
+	// End UActorComponent overrides
 
 	/** starts execution from root */
 	bool StartTree(UBehaviorTree& Asset, EBTExecutionMode::Type ExecuteMode = EBTExecutionMode::Looped);
@@ -136,9 +137,6 @@ public:
 	/** schedule execution flow update in next tick */
 	void ScheduleExecutionUpdate();
 
-	/** remove all runtime data, used on map change */
-	void Cleanup();
-
 	/** tries to find behavior tree instance in context */
 	int32 FindInstanceContainingNode(const UBTNode* Node) const;
 
diff --git a/Engine/Source/Runtime/AIModule/Classes/BrainComponent.h b/Engine/Source/Runtime/AIModule/Classes/BrainComponent.h
--- a/Engine/Source/Runtime/AIModule/Classes/BrainComponent.h
+++ b/Engine/Source/Runtime/AIModule/Classes/BrainComponent.h
@@ -133,6 +133,9 @@ public:
 	UFUNCTION(BlueprintCallable, Category = "AI|Logic")
 	virtual void StopLogic(const FString& Reason);
 
+	/** AI logic won't be needed anymore, stop all activity and run cleanup */
+	virtual void Cleanup() {}
+
 protected:
 	virtual void PauseLogic(const FString& Reason) {}
 	/** MUST be called by child implementations!
diff --git a/Engine/Source/Runtime/AIModule/Private/AIController.cpp b/Engine/Source/Runtime/AIModule/Private/AIController.cpp
--- a/Engine/Source/Runtime/AIModule/Private/AIController.cpp
+++ b/Engine/Source/Runtime/AIModule/Private/AIController.cpp
@@ -465,6 +465,11 @@ void AAIController::UnPossess()
 	{
 		PathFollowingComponent->Cleanup();
 	}
+
+	if (BrainComponent)
+	{
+		BrainComponent->Cleanup();
+	}
 }
 
 void AAIController::InitNavigationControl(UPathFollowingComponent*& PathFollowingComp)
diff --git a/Engine/Source/Runtime/AIModule/Private/BehaviorTree/BTTaskNode.cpp b/Engine/Source/Runtime/AIModule/Private/BehaviorTree/BTTaskNode.cpp
--- a/Engine/Source/Runtime/AIModule/Private/BehaviorTree/BTTaskNode.cpp
+++ b/Engine/Source/Runtime/AIModule/Private/BehaviorTree/BTTaskNode.cpp
@@ -7,6 +7,7 @@ UBTTaskNode::UBTTaskNode(const FObjectInitializer& ObjectInitializer) : Super(Ob
 {
 	NodeName = "UnknownTask";
 	bNotifyTick = false;
+	bNotifyTaskFinished = false;
 }
 
 EBTNodeResult::Type UBTTaskNode::ExecuteTask(UBehaviorTreeComponent& OwnerComp, uint8* NodeMemory)
@@ -43,6 +44,18 @@ void UBTTaskNode::WrappedTickTask(UBehaviorTreeComponent& OwnerComp, uint8* Node
 	}
 }
 
+void UBTTaskNode::WrappedOnTaskFinished(UBehaviorTreeComponent& OwnerComp, uint8* NodeMemory, EBTNodeResult::Type TaskResult) const
+{
+	if (bNotifyTaskFinished)
+	{
+		const UBTNode* NodeOb = bCreateNodeInstance ? GetNodeInstance(OwnerComp, NodeMemory) : this;
+		if (NodeOb)
+		{
+			((UBTTaskNode*)NodeOb)->OnTaskFinished(OwnerComp, NodeMemory, TaskResult);
+		}
+	}
+}
+
 void UBTTaskNode::ReceivedMessage(UBrainComponent* BrainComp, const FAIMessage& Message)
 {
 	UBehaviorTreeComponent* OwnerComp = (UBehaviorTreeComponent*)BrainComp;
@@ -57,6 +70,11 @@ void UBTTaskNode::TickTask(UBehaviorTreeComponent& OwnerComp, uint8* NodeMemory,
 	// empty in base class
 }
 
+void UBTTaskNode::OnTaskFinished(UBehaviorTreeComponent& OwnerComp, uint8* NodeMemory, EBTNodeResult::Type TaskResult)
+{
+	// empty in base class
+}
+
 void UBTTaskNode::OnMessage(UBehaviorTreeComponent& OwnerComp, uint8* NodeMemory, FName Message, int32 RequestID, bool bSuccess)
 {
 	const EBTTaskStatus::Type Status = OwnerComp.GetTaskStatus(this);
diff --git a/Engine/Source/Runtime/AIModule/Private/BehaviorTree/BehaviorTreeComponent.cpp b/Engine/Source/Runtime/AIModule/Private/BehaviorTree/BehaviorTreeComponent.cpp
--- a/Engine/Source/Runtime/AIModule/Private/BehaviorTree/BehaviorTreeComponent.cpp
+++ b/Engine/Source/Runtime/AIModule/Private/BehaviorTree/BehaviorTreeComponent.cpp
@@ -27,7 +27,7 @@ UBehaviorTreeComponent::UBehaviorTreeComponent(const FObjectInitializer& ObjectI
 	bIsPaused = false;
 }
 
-void UBehaviorTreeComponent::BeginDestroy()
+void UBehaviorTreeComponent::UninitializeComponent()
 {
 	UBehaviorTreeManager* BTManager = UBehaviorTreeManager::GetCurrent(GetWorld());
 	if (BTManager)
@@ -36,7 +36,7 @@ void UBehaviorTreeComponent::BeginDestroy()
 	}
 
 	RemoveAllInstances();
-	Super::BeginDestroy();
+	Super::UninitializeComponent();
 }
 
 void UBehaviorTreeComponent::RestartLogic()
@@ -149,7 +149,7 @@ void UBehaviorTreeComponent::StopTree()
 	{
 		FBehaviorTreeInstance& InstanceInfo = InstanceStack[InstanceIndex];
 
-		// notify active nodes
+		// notify active aux nodes
 		for (int32 AuxIndex = 0; AuxIndex < InstanceInfo.ActiveAuxNodes.Num(); AuxIndex++)
 		{
 			const UBTAuxiliaryNode* AuxNode = InstanceInfo.ActiveAuxNodes[AuxIndex];
@@ -158,6 +158,26 @@ void UBehaviorTreeComponent::StopTree()
 			AuxNode->WrappedOnCeaseRelevant(*this, NodeMemory);
 		}
 
+		// notify active parallel tasks
+		for (int32 ParallelIndex = 0; ParallelIndex < InstanceInfo.ParallelTasks.Num(); ParallelIndex++)
+		{
+			const FBehaviorTreeParallelTask& ParalleInfo = InstanceInfo.ParallelTasks[ParallelIndex];
+			if (ParalleInfo.TaskNode && (ParalleInfo.Status != EBTTaskStatus::Inactive))
+			{
+				uint8* NodeMemory = ParalleInfo.TaskNode->GetNodeMemory<uint8>(InstanceInfo);
+				ParalleInfo.TaskNode->WrappedOnTaskFinished(*this, NodeMemory, EBTNodeResult::InProgress);
+			}
+		}
+
+		// notify active task
+		if (InstanceInfo.ActiveNodeType == EBTActiveNode::ActiveTask || InstanceInfo.ActiveNodeType == EBTActiveNode::AbortingTask)
+		{
+			const UBTTaskNode* TaskNode = Cast<const UBTTaskNode>(InstanceInfo.ActiveNode);
+			check(TaskNode != NULL);
+			uint8* NodeMemory = TaskNode->GetNodeMemory<uint8>(InstanceInfo);
+			TaskNode->WrappedOnTaskFinished(*this, NodeMemory, EBTNodeResult::InProgress);
+		}
+		
 		// clear instance
 		InstanceStack[InstanceIndex].Cleanup(*this, EBTMemoryClear::StoreSubtree);
 	}
@@ -223,6 +243,10 @@ void UBehaviorTreeComponent::OnTaskFinished(const UBTTaskNode* TaskNode, EBTNode
 		// cleanup task observers
 		UnregisterMessageObserversFrom(TaskNode);
 
+		// notify task about it
+		uint8* TaskMemory = TaskNode->GetNodeMemory<uint8>(InstanceStack[TaskInstanceIdx]);
+		TaskNode->WrappedOnTaskFinished(*this, TaskMemory, TaskResult);
+
 		// update execution when active task is finished
 		FBehaviorTreeInstance& ActiveInstance = InstanceStack[ActiveInstanceIdx];
 		if (ActiveInstance.ActiveNode == TaskNode)
