Added a helper BP function to retrieve AI controller from an Actor #UE4

The way it works exactly is if the actor passed in is a pawn, then the function retrieves pawn's conroller cast to AIController. Otherwise the function returns actor cast to AIController.

diff --git a/Engine/Source/Runtime/AIModule/Classes/Blueprint/AIBlueprintHelperLibrary.h b/Engine/Source/Runtime/AIModule/Classes/Blueprint/AIBlueprintHelperLibrary.h
--- a/Engine/Source/Runtime/AIModule/Classes/Blueprint/AIBlueprintHelperLibrary.h
+++ b/Engine/Source/Runtime/AIModule/Classes/Blueprint/AIBlueprintHelperLibrary.h
@@ -11,6 +11,7 @@
 class UBehaviorTree;
 class UAnimInstance;
 class APawn;
+class AAIController; 
 class UBlackboardComponent;
 class UAIAsyncTaskBlueprintProxy;
 
@@ -28,6 +29,11 @@ class AIMODULE_API UAIBlueprintHelperLibrary : public UBlueprintFunctionLibrary
 	UFUNCTION(BlueprintCallable, Category="AI", meta=(WorldContext="WorldContextObject", UnsafeDuringActorConstruction="true"))
 	static APawn* SpawnAIFromClass(UObject* WorldContextObject, TSubclassOf<APawn> PawnClass, UBehaviorTree* BehaviorTree, FVector Location, FRotator Rotation=FRotator::ZeroRotator, bool bNoCollisionFail=false);
 
+	/** The way it works exactly is if the actor passed in is a pawn, then the function retrieves 
+	 *	pawn's controller cast to AIController. Otherwise the function returns actor cast to AIController. */
+	UFUNCTION(BlueprintCallable, Category = "AI", meta = (DefaultToSelf = "ControlledObject"))
+	static AAIController* GetAIController(AActor* ControlledActor);
+
 	UFUNCTION(BlueprintPure, Category="AI", meta=(DefaultToSelf="Target"))
 	static UBlackboardComponent* GetBlackboard(AActor* Target);
 
diff --git a/Engine/Source/Runtime/AIModule/Private/Blueprint/AIBlueprintHelperLibrary.cpp b/Engine/Source/Runtime/AIModule/Private/Blueprint/AIBlueprintHelperLibrary.cpp
--- a/Engine/Source/Runtime/AIModule/Private/Blueprint/AIBlueprintHelperLibrary.cpp
+++ b/Engine/Source/Runtime/AIModule/Private/Blueprint/AIBlueprintHelperLibrary.cpp
@@ -151,6 +151,16 @@ APawn* UAIBlueprintHelperLibrary::SpawnAIFromClass(UObject* WorldContextObject,
 	return NewPawn;
 }
 
+AAIController* UAIBlueprintHelperLibrary::GetAIController(AActor* ControlledActor)
+{
+	APawn* AsPawn = Cast<APawn>(ControlledActor);
+	if (AsPawn != nullptr)
+	{
+		return Cast<AAIController>(AsPawn->GetController());
+	}
+	return Cast<AAIController>(ControlledActor);
+}
+
 UBlackboardComponent* UAIBlueprintHelperLibrary::GetBlackboard(AActor* Target)
 {
 	UBlackboardComponent* BlackboardComp = nullptr;
