Fixed issues with direct, navigation less nav paths in absence of any navigation data on the level #UE4

- tested with floating pawn component and made it BP-usable
- also, added a handy macro for logging AI locations

diff --git a/Engine/Source/Runtime/AIModule/Classes/AITypes.h b/Engine/Source/Runtime/AIModule/Classes/AITypes.h
--- a/Engine/Source/Runtime/AIModule/Classes/AITypes.h
+++ b/Engine/Source/Runtime/AIModule/Classes/AITypes.h
@@ -4,6 +4,8 @@
 
 #include "AITypes.generated.h"
 
+#define TEXT_AI_LOCATION(v) (FAISystem::IsValidLocation(v) ? *(v).ToString() : TEXT("Invalid"))
+
 namespace FAISystem
 {
 	static const FRotator InvalidRotation = FRotator(FLT_MAX);
diff --git a/Engine/Source/Runtime/AIModule/Private/AIController.cpp b/Engine/Source/Runtime/AIModule/Private/AIController.cpp
--- a/Engine/Source/Runtime/AIModule/Private/AIController.cpp
+++ b/Engine/Source/Runtime/AIModule/Private/AIController.cpp
@@ -118,7 +118,7 @@ void AAIController::GrabDebugSnapshot(FVisualLogEntry* Snapshot) const
 
 	if (FocusActor == nullptr)
 	{
-		MyCategory.Add(TEXT("Focus Location"), GetFocalPoint().ToString());
+		MyCategory.Add(TEXT("Focus Location"), TEXT_AI_LOCATION(GetFocalPoint()));
 	}
 	Snapshot->Status.Add(MyCategory);
 
@@ -533,14 +533,14 @@ EPathFollowingRequestResult::Type AAIController::MoveToLocation(const FVector& D
 	EPathFollowingRequestResult::Type Result = EPathFollowingRequestResult::Failed;
 	bool bCanRequestMove = true;
 
-	UE_VLOG(this, LogAINavigation, Log, TEXT("MoveToLocation: Goal(%s) AcceptRadius(%.1f%s) bUsePathfinding(%d) bCanStrafe(%d) Filter(%s)"),
-		*Dest.ToString(), AcceptanceRadius, bStopOnOverlap ? TEXT(" + agent") : TEXT(""), bUsePathfinding, bCanStrafe, *GetNameSafe(FilterClass));
+	UE_VLOG(this, LogAINavigation, Log, TEXT("MoveToLocation: Goal(%s) AcceptRadius(%.1f%s) bUsePathfinding(%d) bCanStrafe(%d) Filter(%s)")
+		, TEXT_AI_LOCATION(Dest), AcceptanceRadius, bStopOnOverlap ? TEXT(" + agent") : TEXT(""), bUsePathfinding, bCanStrafe, *GetNameSafe(FilterClass));
 
 	// Check input is valid
 	if (Dest.ContainsNaN())
 	{
-		UE_VLOG(this, LogAINavigation, Error, TEXT("AAIController::MoveToLocation: Destination is not valid! Goal(%s) AcceptRadius(%.1f%s) bUsePathfinding(%d) bCanStrafe(%d)"),
-			*Dest.ToString(), AcceptanceRadius, bStopOnOverlap ? TEXT(" + agent") : TEXT(""), bUsePathfinding, bCanStrafe);
+		UE_VLOG(this, LogAINavigation, Error, TEXT("AAIController::MoveToLocation: Destination is not valid! Goal(%s) AcceptRadius(%.1f%s) bUsePathfinding(%d) bCanStrafe(%d)")
+			, TEXT_AI_LOCATION(Dest), AcceptanceRadius, bStopOnOverlap ? TEXT(" + agent") : TEXT(""), bUsePathfinding, bCanStrafe);
 
 		ensure(!Dest.ContainsNaN());
 		bCanRequestMove = false;
@@ -581,7 +581,7 @@ EPathFollowingRequestResult::Type AAIController::MoveToLocation(const FVector& D
 	if (bCanRequestMove)
 	{
 		FPathFindingQuery Query;
-		const bool bValidQuery = PreparePathfinding(Query, GoalLocation, NULL, true, FilterClass);
+		const bool bValidQuery = PreparePathfinding(Query, GoalLocation, NULL, bUsePathfinding, FilterClass);
 		const FAIRequestID RequestID = bValidQuery ? RequestPathAndMove(Query, NULL, AcceptanceRadius, bStopOnOverlap, NULL) : FAIRequestID::InvalidRequest;
 
 		if (RequestID.IsValid())
diff --git a/Engine/Source/Runtime/Engine/Classes/AI/Navigation/NavigationSystem.h b/Engine/Source/Runtime/Engine/Classes/AI/Navigation/NavigationSystem.h
--- a/Engine/Source/Runtime/Engine/Classes/AI/Navigation/NavigationSystem.h
+++ b/Engine/Source/Runtime/Engine/Classes/AI/Navigation/NavigationSystem.h
@@ -308,9 +308,11 @@ public:
 
 	UCrowdManager* GetCrowdManager() const { return CrowdManager.Get(); }
 
+protected:
 	/** spawn new crowd manager */
 	virtual void CreateCrowdManager();
 
+public:
 	//----------------------------------------------------------------------//
 	// Public querying interface                                                                
 	//----------------------------------------------------------------------//
@@ -600,9 +602,14 @@ public:
 	/** Number of currently running tasks */
 	int32 GetNumRunningBuildTasks() const;
 
+protected:
 	/** Sets up SuportedAgents and NavigationDataCreators. Override it to add additional setup, but make sure to call Super implementation */
 	virtual void DoInitialSetup();
 
+	/** spawn new crowd manager */
+	virtual void UpdateAbstractNavData();
+	
+public:
 	/** Called upon UWorld destruction to release what needs to be released */
 	void CleanUp(ECleanupMode Mode = ECleanupMode::CleanupUnsafe);
 
diff --git a/Engine/Source/Runtime/Engine/Classes/GameFramework/FloatingPawnMovement.h b/Engine/Source/Runtime/Engine/Classes/GameFramework/FloatingPawnMovement.h
--- a/Engine/Source/Runtime/Engine/Classes/GameFramework/FloatingPawnMovement.h
+++ b/Engine/Source/Runtime/Engine/Classes/GameFramework/FloatingPawnMovement.h
@@ -12,7 +12,7 @@
  * FloatingPawnMovement is a movement component that provides simple movement for any Pawn class.
  * Limits on speed and acceleration are provided, while gravity is not implemented.
  */
-UCLASS()
+UCLASS(meta = (BlueprintSpawnableComponent))
 class ENGINE_API UFloatingPawnMovement : public UPawnMovementComponent
 {
 	GENERATED_UCLASS_BODY()
diff --git a/Engine/Source/Runtime/Engine/Private/AI/Navigation/AbstractNavData.cpp b/Engine/Source/Runtime/Engine/Private/AI/Navigation/AbstractNavData.cpp
--- a/Engine/Source/Runtime/Engine/Private/AI/Navigation/AbstractNavData.cpp
+++ b/Engine/Source/Runtime/Engine/Private/AI/Navigation/AbstractNavData.cpp
@@ -37,6 +37,7 @@ FPathFindingResult AAbstractNavData::FindPathAbstract(const FNavAgentProperties&
 
 	Result.Path->GetPathPoints().Add(FNavPathPoint(Query.StartLocation));
 	Result.Path->GetPathPoints().Add(FNavPathPoint(Query.EndLocation));
+	Result.Path->MarkReady();
 	Result.Result = ENavigationQueryResult::Success;
 
 	return Result;
diff --git a/Engine/Source/Runtime/Engine/Private/AI/Navigation/NavigationSystem.cpp b/Engine/Source/Runtime/Engine/Private/AI/Navigation/NavigationSystem.cpp
--- a/Engine/Source/Runtime/Engine/Private/AI/Navigation/NavigationSystem.cpp
+++ b/Engine/Source/Runtime/Engine/Private/AI/Navigation/NavigationSystem.cpp
@@ -262,6 +262,20 @@ void UNavigationSystem::DoInitialSetup()
 		return;
 	}
 	
+	UpdateAbstractNavData();
+
+	CreateCrowdManager();
+
+	bInitialSetupHasBeenPerformed = true;
+}
+
+void UNavigationSystem::UpdateAbstractNavData()
+{
+	if (AbstractNavData != nullptr)
+	{
+		return;
+	}
+
 	// spawn abstract nav data separately
 	// it's responsible for direct paths and shouldn't be picked for any agent type as default one
 	UWorld* NavWorld = GetWorld();
@@ -282,11 +296,14 @@ void UNavigationSystem::DoInitialSetup()
 		AbstractNavData = CreateNavigationDataInstance(DummyConfig);
 	}
 
-	CreateCrowdManager();
-
-	bInitialSetupHasBeenPerformed = true;
+	if (AbstractNavData->IsRegistered() == false)
+	{
+		// fake registration since it's a special navigation data type 
+		// and it would get discarded for not implementing any particular
+		// navigation agent
+		AbstractNavData->OnRegistered();
+	}
 }
-
 void UNavigationSystem::PostInitProperties()
 {
 	Super::PostInitProperties();
@@ -524,6 +541,7 @@ void UNavigationSystem::OnWorldInitDone(FNavigationSystem::EMode Mode)
 	else // just register data already present
 	{
 		RegisterNavigationDataInstances();
+		UpdateAbstractNavData();
 	}
 }
 
