Fixed AIMoveTo K2 node not being able to follow AActors #UE4

- there was a bug of not passing goal actor down the call chain which was resulting in AI not knowing it was supposed to update the path on goal's move

diff --git a/Engine/Source/Runtime/AIModule/Private/Blueprint/AIBlueprintHelperLibrary.cpp b/Engine/Source/Runtime/AIModule/Private/Blueprint/AIBlueprintHelperLibrary.cpp
--- a/Engine/Source/Runtime/AIModule/Private/Blueprint/AIBlueprintHelperLibrary.cpp
+++ b/Engine/Source/Runtime/AIModule/Private/Blueprint/AIBlueprintHelperLibrary.cpp
@@ -97,7 +97,7 @@ UAIAsyncTaskBlueprintProxy* UAIBlueprintHelperLibrary::CreateMoveToProxyObject(U
 
 		if (bValidQuery)
 		{
-			const FAIRequestID RequestID = AIController->RequestPathAndMove(Query, NULL, AcceptanceRadius, bStopOnOverlap, NULL);
+			const FAIRequestID RequestID = AIController->RequestPathAndMove(Query, TargetActor, AcceptanceRadius, bStopOnOverlap, NULL);
 			if (RequestID.IsValid())
 			{
 				MyObj->AIController = AIController;
