Modified default value of AIMoveTo K2 node's AcceptanceRAdius to non-zero 5uu #UE4

The main reason being that 0-radius is very tricky for AI and is generally not recomended

diff --git a/Engine/Source/Runtime/AIModule/Classes/Blueprint/AIBlueprintHelperLibrary.h b/Engine/Source/Runtime/AIModule/Classes/Blueprint/AIBlueprintHelperLibrary.h
--- a/Engine/Source/Runtime/AIModule/Classes/Blueprint/AIBlueprintHelperLibrary.h
+++ b/Engine/Source/Runtime/AIModule/Classes/Blueprint/AIBlueprintHelperLibrary.h
@@ -21,7 +21,7 @@ class AIMODULE_API UAIBlueprintHelperLibrary : public UBlueprintFunctionLibrary
 	GENERATED_UCLASS_BODY()
 
 	UFUNCTION(BlueprintCallable, meta=(WorldContext="WorldContextObject", BlueprintInternalUseOnly = "TRUE"))
-	static UAIAsyncTaskBlueprintProxy* CreateMoveToProxyObject(UObject* WorldContextObject, APawn* Pawn, FVector Destination, AActor* TargetActor = NULL, float AcceptanceRadius = 0, bool bStopOnOverlap = false);
+	static UAIAsyncTaskBlueprintProxy* CreateMoveToProxyObject(UObject* WorldContextObject, APawn* Pawn, FVector Destination, AActor* TargetActor = NULL, float AcceptanceRadius = 5.f, bool bStopOnOverlap = false);
 
 	UFUNCTION(BlueprintCallable, Category="AI", meta=(DefaultToSelf="MessageSource"))
 	static void SendAIMessage(APawn* Target, FName Message, UObject* MessageSource, bool bSuccess = true);
