Fixed a comment on UBTDecorator_ConeCheck::ConeHalfAngle #UE4

diff --git a/Engine/Source/Runtime/AIModule/Classes/BehaviorTree/Decorators/BTDecorator_ConeCheck.h b/Engine/Source/Runtime/AIModule/Classes/BehaviorTree/Decorators/BTDecorator_ConeCheck.h
--- a/Engine/Source/Runtime/AIModule/Classes/BehaviorTree/Decorators/BTDecorator_ConeCheck.h
+++ b/Engine/Source/Runtime/AIModule/Classes/BehaviorTree/Decorators/BTDecorator_ConeCheck.h
@@ -23,7 +23,7 @@ class AIMODULE_API UBTDecorator_ConeCheck : public UBTDecorator
 
 	typedef FBTConeCheckDecoratorMemory TNodeInstanceMemory;
 
-	/** max allowed time for execution of underlying node */
+	/** Angle between cone direction and code cone edge, or a half of the total cone angle */
 	UPROPERTY(Category=Decorator, EditAnywhere)
 	float ConeHalfAngle;
 	
