Added a sanity check for Navigation System's existance in FEQSHelpers::FindNavMeshForQuery #UE4

diff --git a/Engine/Source/Runtime/AIModule/Private/EnvironmentQuery/EnvQueryTypes.cpp b/Engine/Source/Runtime/AIModule/Private/EnvironmentQuery/EnvQueryTypes.cpp
--- a/Engine/Source/Runtime/AIModule/Private/EnvironmentQuery/EnvQueryTypes.cpp
+++ b/Engine/Source/Runtime/AIModule/Private/EnvironmentQuery/EnvQueryTypes.cpp
@@ -241,6 +241,11 @@ namespace FEQSHelpers
 	const ARecastNavMesh* FindNavMeshForQuery(FEnvQueryInstance& QueryInstance)
 	{
 		const UNavigationSystem* NavSys = QueryInstance.World->GetNavigationSystem();
+		
+		if (NavSys == nullptr)
+		{
+			return nullptr;
+		}
 
 		// try to match navigation agent for querier
 		INavAgentInterface* NavAgent = QueryInstance.Owner.IsValid() ? Cast<INavAgentInterface>(QueryInstance.Owner.Get()) : NULL;
