Parametrized focus priority used by BTService_DefaultFocus, but kept it unexposed #UE4

The reason for not exposing it is that this service by definition should be used to set the default-priority focus. Underlying logic has not been created with other levels in mind. By having it a property however deriving from BTService_DefaultFocus and reusing the logic is easy now.

- github PR#736 by dzeligman

diff --git a/Engine/Source/Runtime/AIModule/Classes/BehaviorTree/Services/BTService_DefaultFocus.h b/Engine/Source/Runtime/AIModule/Classes/BehaviorTree/Services/BTService_DefaultFocus.h
--- a/Engine/Source/Runtime/AIModule/Classes/BehaviorTree/Services/BTService_DefaultFocus.h
+++ b/Engine/Source/Runtime/AIModule/Classes/BehaviorTree/Services/BTService_DefaultFocus.h
@@ -25,7 +25,13 @@ struct FBTFocusMemory
 UCLASS(hidecategories=(Service))
 class AIMODULE_API UBTService_DefaultFocus : public UBTService_BlackboardBase
 {
-	GENERATED_UCLASS_BODY()
+	GENERATED_BODY()
+	
+	// not exposed to users on purpose. Here to make reusing focus-setting mechanics by derived classes possible
+	UPROPERTY()
+	uint8 FocusPriority;
+
+	UBTService_DefaultFocus(const FObjectInitializer& ObjectInitializer);
 
 	virtual uint16 GetInstanceMemorySize() const override { return sizeof(FBTFocusMemory); }
 	virtual void OnBecomeRelevant(UBehaviorTreeComponent& OwnerComp, uint8* NodeMemory) override;
diff --git a/Engine/Source/Runtime/AIModule/Private/BehaviorTree/Services/BTService_DefaultFocus.cpp b/Engine/Source/Runtime/AIModule/Private/BehaviorTree/Services/BTService_DefaultFocus.cpp
--- a/Engine/Source/Runtime/AIModule/Private/BehaviorTree/Services/BTService_DefaultFocus.cpp
+++ b/Engine/Source/Runtime/AIModule/Private/BehaviorTree/Services/BTService_DefaultFocus.cpp
@@ -14,6 +14,8 @@ UBTService_DefaultFocus::UBTService_DefaultFocus(const FObjectInitializer& Objec
 	bNotifyBecomeRelevant = true;
 	bNotifyCeaseRelevant = true;
 
+	FocusPriority = EAIFocusPriority::Default;
+
 	// accept only actors and vectors
 	BlackboardKey.AddObjectFilter(this, AActor::StaticClass());
 	BlackboardKey.AddVectorFilter(this);
@@ -38,7 +40,7 @@ void UBTService_DefaultFocus::OnBecomeRelevant(UBehaviorTreeComponent& OwnerComp
 			AActor* TargetActor = Cast<AActor>(KeyValue);
 			if (TargetActor)
 			{
-				OwnerController->SetFocus(TargetActor, EAIFocusPriority::Default);
+				OwnerController->SetFocus(TargetActor, FocusPriority);
 				MyMemory->FocusActorSet = TargetActor;
 				MyMemory->bActorSet = true;
 			}
@@ -46,7 +48,7 @@ void UBTService_DefaultFocus::OnBecomeRelevant(UBehaviorTreeComponent& OwnerComp
 		else
 		{
 			const FVector FocusLocation = MyBlackboard->GetValue<UBlackboardKeyType_Vector>(BlackboardKey.GetSelectedKeyID());
-			OwnerController->SetFocalPoint(FocusLocation, EAIFocusPriority::Default);
+			OwnerController->SetFocalPoint(FocusLocation, FocusPriority);
 			MyMemory->FocusLocationSet = FocusLocation;
 		}
 	}
@@ -64,16 +66,16 @@ void UBTService_DefaultFocus::OnCeaseRelevant(UBehaviorTreeComponent& OwnerComp,
 		bool bClearFocus = false;
 		if (MyMemory->bActorSet)
 		{
-			bClearFocus = (MyMemory->FocusActorSet == OwnerController->GetFocusActorForPriority(EAIFocusPriority::Default));
+			bClearFocus = (MyMemory->FocusActorSet == OwnerController->GetFocusActorForPriority(FocusPriority));
 		}
 		else
 		{
-			bClearFocus = (MyMemory->FocusLocationSet == OwnerController->GetFocalPointForPriority(EAIFocusPriority::Default));
+			bClearFocus = (MyMemory->FocusLocationSet == OwnerController->GetFocalPointForPriority(FocusPriority));
 		}
 
 		if (bClearFocus)
 		{
-			OwnerController->ClearFocus(EAIFocusPriority::Default);
+			OwnerController->ClearFocus(FocusPriority);
 		}
 	}
 }
