Replaced a bunch of direct Pawn.BaseEyeHeight uses with apropriate Pawn and Actor functions #UE4

- github PR#711 by slonopotamus

diff --git a/Engine/Source/Runtime/AIModule/Classes/AIController.h b/Engine/Source/Runtime/AIModule/Classes/AIController.h
--- a/Engine/Source/Runtime/AIModule/Classes/AIController.h
+++ b/Engine/Source/Runtime/AIModule/Classes/AIController.h
@@ -296,7 +296,6 @@ public:
 #endif
 
 	virtual void Reset() override;
-	virtual void GetPlayerViewPoint(FVector& out_Location, FRotator& out_Rotation) const override;
 
 	/**
 	 * Checks line to center and top of other actor
diff --git a/Engine/Source/Runtime/AIModule/Private/AIController.cpp b/Engine/Source/Runtime/AIModule/Private/AIController.cpp
--- a/Engine/Source/Runtime/AIModule/Private/AIController.cpp
+++ b/Engine/Source/Runtime/AIModule/Private/AIController.cpp
@@ -144,20 +144,6 @@ void AAIController::GrabDebugSnapshot(FVisualLogEntry* Snapshot) const
 }
 #endif // ENABLE_VISUAL_LOG
 
-void AAIController::GetPlayerViewPoint(FVector& out_Location, FRotator& out_Rotation) const
-{
-	// AI does things from the Pawn
-	if (GetPawn() != NULL)
-	{
-		out_Location = GetPawn()->GetActorLocation();
-		out_Rotation = GetPawn()->GetActorRotation();
-	}
-	else
-	{
-		Super::GetPlayerViewPoint(out_Location, out_Rotation);
-	}
-}
-
 void AAIController::SetFocalPoint(FVector NewFocus, EAIFocusPriority::Type InPriority)
 {
 	if (InPriority >= FocusInformation.Priorities.Num())
@@ -290,11 +276,13 @@ bool AAIController::LineOfSightTo(const AActor* Other, FVector ViewPoint, bool b
 
 	if (ViewPoint.IsZero())
 	{
-		AActor*	ViewTarg = GetViewTarget();
-		ViewPoint = ViewTarg->GetActorLocation();
-		if (ViewTarg == GetPawn())
+		FRotator ViewRotation;
+		GetActorEyesViewPoint(ViewPoint, ViewRotation);
+
+		// if we still don't have a view point we simply fail
+		if (ViewPoint.IsZero())
 		{
-			ViewPoint.Z += GetPawn()->BaseEyeHeight; //look from eyes
+			return false;
 		}
 	}
 
@@ -405,7 +393,7 @@ void AAIController::UpdateControlRotation(float DeltaTime, bool bUpdatePawn)
 
 	if (Pawn)
 	{
-		FVector Direction = FAISystem::IsValidLocation(FocalPoint) ? (FocalPoint - Pawn->GetActorLocation()) : Pawn->GetActorForwardVector();
+		FVector Direction = FAISystem::IsValidLocation(FocalPoint) ? (FocalPoint - Pawn->GetPawnViewLocation()) : Pawn->GetActorForwardVector();
 		FRotator NewControlRotation = Direction.Rotation();
 
 		// Don't pitch view unless looking at another pawn
diff --git a/Engine/Source/Runtime/AIModule/Private/Perception/AIPerceptionComponent.cpp b/Engine/Source/Runtime/AIModule/Private/Perception/AIPerceptionComponent.cpp
--- a/Engine/Source/Runtime/AIModule/Private/Perception/AIPerceptionComponent.cpp
+++ b/Engine/Source/Runtime/AIModule/Private/Perception/AIPerceptionComponent.cpp
@@ -310,23 +310,12 @@ FVector UAIPerceptionComponent::GetActorLocation(const AActor& Actor) const
 
 void UAIPerceptionComponent::GetLocationAndDirection(FVector& Location, FVector& Direction) const
 {
-	AController* OwnerController = Cast<AController>(GetOuter());
-	if (OwnerController != NULL)
-	{
-		const APawn* OwnerPawn = OwnerController->GetPawn();
-		if (OwnerPawn != NULL)
-		{
-			Location = OwnerPawn->GetActorLocation() + FVector(0,0,OwnerPawn->BaseEyeHeight);
-			Direction = OwnerPawn->GetActorRotation().Vector();
-			return;
-		}
-	}
-	
 	const AActor* OwnerActor = Cast<AActor>(GetOuter());
-	if (OwnerActor != NULL)
+	if (OwnerActor != nullptr)
 	{
-		Location = OwnerActor->GetActorLocation();
-		Direction = OwnerActor->GetActorRotation().Vector();
+		FRotator ViewRotation;
+		OwnerActor->GetActorEyesViewPoint(Location, ViewRotation);
+		Direction = ViewRotation.Vector();
 	}
 }
 
diff --git a/Engine/Source/Runtime/AIModule/Private/Perception/PawnSensingComponent.cpp b/Engine/Source/Runtime/AIModule/Private/Perception/PawnSensingComponent.cpp
--- a/Engine/Source/Runtime/AIModule/Private/Perception/PawnSensingComponent.cpp
+++ b/Engine/Source/Runtime/AIModule/Private/Perception/PawnSensingComponent.cpp
@@ -375,12 +375,8 @@ FVector UPawnSensingComponent::GetSensorLocation() const
 
 	if (SensorActor != NULL)
 	{
-		SensorLocation = SensorActor->GetActorLocation();
-		const APawn* SensorPawn = Cast<const APawn>(SensorActor);
-		if (SensorPawn != NULL)
-		{
-			SensorLocation.Z += SensorPawn->BaseEyeHeight;
-		}
+		FRotator ViewRotation;
+		SensorActor->GetActorEyesViewPoint(SensorLocation, ViewRotation);
 	}
 
 	return SensorLocation;
diff --git a/Engine/Source/Runtime/Engine/Private/Controller.cpp b/Engine/Source/Runtime/Engine/Private/Controller.cpp
--- a/Engine/Source/Runtime/Engine/Private/Controller.cpp
+++ b/Engine/Source/Runtime/Engine/Private/Controller.cpp
@@ -164,12 +164,8 @@ bool AController::LineOfSightTo(const AActor* Other, FVector ViewPoint, bool bAl
 
 	if ( ViewPoint.IsZero() )
 	{
-		AActor*	ViewTarg = GetViewTarget();
-		ViewPoint = ViewTarg->GetActorLocation();
-		if( ViewTarg == Pawn )
-		{
-			ViewPoint.Z += Pawn->BaseEyeHeight; //look from eyes
-		}
+		FRotator ViewRotation;
+		GetActorEyesViewPoint(ViewPoint, ViewRotation);
 	}
 
 	static FName NAME_LineOfSight = FName(TEXT("LineOfSight"));
