added validation to prevent behavior tree debugger from crashing, will log error message instead until we have a good repro
UE-6911
#ue4

diff --git a/Engine/Source/Runtime/AIModule/Private/BehaviorTree/BehaviorTreeComponent.cpp b/Engine/Source/Runtime/AIModule/Private/BehaviorTree/BehaviorTreeComponent.cpp
--- a/Engine/Source/Runtime/AIModule/Private/BehaviorTree/BehaviorTreeComponent.cpp
+++ b/Engine/Source/Runtime/AIModule/Private/BehaviorTree/BehaviorTreeComponent.cpp
@@ -2025,7 +2025,15 @@ void UBehaviorTreeComponent::UpdateDebuggerAfterExecution(const UBTTaskNode* Tas
 
 	// accessing RuntimeDesc should never be out of bounds (active task MUST be part of active instance)
 	const uint16& ExecutionIndex = TaskNode->GetExecutionIndex();
-	CurrentStep.InstanceStack[InstanceIdx].RuntimeDesc[ExecutionIndex] = ComposedDesc;
+	if (CurrentStep.InstanceStack[InstanceIdx].RuntimeDesc.IsValidIndex(ExecutionIndex))
+	{
+		CurrentStep.InstanceStack[InstanceIdx].RuntimeDesc[ExecutionIndex] = ComposedDesc;
+	}
+	else
+	{
+		UE_VLOG(GetOwner(), LogBehaviorTree, Error, TEXT("Incomplete debugger data! No runtime description for executed task, instance %d has only %d entries!"),
+			InstanceIdx, CurrentStep.InstanceStack[InstanceIdx].RuntimeDesc.Num());
+	}
 #endif
 }
 
