Fixed compilation errors related to including CrowdAgentInterface.h in a game project sources #UE4

- this is a fix for a live issue experienced by 4.7 Preview 2 users (UE-7450)

diff --git a/Engine/Source/Runtime/AIModule/Classes/Navigation/CrowdAgentInterface.h b/Engine/Source/Runtime/AIModule/Classes/Navigation/CrowdAgentInterface.h
--- a/Engine/Source/Runtime/AIModule/Classes/Navigation/CrowdAgentInterface.h
+++ b/Engine/Source/Runtime/AIModule/Classes/Navigation/CrowdAgentInterface.h
@@ -1,6 +1,7 @@
 // Copyright 1998-2015 Epic Games, Inc. All Rights Reserved.
 
 #pragma once
+#include "AITypes.h"
 #include "CrowdAgentInterface.generated.h"
 
 UINTERFACE(MinimalAPI, meta=(CannotImplementInterfaceInBlueprint))
