merged fixes for behavior tree decorator relevancy issues

diff --git a/Engine/Source/Runtime/AIModule/Classes/BehaviorTree/BTCompositeNode.h b/Engine/Source/Runtime/AIModule/Classes/BehaviorTree/BTCompositeNode.h
--- a/Engine/Source/Runtime/AIModule/Classes/BehaviorTree/BTCompositeNode.h
+++ b/Engine/Source/Runtime/AIModule/Classes/BehaviorTree/BTCompositeNode.h
@@ -162,6 +162,9 @@ class AIMODULE_API UBTCompositeNode : public UBTNode
 	/** find branch containing specified node index */
 	int32 GetMatchingChildIndex(int32 ActiveInstanceIdx, FBTNodeIndex& NodeIdx) const;
 
+	/** get first execution index of given branch */
+	uint16 GetBranchExecutionIndex(uint16 NodeInBranchIdx) const;
+
 	/** is child execution allowed by decorators? */
 	bool DoDecoratorsAllowExecution(UBehaviorTreeComponent& OwnerComp, int32 InstanceIdx, int32 ChildIdx) const;
 
diff --git a/Engine/Source/Runtime/AIModule/Private/BehaviorTree/BTCompositeNode.cpp b/Engine/Source/Runtime/AIModule/Private/BehaviorTree/BTCompositeNode.cpp
--- a/Engine/Source/Runtime/AIModule/Private/BehaviorTree/BTCompositeNode.cpp
+++ b/Engine/Source/Runtime/AIModule/Private/BehaviorTree/BTCompositeNode.cpp
@@ -471,6 +471,23 @@ int32 UBTCompositeNode::GetMatchingChildIndex(int32 ActiveInstanceIdx, FBTNodeIn
 	return (ActiveInstanceIdx > NodeIdx.InstanceIndex) ? UnlimitedRange : OutsideRange;
 }
 
+uint16 UBTCompositeNode::GetBranchExecutionIndex(uint16 NodeInBranchIdx) const
+{
+	uint16 PrevBranchStartIdx = GetExecutionIndex();
+	for (int32 ChildIndex = 0; ChildIndex < Children.Num(); ChildIndex++)
+	{
+		const uint16 BranchStartIdx = GetChildExecutionIndex(ChildIndex, EBTChildIndex::FirstNode);
+		if (BranchStartIdx > NodeInBranchIdx)
+		{
+			break;
+		}
+
+		PrevBranchStartIdx = BranchStartIdx;
+	}
+
+	return PrevBranchStartIdx;
+}
+
 int32 UBTCompositeNode::GetNextChild(FBehaviorTreeSearchData& SearchData, int32 LastChildIdx, EBTNodeResult::Type LastResult) const
 {
 	FBTCompositeMemory* NodeMemory = GetNodeMemory<FBTCompositeMemory>(SearchData);
diff --git a/Engine/Source/Runtime/AIModule/Private/BehaviorTree/BehaviorTreeComponent.cpp b/Engine/Source/Runtime/AIModule/Private/BehaviorTree/BehaviorTreeComponent.cpp
--- a/Engine/Source/Runtime/AIModule/Private/BehaviorTree/BehaviorTreeComponent.cpp
+++ b/Engine/Source/Runtime/AIModule/Private/BehaviorTree/BehaviorTreeComponent.cpp
@@ -958,10 +958,11 @@ void UBehaviorTreeComponent::ProcessExecutionRequest()
 		// additional operations for restarting:
 		if (!ExecutionRequest.bTryNextChild)
 		{
-			// mark all decorators less important than current search parent for removal
-			const int32 ParentIndexWithServices = ExecutionRequest.ExecuteNode->GetChildExecutionIndex(0, EBTChildIndex::FirstNode) - 1;
-			FBTNodeIndex SearchParentIdx(ExecutionRequest.ExecuteInstanceIdx, ParentIndexWithServices);
-			UnregisterAuxNodesUpTo(SearchParentIdx);
+			// mark all decorators less important than current search start node for removal
+			// (all from first node on branch with search start)
+			const int32 StartNodeIdx = ExecutionRequest.ExecuteNode->GetBranchExecutionIndex(ExecutionRequest.SearchStart.ExecutionIndex) - 1;
+			const FBTNodeIndex DeactivateIdx(ExecutionRequest.ExecuteInstanceIdx, FMath::Max(StartNodeIdx, 0));
+			UnregisterAuxNodesUpTo(DeactivateIdx);
 
 			// reactivate top search node, so it could use search range correctly
 			BT_SEARCHLOG(SearchData, Verbose, TEXT("Reactivate node: %s [restart]"), *UBehaviorTreeTypes::DescribeNodeHelper(TestNode));
