Added getter and setter for UPathFollowingComponent::bStopMovementOnFinish #UE4

github PR#746 by gzeligman

diff --git a/Engine/Source/Runtime/AIModule/Classes/Navigation/PathFollowingComponent.h b/Engine/Source/Runtime/AIModule/Classes/Navigation/PathFollowingComponent.h
--- a/Engine/Source/Runtime/AIModule/Classes/Navigation/PathFollowingComponent.h
+++ b/Engine/Source/Runtime/AIModule/Classes/Navigation/PathFollowingComponent.h
@@ -182,6 +182,12 @@ class AIMODULE_API UPathFollowingComponent : public UActorComponent, public IAIR
 	/** set block detection params */
 	void SetBlockDetection(float DistanceThreshold, float Interval, int32 NumSamples);
 
+	/** @returns state of movement stopping on finish */
+	FORCEINLINE bool IsStopMovementOnFinishActive() const { return bStopMovementOnFinish; }
+	
+	/** set whether movement is stopped on finish of move. */
+	FORCEINLINE void SetStopMovementOnFinish(bool bEnable) { bStopMovementOnFinish = bEnable; }
+
 	/** set threshold for precise reach tests in intermediate goals (minimal test radius)  */
 	void SetPreciseReachThreshold(float AgentRadiusMultiplier, float AgentHalfHeightMultiplier);
 
