fixed blackboard observers delegates storing their runtime data in shared node templates

diff --git a/Engine/Source/Runtime/AIModule/Classes/BehaviorTree/BlackboardComponent.h b/Engine/Source/Runtime/AIModule/Classes/BehaviorTree/BlackboardComponent.h
--- a/Engine/Source/Runtime/AIModule/Classes/BehaviorTree/BlackboardComponent.h
+++ b/Engine/Source/Runtime/AIModule/Classes/BehaviorTree/BlackboardComponent.h
@@ -65,7 +65,7 @@ public:
 	bool HasValidAsset() const;
 
 	/** register observer for blackboard key */
-	FDelegateHandle RegisterObserver(FBlackboard::FKey KeyID, FOnBlackboardChange ObserverDelegate);
+	FDelegateHandle RegisterObserver(FBlackboard::FKey KeyID, UObject* NotifyOwner, FOnBlackboardChange ObserverDelegate);
 
 	/** unregister observer from blackboard key */
 	DELEGATE_DEPRECATED("This overload of UnregisterObserver is deprecated, instead pass the result of RegisterObserver.")
@@ -74,6 +74,9 @@ public:
 	/** unregister observer from blackboard key */
 	void UnregisterObserver(FBlackboard::FKey KeyID, FDelegateHandle ObserverHandle);
 
+	/** unregister all observers associated with given owner */
+	void UnregisterObserversFrom(UObject* NotifyOwner);
+
 	/** pause change notifies and add them to queue */
 	void PauseUpdates();
 
@@ -278,6 +281,9 @@ protected:
 	/** observers registered for blackboard keys */
 	TMultiMap<uint8, FOnBlackboardChange> Observers;
 
+	/** observers registered from owner objects */
+	TMultiMap<UObject*, FDelegateHandle> ObserverHandles;
+
 	/** queued key change notification, will be processed on ResumeUpdates call */
 	mutable TArray<uint8> QueuedUpdates;
 
diff --git a/Engine/Source/Runtime/AIModule/Classes/BehaviorTree/Decorators/BTDecorator_BlackboardBase.h b/Engine/Source/Runtime/AIModule/Classes/BehaviorTree/Decorators/BTDecorator_BlackboardBase.h
--- a/Engine/Source/Runtime/AIModule/Classes/BehaviorTree/Decorators/BTDecorator_BlackboardBase.h
+++ b/Engine/Source/Runtime/AIModule/Classes/BehaviorTree/Decorators/BTDecorator_BlackboardBase.h
@@ -29,12 +29,6 @@ protected:
 	UPROPERTY(EditAnywhere, Category=Blackboard)
 	FBlackboardKeySelector BlackboardKey;
 
-	FOnBlackboardChange BBKeyObserver;
-
-	TMap<FBlackboard::FKey, FDelegateHandle> BBKeyObserverDelegateHandles;
-
-	virtual void PostInitProperties() override;
-
 	/** called when execution flow controller becomes active */
 	virtual void OnBecomeRelevant(UBehaviorTreeComponent& OwnerComp, uint8* NodeMemory) override;
 
diff --git a/Engine/Source/Runtime/AIModule/Classes/BehaviorTree/Decorators/BTDecorator_BlueprintBase.h b/Engine/Source/Runtime/AIModule/Classes/BehaviorTree/Decorators/BTDecorator_BlueprintBase.h
--- a/Engine/Source/Runtime/AIModule/Classes/BehaviorTree/Decorators/BTDecorator_BlueprintBase.h
+++ b/Engine/Source/Runtime/AIModule/Classes/BehaviorTree/Decorators/BTDecorator_BlueprintBase.h
@@ -104,8 +104,6 @@ protected:
 	/** set if ReceiveConditionCheck is implemented by blueprint */
 	uint32 ReceiveConditionCheckImplementations : 2;
 
-	FOnBlackboardChange BBKeyObserver;
-
 	bool CalculateRawConditionValueImpl(UBehaviorTreeComponent& OwnerComp) const;
 	
 	virtual void OnBecomeRelevant(UBehaviorTreeComponent& OwnerComp, uint8* NodeMemory) override;
@@ -200,8 +198,6 @@ protected:
 	UFUNCTION(BlueprintCallable, Category="AI|BehaviorTree")
 	bool IsDecoratorObserverActive() const;
 
-	TMap<FBlackboard::FKey, FDelegateHandle> BBKeyObserverDelegateHandles;
-
 	friend FBehaviorBlueprintDetails;
 
 	FORCEINLINE bool GetNeedsTickForConditionChecking() const { return ReceiveConditionCheckImplementations != 0 && (bIsObservingBB == false || bCheckConditionOnlyBlackBoardChanges == false); }
diff --git a/Engine/Source/Runtime/AIModule/Classes/BehaviorTree/Decorators/BTDecorator_CompareBBEntries.h b/Engine/Source/Runtime/AIModule/Classes/BehaviorTree/Decorators/BTDecorator_CompareBBEntries.h
--- a/Engine/Source/Runtime/AIModule/Classes/BehaviorTree/Decorators/BTDecorator_CompareBBEntries.h
+++ b/Engine/Source/Runtime/AIModule/Classes/BehaviorTree/Decorators/BTDecorator_CompareBBEntries.h
@@ -40,14 +40,8 @@ protected:
 	UPROPERTY(EditAnywhere, Category=Blackboard)
 	FBlackboardKeySelector BlackboardKeyB;
 
-	FOnBlackboardChange BBKeyObserver;
-
-	TMap<FBlackboard::FKey, FDelegateHandle> BBKeyObserverDelegateHandles;
-
 public:
 
-	virtual void PostInitProperties() override;
-
 	virtual bool CalculateRawConditionValue(UBehaviorTreeComponent& OwnerComp, uint8* NodeMemory) const override;
 	virtual FString GetStaticDescription() const override;
 
diff --git a/Engine/Source/Runtime/AIModule/Private/BehaviorTree/BlackboardComponent.cpp b/Engine/Source/Runtime/AIModule/Private/BehaviorTree/BlackboardComponent.cpp
--- a/Engine/Source/Runtime/AIModule/Private/BehaviorTree/BlackboardComponent.cpp
+++ b/Engine/Source/Runtime/AIModule/Private/BehaviorTree/BlackboardComponent.cpp
@@ -244,7 +244,7 @@ int32 UBlackboardComponent::GetNumKeys() const
 	return BlackboardAsset ? BlackboardAsset->GetNumKeys() : 0;
 }
 
-FDelegateHandle UBlackboardComponent::RegisterObserver(FBlackboard::FKey KeyID, FOnBlackboardChange ObserverDelegate)
+FDelegateHandle UBlackboardComponent::RegisterObserver(FBlackboard::FKey KeyID, UObject* NotifyOwner, FOnBlackboardChange ObserverDelegate)
 {
 	for (auto It = Observers.CreateConstKeyIterator(KeyID); It; ++It)
 	{
@@ -255,7 +255,10 @@ FDelegateHandle UBlackboardComponent::RegisterObserver(FBlackboard::FKey KeyID,
 		}
 	}
 
-	return Observers.Add(KeyID, ObserverDelegate).GetHandle();
+	FDelegateHandle Handle = Observers.Add(KeyID, ObserverDelegate).GetHandle();
+	ObserverHandles.Add(NotifyOwner, Handle);
+
+	return Handle;
 }
 
 void UBlackboardComponent::UnregisterObserver(FBlackboard::FKey KeyID, FOnBlackboardChange ObserverDelegate)
@@ -264,8 +267,17 @@ void UBlackboardComponent::UnregisterObserver(FBlackboard::FKey KeyID, FOnBlackb
 	{
 		if (It.Value().DEPRECATED_Compare(ObserverDelegate))
 		{
+			for (auto HandleIt = ObserverHandles.CreateIterator(); HandleIt; ++HandleIt)
+			{
+				if (HandleIt.Value() == It.Value().GetHandle())
+				{
+					HandleIt.RemoveCurrent();
+					break;
+				}
+			}
+
 			It.RemoveCurrent();
-			return;
+			break;
 		}
 	}
 }
@@ -276,9 +288,36 @@ void UBlackboardComponent::UnregisterObserver(FBlackboard::FKey KeyID, FDelegate
 	{
 		if (It.Value().GetHandle() == ObserverHandle)
 		{
+			for (auto HandleIt = ObserverHandles.CreateIterator(); HandleIt; ++HandleIt)
+			{
+				if (HandleIt.Value() == ObserverHandle)
+				{
+					HandleIt.RemoveCurrent();
+					break;
+				}
+			}
+
 			It.RemoveCurrent();
-			return;
+			break;
+		}
+	}
+}
+
+void UBlackboardComponent::UnregisterObserversFrom(UObject* NotifyOwner)
+{
+	for (auto It = ObserverHandles.CreateKeyIterator(NotifyOwner); It; ++It)
+	{
+		for (auto ObsIt = Observers.CreateIterator(); ObsIt; ++ObsIt)
+		{
+			if (ObsIt.Value().GetHandle() == It.Value())
+			{
+				ObsIt.RemoveCurrent();
+				break;
+			}
 		}
+
+		It.RemoveCurrent();
+		// check other delegates from NotifyOwner as well
 	}
 }
 
@@ -368,8 +407,6 @@ FString UBlackboardComponent::GetDebugInfoString(EBlackboardDescription::Type Mo
 		TArray<uint8> ObserversKeys;
 		if (Observers.GetKeys(ObserversKeys) > 0)
 		{
-			DebugString += TEXT("Observed Keys:\n");
-		
 			for (int32 KeyIndex = 0; KeyIndex < ObserversKeys.Num(); ++KeyIndex)
 			{
 				const FBlackboard::FKey KeyID = ObserversKeys[KeyIndex];
diff --git a/Engine/Source/Runtime/AIModule/Private/BehaviorTree/Decorators/BTDecorator_BlackboardBase.cpp b/Engine/Source/Runtime/AIModule/Private/BehaviorTree/Decorators/BTDecorator_BlackboardBase.cpp
--- a/Engine/Source/Runtime/AIModule/Private/BehaviorTree/Decorators/BTDecorator_BlackboardBase.cpp
+++ b/Engine/Source/Runtime/AIModule/Private/BehaviorTree/Decorators/BTDecorator_BlackboardBase.cpp
@@ -14,12 +14,6 @@ UBTDecorator_BlackboardBase::UBTDecorator_BlackboardBase(const FObjectInitialize
 	// empty KeySelector = allow everything
 }
 
-void UBTDecorator_BlackboardBase::PostInitProperties()
-{
-	Super::PostInitProperties();
-	BBKeyObserver = FOnBlackboardChange::CreateUObject(this, &UBTDecorator_BlackboardBase::OnBlackboardChange);
-}
-
 void UBTDecorator_BlackboardBase::InitializeFromAsset(UBehaviorTree& Asset)
 {
 	Super::InitializeFromAsset(Asset);
@@ -32,7 +26,7 @@ void UBTDecorator_BlackboardBase::OnBecomeRelevant(UBehaviorTreeComponent& Owner
 	if (BlackboardComp)
 	{
 		auto KeyID = BlackboardKey.GetSelectedKeyID();
-		BBKeyObserverDelegateHandles.Add(KeyID, BlackboardComp->RegisterObserver(KeyID, BBKeyObserver));
+		BlackboardComp->RegisterObserver(KeyID, this, FOnBlackboardChange::CreateUObject(this, &UBTDecorator_BlackboardBase::OnBlackboardChange));
 	}
 }
 
@@ -41,9 +35,7 @@ void UBTDecorator_BlackboardBase::OnCeaseRelevant(UBehaviorTreeComponent& OwnerC
 	UBlackboardComponent* BlackboardComp = OwnerComp.GetBlackboardComponent();
 	if (BlackboardComp)
 	{
-		auto KeyID = BlackboardKey.GetSelectedKeyID();
-		BlackboardComp->UnregisterObserver(KeyID, BBKeyObserverDelegateHandles.FindRef(KeyID));
-		BBKeyObserverDelegateHandles.Remove(KeyID);
+		BlackboardComp->UnregisterObserversFrom(this);
 	}
 }
 
@@ -63,4 +55,4 @@ FName UBTDecorator_BlackboardBase::GetNodeIconName() const
 	return FName("BTEditor.Graph.BTNode.Decorator.Blackboard.Icon");
 }
 
-#endif	// WITH_EDITOR
\ No newline at end of file
+#endif	// WITH_EDITOR
diff --git a/Engine/Source/Runtime/AIModule/Private/BehaviorTree/Decorators/BTDecorator_BlueprintBase.cpp b/Engine/Source/Runtime/AIModule/Private/BehaviorTree/Decorators/BTDecorator_BlueprintBase.cpp
--- a/Engine/Source/Runtime/AIModule/Private/BehaviorTree/Decorators/BTDecorator_BlueprintBase.cpp
+++ b/Engine/Source/Runtime/AIModule/Private/BehaviorTree/Decorators/BTDecorator_BlueprintBase.cpp
@@ -49,11 +49,6 @@ void UBTDecorator_BlueprintBase::PostInitProperties()
 		
 	InitializeProperties();
 	
-	if (bIsObservingBB)
-	{
-		BBKeyObserver = FOnBlackboardChange::CreateUObject(this, &UBTDecorator_BlueprintBase::OnBlackboardChange);
-	}
-
 	if (ReceiveConditionCheckImplementations || bIsObservingBB)
 	{
 		bNotifyBecomeRelevant = true;
@@ -106,7 +101,7 @@ void UBTDecorator_BlueprintBase::OnBecomeRelevant(UBehaviorTreeComponent& OwnerC
 			const FBlackboard::FKey KeyID = BlackboardComp->GetKeyID(ObservedKeyNames[NameIndex]);
 			if (KeyID != FBlackboard::InvalidKey)
 			{
-				BBKeyObserverDelegateHandles.Add(KeyID, BlackboardComp->RegisterObserver(KeyID, BBKeyObserver));
+				BlackboardComp->RegisterObserver(KeyID, this, FOnBlackboardChange::CreateUObject(this, &UBTDecorator_BlueprintBase::OnBlackboardChange));
 			}
 		}
 	}
@@ -117,15 +112,7 @@ void UBTDecorator_BlueprintBase::OnCeaseRelevant(UBehaviorTreeComponent& OwnerCo
 	UBlackboardComponent* BlackboardComp = OwnerComp.GetBlackboardComponent();
 	if (BlackboardComp)
 	{
-		for (int32 NameIndex = 0; NameIndex < ObservedKeyNames.Num(); NameIndex++)
-		{
-			const FBlackboard::FKey KeyID = BlackboardComp->GetKeyID(ObservedKeyNames[NameIndex]);
-			if (KeyID != FBlackboard::InvalidKey)
-			{
-				BlackboardComp->UnregisterObserver(KeyID, BBKeyObserverDelegateHandles.FindRef(KeyID));
-				BBKeyObserverDelegateHandles.Remove(KeyID);
-			}
-		}
+		BlackboardComp->UnregisterObserversFrom(this);
 	}
 		
 	if (AIOwner != nullptr && ReceiveObserverDeactivatedImplementations & FBTNodeBPImplementationHelper::AISpecific)
diff --git a/Engine/Source/Runtime/AIModule/Private/BehaviorTree/Decorators/BTDecorator_CompareBBEntries.cpp b/Engine/Source/Runtime/AIModule/Private/BehaviorTree/Decorators/BTDecorator_CompareBBEntries.cpp
--- a/Engine/Source/Runtime/AIModule/Private/BehaviorTree/Decorators/BTDecorator_CompareBBEntries.cpp
+++ b/Engine/Source/Runtime/AIModule/Private/BehaviorTree/Decorators/BTDecorator_CompareBBEntries.cpp
@@ -13,12 +13,6 @@ UBTDecorator_CompareBBEntries::UBTDecorator_CompareBBEntries(const FObjectInitia
 	bNotifyCeaseRelevant = true;
 }
 
-void UBTDecorator_CompareBBEntries::PostInitProperties()
-{
-	Super::PostInitProperties();
-	BBKeyObserver = FOnBlackboardChange::CreateUObject(this, &UBTDecorator_CompareBBEntries::OnBlackboardChange);
-}
-
 void UBTDecorator_CompareBBEntries::InitializeFromAsset(UBehaviorTree& Asset)
 {
 	Super::InitializeFromAsset(Asset);
@@ -70,8 +64,8 @@ void UBTDecorator_CompareBBEntries::OnBecomeRelevant(UBehaviorTreeComponent& Own
 	UBlackboardComponent* BlackboardComp = OwnerComp.GetBlackboardComponent();
 	if (BlackboardComp)
 	{
-		BBKeyObserverDelegateHandles.Add(BlackboardKeyA.GetSelectedKeyID(), BlackboardComp->RegisterObserver(BlackboardKeyA.GetSelectedKeyID(), BBKeyObserver));
-		BBKeyObserverDelegateHandles.Add(BlackboardKeyA.GetSelectedKeyID(), BlackboardComp->RegisterObserver(BlackboardKeyB.GetSelectedKeyID(), BBKeyObserver));
+		BlackboardComp->RegisterObserver(BlackboardKeyA.GetSelectedKeyID(), this, FOnBlackboardChange::CreateUObject(this, &UBTDecorator_CompareBBEntries::OnBlackboardChange));
+		BlackboardComp->RegisterObserver(BlackboardKeyB.GetSelectedKeyID(), this, FOnBlackboardChange::CreateUObject(this, &UBTDecorator_CompareBBEntries::OnBlackboardChange));
 	}
 }
 
@@ -80,10 +74,7 @@ void UBTDecorator_CompareBBEntries::OnCeaseRelevant(UBehaviorTreeComponent& Owne
 	UBlackboardComponent* BlackboardComp = OwnerComp.GetBlackboardComponent();
 	if (BlackboardComp)
 	{
-		BlackboardComp->UnregisterObserver(BlackboardKeyA.GetSelectedKeyID(), BBKeyObserverDelegateHandles.FindRef(BlackboardKeyA.GetSelectedKeyID()));
-		BlackboardComp->UnregisterObserver(BlackboardKeyB.GetSelectedKeyID(), BBKeyObserverDelegateHandles.FindRef(BlackboardKeyB.GetSelectedKeyID()));
-		BBKeyObserverDelegateHandles.Remove(BlackboardKeyA.GetSelectedKeyID());
-		BBKeyObserverDelegateHandles.Remove(BlackboardKeyB.GetSelectedKeyID());
+		BlackboardComp->UnregisterObserversFrom(this);
 	}
 }
 
