Expanded a comment on FEnvQueryResult.Owner to be more enlightening #UE4

diff --git a/Engine/Source/Runtime/AIModule/Classes/EnvironmentQuery/EnvQueryTypes.h b/Engine/Source/Runtime/AIModule/Classes/EnvironmentQuery/EnvQueryTypes.h
--- a/Engine/Source/Runtime/AIModule/Classes/EnvironmentQuery/EnvQueryTypes.h
+++ b/Engine/Source/Runtime/AIModule/Classes/EnvironmentQuery/EnvQueryTypes.h
@@ -424,7 +424,7 @@ public:
 	/** instance ID */
 	int32 QueryID;
 
-	/** instance owner */
+	/** instance owner. Mind that it doesn't have to be the query's "Querier". This is just the object that is responsible for this query instance. */
 	TWeakObjectPtr<UObject> Owner;
 
 	FORCEINLINE float GetItemScore(int32 Index) const { return Items.IsValidIndex(Index) ? Items[Index].Score : 0.0f; }
