AI compilation warnings fixed after dev-to-main merge #UE4

diff --git a/Engine/Source/Runtime/AIModule/Private/EnvironmentQuery/Generators/EnvQueryGenerator_Donut.cpp b/Engine/Source/Runtime/AIModule/Private/EnvironmentQuery/Generators/EnvQueryGenerator_Donut.cpp
--- a/Engine/Source/Runtime/AIModule/Private/EnvironmentQuery/Generators/EnvQueryGenerator_Donut.cpp
+++ b/Engine/Source/Runtime/AIModule/Private/EnvironmentQuery/Generators/EnvQueryGenerator_Donut.cpp
@@ -39,7 +39,7 @@ float UEnvQueryGenerator_Donut::GetArcBisectorAngle(FEnvQueryInstance& QueryInst
 
 			if (Start.Num() > 0 && End.Num() > 0)
 			{
-				const FVector LineDir = (End[0] - Start[0]).SafeNormal();
+				const FVector LineDir = (End[0] - Start[0]).GetSafeNormal();
 				const FRotator LineRot = LineDir.Rotation();
 				BisectAngle = LineRot.Yaw;
 			}
