Fixed EQS manager being removed from editor-time AISystem #UE4

This is not a 4.7 fix since the change to when exactly AISystem is being created while editor loads a map is not in 4.7 neither and this is a fallout fix.

diff --git a/Engine/Source/Runtime/AIModule/Private/AISystem.cpp b/Engine/Source/Runtime/AIModule/Private/AISystem.cpp
--- a/Engine/Source/Runtime/AIModule/Private/AISystem.cpp
+++ b/Engine/Source/Runtime/AIModule/Private/AISystem.cpp
@@ -65,10 +65,13 @@ void UAISystem::WorldOriginLocationChanged(FIntVector OldOriginLocation, FIntVec
 
 void UAISystem::CleanupWorld(bool bSessionEnded, bool bCleanupResources, UWorld* NewWorld)
 {
-	if (EnvironmentQueryManager)
+	if (bCleanupResources)
 	{
-		EnvironmentQueryManager->OnWorldCleanup();
-		EnvironmentQueryManager = nullptr;
+		if (EnvironmentQueryManager)
+		{
+			EnvironmentQueryManager->OnWorldCleanup();
+			EnvironmentQueryManager = nullptr;
+		}
 	}
 }
 
