increased threshold of 2D reach test in path following from 0 to 5% of agent's radius UE-7673
low speed was causing some jittering, because agent couldn't overshot current segment and point blank precision can't be achieved

diff --git a/Engine/Source/Runtime/AIModule/Private/Navigation/PathFollowingComponent.cpp b/Engine/Source/Runtime/AIModule/Private/Navigation/PathFollowingComponent.cpp
--- a/Engine/Source/Runtime/AIModule/Private/Navigation/PathFollowingComponent.cpp
+++ b/Engine/Source/Runtime/AIModule/Private/Navigation/PathFollowingComponent.cpp
@@ -881,10 +881,16 @@ bool UPathFollowingComponent::HasReachedCurrentTarget(const FVector& CurrentLoca
 	}
 
 	// or standing at target position
-	// don't use acceptance radius here, it has to be exact for moving near corners
+	// don't use acceptance radius here, it has to be exact for moving near corners (2D test < 5% of agent radius)
 	const float GoalRadius = 0.0f;
 	const float GoalHalfHeight = 0.0f;
-	return HasReachedInternal(CurrentTarget, GoalRadius, GoalHalfHeight, CurrentLocation, CurrentAcceptanceRadius, /*bSuccessOnRadiusOverlap=*/false);
+	const float OrgAgentRadiusPct = MinAgentRadiusPct;
+	((UPathFollowingComponent*)this)->MinAgentRadiusPct = 0.05f;
+
+	const bool bReached = HasReachedInternal(CurrentTarget, GoalRadius, GoalHalfHeight, CurrentLocation, CurrentAcceptanceRadius, /*bSuccessOnRadiusOverlap=*/true);
+
+	((UPathFollowingComponent*)this)->MinAgentRadiusPct = OrgAgentRadiusPct;
+	return bReached;
 }
 
 bool UPathFollowingComponent::HasReachedInternal(const FVector& Goal, float GoalRadius, float GoalHalfHeight, const FVector& AgentLocation, float RadiusThreshold, bool bSuccessOnRadiusOverlap) const
@@ -932,8 +938,8 @@ void UPathFollowingComponent::DebugReachTest(float& CurrentDot, float& CurrentDi
 
 	float GoalRadius = 0.0f;
 	float GoalHalfHeight = 0.0f;
-	bool bUseAgentRadius = false;
 	float RadiusThreshold = 0.0f;
+	float AgentRadiusPct = 0.05f;
 
 	FVector AgentLocation = MovementComp->GetActorFeetLocation();
 	FVector GoalLocation = GetCurrentTargetLocation();
@@ -942,7 +948,7 @@ void UPathFollowingComponent::DebugReachTest(float& CurrentDot, float& CurrentDi
 	if (bFollowingLastSegment)
 	{
 		GoalLocation = Path->GetPathPoints()[Path->GetPathPoints().Num() - 1].Location;
-		bUseAgentRadius = bStopOnOverlap;
+		AgentRadiusPct = MinAgentRadiusPct;
 
 		// take goal's current location, unless path is partial
 		if (DestinationActor.IsValid() && !Path->IsPartial())
@@ -972,7 +978,7 @@ void UPathFollowingComponent::DebugReachTest(float& CurrentDot, float& CurrentDi
 	MovingAgent->GetSimpleCollisionCylinder(AgentRadius, AgentHalfHeight);
 
 	CurrentDistance = ToGoal.Size2D();
-	const float UseRadius = FMath::Max(RadiusThreshold, GoalRadius + (bUseAgentRadius ? AgentRadius * MinAgentRadiusPct : 0.0f));
+	const float UseRadius = FMath::Max(RadiusThreshold, GoalRadius + (AgentRadius * AgentRadiusPct));
 	bDistanceFailed = (CurrentDistance > UseRadius) ? 1 : 0;
 
 	CurrentHeight = FMath::Abs(ToGoal.Z);
