Added a transition API to mitigate fallout after a drastic change in BTDecorator_BlueprintBase's API #UE4

diff --git a/Engine/Source/Runtime/AIModule/Classes/BehaviorTree/Decorators/BTDecorator_BlueprintBase.h b/Engine/Source/Runtime/AIModule/Classes/BehaviorTree/Decorators/BTDecorator_BlueprintBase.h
--- a/Engine/Source/Runtime/AIModule/Classes/BehaviorTree/Decorators/BTDecorator_BlueprintBase.h
+++ b/Engine/Source/Runtime/AIModule/Classes/BehaviorTree/Decorators/BTDecorator_BlueprintBase.h
@@ -102,7 +102,7 @@ protected:
 	uint32 ReceiveObserverDeactivatedImplementations : 2;
 
 	/** set if ReceiveConditionCheck is implemented by blueprint */
-	uint32 ReceiveConditionCheckImplementations : 2;
+	uint32 PerformConditionCheckImplementations : 2;
 
 	bool CalculateRawConditionValueImpl(UBehaviorTreeComponent& OwnerComp) const;
 	
@@ -146,7 +146,7 @@ protected:
 	 *	@Note that if both generic and AI event versions are implemented only the more
 	 *	suitable one will be called, meaning the AI version if called for AI, generic one otherwise */
 	UFUNCTION(BlueprintImplementableEvent)
-	virtual bool ReceiveConditionCheck(AActor* OwnerActor);
+	virtual bool PerformConditionCheck(AActor* OwnerActor);
 
 	/** Alternative AI version of ReceiveTick
 	 *	@see ReceiveTick for more details
@@ -188,7 +188,7 @@ protected:
 	 *	@Note that if both generic and AI event versions are implemented only the more
 	 *	suitable one will be called, meaning the AI version if called for AI, generic one otherwise */
 	UFUNCTION(BlueprintImplementableEvent, Category = AI)
-	virtual bool ReceiveConditionCheckAI(AAIController* OwnerController, APawn* ControlledPawn);
+	virtual bool PerformConditionCheckAI(AAIController* OwnerController, APawn* ControlledPawn);
 		
 	/** check if decorator is part of currently active branch */
 	UFUNCTION(BlueprintCallable, Category="AI|BehaviorTree")
@@ -200,5 +200,17 @@ protected:
 
 	friend FBehaviorBlueprintDetails;
 
-	FORCEINLINE bool GetNeedsTickForConditionChecking() const { return ReceiveConditionCheckImplementations != 0 && (bIsObservingBB == false || bCheckConditionOnlyBlackBoardChanges == false); }
+	FORCEINLINE bool GetNeedsTickForConditionChecking() const { return PerformConditionCheckImplementations != 0 && (bIsObservingBB == false || bCheckConditionOnlyBlackBoardChanges == false); }
+
+	//----------------------------------------------------------------------//
+	// DEPRECATED
+	//----------------------------------------------------------------------//
+	/** finishes condition check */
+	DEPRECATED(4.7, "This function is deprecated. To implement your condition checking implement PerformConditionCheck or PerformConditionCheckAI function in your Blueprint.")
+	UFUNCTION(BlueprintCallable, Category = "AI|BehaviorTree", meta = (DeprecatedFunction, DeprecationMessage = "To implement your condition checking implement PerformConditionCheck or PerformConditionCheckAI function in your Blueprint."))
+	void FinishConditionCheck(bool bAllowExecution);
+
+	DEPRECATED(4.7, "This function is deprecated. To implement your condition checking implement PerformConditionCheck or PerformConditionCheckAI function in your Blueprint.")
+	UFUNCTION(BlueprintImplementableEvent, meta = (DeprecatedFunction, DeprecationMessage = "To implement your condition checking implement PerformConditionCheck or PerformConditionCheckAI function in your Blueprint."))
+	virtual void ReceiveConditionCheck(AActor* OwnerActor);
 };
diff --git a/Engine/Source/Runtime/AIModule/Private/BehaviorTree/Decorators/BTDecorator_BlueprintBase.cpp b/Engine/Source/Runtime/AIModule/Private/BehaviorTree/Decorators/BTDecorator_BlueprintBase.cpp
--- a/Engine/Source/Runtime/AIModule/Private/BehaviorTree/Decorators/BTDecorator_BlueprintBase.cpp
+++ b/Engine/Source/Runtime/AIModule/Private/BehaviorTree/Decorators/BTDecorator_BlueprintBase.cpp
@@ -15,7 +15,7 @@ UBTDecorator_BlueprintBase::UBTDecorator_BlueprintBase(const FObjectInitializer&
 	ReceiveExecutionFinishImplementations = FBTNodeBPImplementationHelper::CheckEventImplementationVersion(TEXT("ReceiveExecutionFinish"), TEXT("ReceiveExecutionFinishAI"), this, StopAtClass);
 	ReceiveObserverActivatedImplementations = FBTNodeBPImplementationHelper::CheckEventImplementationVersion(TEXT("ReceiveObserverActivated"), TEXT("ReceiveObserverActivatedAI"), this, StopAtClass);
 	ReceiveObserverDeactivatedImplementations = FBTNodeBPImplementationHelper::CheckEventImplementationVersion(TEXT("ReceiveObserverDeactivated"), TEXT("ReceiveObserverDeactivatedAI"), this, StopAtClass);
-	ReceiveConditionCheckImplementations = FBTNodeBPImplementationHelper::CheckEventImplementationVersion(TEXT("ReceiveConditionCheck"), TEXT("ReceiveConditionCheckAI"), this, StopAtClass);
+	PerformConditionCheckImplementations = FBTNodeBPImplementationHelper::CheckEventImplementationVersion(TEXT("PerformConditionCheck"), TEXT("PerformConditionCheckAI"), this, StopAtClass);
 
 	bNotifyBecomeRelevant = ReceiveObserverActivatedImplementations != 0;
 	bNotifyCeaseRelevant = ReceiveObserverDeactivatedImplementations != 0;
@@ -49,7 +49,7 @@ void UBTDecorator_BlueprintBase::PostInitProperties()
 		
 	InitializeProperties();
 	
-	if (ReceiveConditionCheckImplementations || bIsObservingBB)
+	if (PerformConditionCheckImplementations || bIsObservingBB)
 	{
 		bNotifyBecomeRelevant = true;
 		bNotifyCeaseRelevant = true;
@@ -176,7 +176,7 @@ void UBTDecorator_BlueprintBase::TickNode(UBehaviorTreeComponent& OwnerComp, uin
 bool UBTDecorator_BlueprintBase::GetShouldAbort(UBehaviorTreeComponent& OwnerComp) const 
 {
 	// if there's no condition-checking function implemented we always want to abort on any change
-	if (ReceiveConditionCheckImplementations == 0)
+	if (PerformConditionCheckImplementations == 0)
 	{
 		return true;
 	}
@@ -199,18 +199,18 @@ bool UBTDecorator_BlueprintBase::GetShouldAbort(UBehaviorTreeComponent& OwnerCom
 bool UBTDecorator_BlueprintBase::CalculateRawConditionValueImpl(UBehaviorTreeComponent& OwnerComp) const
 {
 	bool CurrentCallResult = false;
-	if (ReceiveConditionCheckImplementations != 0)
+	if (PerformConditionCheckImplementations != 0)
 	{
 		// can't use const functions with blueprints
 		UBTDecorator_BlueprintBase* MyNode = (UBTDecorator_BlueprintBase*)this;
 
-		if (AIOwner != nullptr && ReceiveConditionCheckImplementations & FBTNodeBPImplementationHelper::AISpecific)
+		if (AIOwner != nullptr && PerformConditionCheckImplementations & FBTNodeBPImplementationHelper::AISpecific)
 		{
-			CurrentCallResult = MyNode->ReceiveConditionCheckAI(MyNode->AIOwner, MyNode->AIOwner->GetPawn());
+			CurrentCallResult = MyNode->PerformConditionCheckAI(MyNode->AIOwner, MyNode->AIOwner->GetPawn());
 		}
-		else if (ReceiveConditionCheckImplementations & FBTNodeBPImplementationHelper::Generic)
+		else if (PerformConditionCheckImplementations & FBTNodeBPImplementationHelper::Generic)
 		{
-			CurrentCallResult = MyNode->ReceiveConditionCheck(MyNode->ActorOwner);
+			CurrentCallResult = MyNode->PerformConditionCheck(MyNode->ActorOwner);
 		}
 	}
 
@@ -281,7 +281,7 @@ void UBTDecorator_BlueprintBase::OnInstanceDestroyed(UBehaviorTreeComponent& Own
 
 FName UBTDecorator_BlueprintBase::GetNodeIconName() const
 {
-	if(ReceiveConditionCheckImplementations != 0)
+	if(PerformConditionCheckImplementations != 0)
 	{
 		return FName("BTEditor.Graph.BTNode.Decorator.Conditional.Icon");
 	}
@@ -297,3 +297,10 @@ bool UBTDecorator_BlueprintBase::UsesBlueprint() const
 }
 
 #endif // WITH_EDITOR
+
+//----------------------------------------------------------------------//
+// DEPRECATED
+//----------------------------------------------------------------------//
+void UBTDecorator_BlueprintBase::FinishConditionCheck(bool bAllowExecution)
+{
+}
