AI Sense configuration serialization and setting up fixes #UE4

- fixed BP users not being able to set max age for stimuli of given sense

diff --git a/Engine/Source/Runtime/AIModule/Classes/Perception/AIPerceptionComponent.h b/Engine/Source/Runtime/AIModule/Classes/Perception/AIPerceptionComponent.h
--- a/Engine/Source/Runtime/AIModule/Classes/Perception/AIPerceptionComponent.h
+++ b/Engine/Source/Runtime/AIModule/Classes/Perception/AIPerceptionComponent.h
@@ -265,5 +265,6 @@ private:
 	friend UAIPerceptionSystem;
 
 	void StoreListenerId(FPerceptionListenerID InListenerId) { PerceptionListenerId = InListenerId; }
+	void SetMaxStimulusAge(int32 ConfigIndex, float MaxAge);
 };
 
diff --git a/Engine/Source/Runtime/AIModule/Classes/Perception/AISenseConfig.h b/Engine/Source/Runtime/AIModule/Classes/Perception/AISenseConfig.h
--- a/Engine/Source/Runtime/AIModule/Classes/Perception/AISenseConfig.h
+++ b/Engine/Source/Runtime/AIModule/Classes/Perception/AISenseConfig.h
@@ -14,15 +14,24 @@ class AIMODULE_API UAISenseConfig : public UObject
 {
 	GENERATED_BODY()
 
-public:
+protected:
+	/** specifies age limit after stimuli generated by this sense become forgotten. 0 means "never" */
+	UPROPERTY(EditDefaultsOnly, BlueprintReadOnly, Category = "Sense", meta = (ClampMin = "0.0", UIMin = "0.0"))
+	float MaxAge;
+	
+	/** determines whether given sense starts in an enabled state */
 	UPROPERTY(EditDefaultsOnly, BlueprintReadOnly, Category = "Sense")
 	uint32 bStartsEnabled : 1;
-
-	UAISenseConfig(const FObjectInitializer&) : bStartsEnabled(true) {}
+	
+public:
+	UAISenseConfig(const FObjectInitializer& ObjectInitializer);
 
 	virtual TSubclassOf<UAISense> GetSenseImplementation() const { return UAISense::StaticClass(); }
 	FAISenseID GetSenseID() const;
 	
+	float GetMaxAge() const { return MaxAge; }
+	bool IsEnabled() const { return bStartsEnabled; }
+
 #if !UE_BUILD_SHIPPING
 	//----------------------------------------------------------------------//
 	// DEBUG	
diff --git a/Engine/Source/Runtime/AIModule/Classes/Perception/AISenseConfig_Hearing.h b/Engine/Source/Runtime/AIModule/Classes/Perception/AISenseConfig_Hearing.h
--- a/Engine/Source/Runtime/AIModule/Classes/Perception/AISenseConfig_Hearing.h
+++ b/Engine/Source/Runtime/AIModule/Classes/Perception/AISenseConfig_Hearing.h
@@ -29,6 +29,8 @@ public:
 	UPROPERTY(EditDefaultsOnly, BlueprintReadOnly, Category = "Sense", config)
 	FAISenseAffiliationFilter DetectionByAffiliation;
 
+	UAISenseConfig_Hearing(const FObjectInitializer& ObjectInitializer);
+
 	virtual TSubclassOf<UAISense> GetSenseImplementation() const override;
 #if !UE_BUILD_SHIPPING
 	//----------------------------------------------------------------------//
diff --git a/Engine/Source/Runtime/AIModule/Private/Perception/AIPerceptionComponent.cpp b/Engine/Source/Runtime/AIModule/Private/Perception/AIPerceptionComponent.cpp
--- a/Engine/Source/Runtime/AIModule/Private/Perception/AIPerceptionComponent.cpp
+++ b/Engine/Source/Runtime/AIModule/Private/Perception/AIPerceptionComponent.cpp
@@ -47,7 +47,6 @@ UAIPerceptionComponent::UAIPerceptionComponent(const FObjectInitializer& ObjectI
 	: Super(ObjectInitializer)
 	, PerceptionListenerId(FPerceptionListenerID::InvalidID())
 {
-	MaxActiveAge.AddZeroed(FAISenseID::GetSize());
 }
 
 void UAIPerceptionComponent::RequestStimuliListenerUpdate()
@@ -113,12 +112,15 @@ void UAIPerceptionComponent::ConfigureSense(UAISenseConfig& Config)
 		{
 			ConfigIndex = Index;
 			SensesConfig[Index] = &Config;
+			SetMaxStimulusAge(Index, Config.GetMaxAge());
+			break;
 		}
 	}
 
 	if (ConfigIndex == INDEX_NONE)
 	{
 		ConfigIndex = SensesConfig.Add(&Config);
+		SetMaxStimulusAge(ConfigIndex, Config.GetMaxAge());
 	}
 
 	if (IsRegistered())
@@ -128,6 +130,15 @@ void UAIPerceptionComponent::ConfigureSense(UAISenseConfig& Config)
 	// else the sense will be auto-configured during OnRegister
 }
 
+void UAIPerceptionComponent::SetMaxStimulusAge(int32 ConfigIndex, float MaxAge)
+{
+	if (MaxActiveAge.IsValidIndex(ConfigIndex) == false)
+	{
+		MaxActiveAge.AddUninitialized(ConfigIndex - MaxActiveAge.Num() + 1);
+	}
+	MaxActiveAge[ConfigIndex] = MaxAge;
+}
+
 void UAIPerceptionComponent::OnRegister()
 {
 	Super::OnRegister();
@@ -151,16 +162,19 @@ void UAIPerceptionComponent::OnRegister()
 						// make sure it's registered with perception system
 						AIPerceptionSys->RegisterSenseClass(SenseImplementation);
 
-						if (SenseConfig->bStartsEnabled)
+						if (SenseConfig->IsEnabled())
 						{
 							PerceptionFilter.AcceptChannel(UAISense::GetSenseID(SenseImplementation));
 						}
+
+						const FAISenseID SenseID = UAISense::GetSenseID(SenseImplementation);
+						check(SenseID.IsValid());
+						SetMaxStimulusAge(SenseID, SenseConfig->GetMaxAge());
 					}
 				}
 			}
 
 			AIPerceptionSys->UpdateListener(*this);
-
 		}
 	}
 
@@ -504,7 +518,7 @@ bool UAIPerceptionComponent::HasAnyActiveStimulus(const AActor& Source) const
 	{
 		if (Info->LastSensedStimuli[SenseID].WasSuccessfullySensed() &&
 			Info->LastSensedStimuli[SenseID].GetAge() < FAIStimulus::NeverHappenedAge &&
-			Info->LastSensedStimuli[SenseID].GetAge() <= MaxActiveAge[SenseID])
+			(Info->LastSensedStimuli[SenseID].GetAge() <= MaxActiveAge[SenseID] || MaxActiveAge[SenseID] == 0.f))
 		{
 			return true;
 		}
@@ -520,7 +534,7 @@ bool UAIPerceptionComponent::HasActiveStimulus(const AActor& Source, FAISenseID
 		&& Info->LastSensedStimuli.IsValidIndex(Sense) 
 		&& Info->LastSensedStimuli[Sense].WasSuccessfullySensed()
 		&& Info->LastSensedStimuli[Sense].GetAge() < FAIStimulus::NeverHappenedAge
-		&& Info->LastSensedStimuli[Sense].GetAge() <= MaxActiveAge[Sense]);
+		&& (Info->LastSensedStimuli[Sense].GetAge() <= MaxActiveAge[Sense] || MaxActiveAge[Sense] == 0.f));
 }
 
 void UAIPerceptionComponent::RemoveDeadData()
diff --git a/Engine/Source/Runtime/AIModule/Private/Perception/AISense.cpp b/Engine/Source/Runtime/AIModule/Private/Perception/AISense.cpp
--- a/Engine/Source/Runtime/AIModule/Private/Perception/AISense.cpp
+++ b/Engine/Source/Runtime/AIModule/Private/Perception/AISense.cpp
@@ -92,12 +92,23 @@ FString UAISense::GetDebugLegend() const
 {
 	return FString::Printf(TEXT("{%s} %s,"), *DebugDrawColor.ToString(), *GetName());
 }
+#endif // !UE_BUILD_SHIPPING
+
+//----------------------------------------------------------------------//
+// 
+//----------------------------------------------------------------------//
+UAISenseConfig::UAISenseConfig(const FObjectInitializer& ObjectInitializer)
+	: Super(ObjectInitializer), bStartsEnabled(true)
+{
+}
 
+#if !UE_BUILD_SHIPPING
 void UAISenseConfig::DrawDebugInfo(UCanvas& Canvas, UAIPerceptionComponent& PerceptionComponent) const
 {
 
 }
 #endif // !UE_BUILD_SHIPPING
+
 //----------------------------------------------------------------------//
 // 
 //----------------------------------------------------------------------//
@@ -105,6 +116,7 @@ TSubclassOf<UAISense> UAISenseConfig_Sight::GetSenseImplementation() const
 { 
 	return *Implementation; 
 }
+
 #if !UE_BUILD_SHIPPING
 void UAISenseConfig_Sight::DrawDebugInfo(UCanvas& Canvas, UAIPerceptionComponent& PerceptionComponent) const
 {
@@ -125,9 +137,16 @@ void UAISenseConfig_Sight::DrawDebugInfo(UCanvas& Canvas, UAIPerceptionComponent
 	}
 }
 #endif // !UE_BUILD_SHIPPING
+
 //----------------------------------------------------------------------//
 // UAISenseConfig_Hearing
 //----------------------------------------------------------------------//
+
+UAISenseConfig_Hearing::UAISenseConfig_Hearing(const FObjectInitializer& ObjectInitializer) 
+	: Super(ObjectInitializer), HearingRange(3000.f)
+{
+}
+
 TSubclassOf<UAISense> UAISenseConfig_Hearing::GetSenseImplementation() const 
 { 
 	return *Implementation; 
