Fixed game crashing on trying to attach unbound dynamic delegate in AIPerceptionComponent::PostInitProperties #UE4

diff --git a/Engine/Source/Runtime/AIModule/Private/Perception/AIPerceptionComponent.cpp b/Engine/Source/Runtime/AIModule/Private/Perception/AIPerceptionComponent.cpp
--- a/Engine/Source/Runtime/AIModule/Private/Perception/AIPerceptionComponent.cpp
+++ b/Engine/Source/Runtime/AIModule/Private/Perception/AIPerceptionComponent.cpp
@@ -46,6 +46,7 @@ const int32 UAIPerceptionComponent::InitialStimuliToProcessArraySize = 10;
 UAIPerceptionComponent::UAIPerceptionComponent(const FObjectInitializer& ObjectInitializer)
 	: Super(ObjectInitializer)
 	, PerceptionListenerId(FPerceptionListenerID::InvalidID())
+	, bCleanedUp(false)
 {
 }
 
@@ -90,16 +91,6 @@ UAISenseConfig* UAIPerceptionComponent::GetSenseConfig(const FAISenseID& SenseID
 void UAIPerceptionComponent::PostInitProperties() 
 {
 	Super::PostInitProperties();
-
-	if (HasAnyFlags(RF_ClassDefaultObject) == false)
-	{
-		AActor* Owner = GetOwner();
-		if (Owner)
-		{
-			Owner->OnEndPlay.AddDynamic(this, &UAIPerceptionComponent::OnOwnerEndPlay);
-			AIOwner = Cast<AAIController>(Owner);
-		}
-	}
 }
 
 void UAIPerceptionComponent::ConfigureSense(UAISenseConfig& Config)
@@ -143,8 +134,17 @@ void UAIPerceptionComponent::OnRegister()
 {
 	Super::OnRegister();
 
+	bCleanedUp = false;
+
+	AActor* Owner = GetOwner();
+	if (Owner != nullptr)
+	{
+		Owner->OnEndPlay.AddUniqueDynamic(this, &UAIPerceptionComponent::OnOwnerEndPlay);
+		AIOwner = Cast<AAIController>(Owner);
+	}
+
 	UAIPerceptionSystem* AIPerceptionSys = UAIPerceptionSystem::GetCurrent(GetWorld());
-	if (AIPerceptionSys != NULL)
+	if (AIPerceptionSys != nullptr)
 	{
 		PerceptionFilter.Clear();
 
@@ -189,9 +189,8 @@ void UAIPerceptionComponent::OnRegister()
 
 void UAIPerceptionComponent::OnUnregister()
 {
-	Super::OnUnregister();
-
 	CleanUp();
+	Super::OnUnregister();
 }
 
 void UAIPerceptionComponent::OnOwnerEndPlay(EEndPlayReason::Type EndPlayReason)
@@ -201,13 +200,21 @@ void UAIPerceptionComponent::OnOwnerEndPlay(EEndPlayReason::Type EndPlayReason)
 
 void UAIPerceptionComponent::CleanUp()
 {
-	if (IsRegistered())
+	if (bCleanedUp == false)
 	{
 		UAIPerceptionSystem* AIPerceptionSys = UAIPerceptionSystem::GetCurrent(GetWorld());
-		if (AIPerceptionSys != NULL)
+		if (AIPerceptionSys != nullptr)
 		{
 			AIPerceptionSys->UnregisterListener(*this);
 		}
+
+		AActor* Owner = GetOwner();
+		if (Owner != nullptr)
+		{
+			Owner->OnEndPlay.RemoveDynamic(this, &UAIPerceptionComponent::OnOwnerEndPlay);
+		}
+
+		bCleanedUp = true;
 	}
 }
 
