File missing in prev submission #UE4

diff --git a/Engine/Source/Runtime/AIModule/Classes/Perception/AIPerceptionComponent.h b/Engine/Source/Runtime/AIModule/Classes/Perception/AIPerceptionComponent.h
--- a/Engine/Source/Runtime/AIModule/Classes/Perception/AIPerceptionComponent.h
+++ b/Engine/Source/Runtime/AIModule/Classes/Perception/AIPerceptionComponent.h
@@ -151,7 +151,7 @@ protected:
 
 private:
 	TActorPerceptionContainer PerceptualData;
-	
+		
 protected:	
 	struct FStimulusToProcess
 	{
@@ -170,6 +170,9 @@ protected:
 	/** max age of stimulus to consider it "active" (e.g. target is visible) */
 	TArray<float> MaxActiveAge;
 
+private:
+	uint32 bCleanedUp : 1;
+
 public:
 
 	virtual void PostInitProperties() override;
