Added UI min settings for wait time in BTTask_Wait.h #UE4

diff --git a/Engine/Source/Runtime/AIModule/Classes/BehaviorTree/Tasks/BTTask_Wait.h b/Engine/Source/Runtime/AIModule/Classes/BehaviorTree/Tasks/BTTask_Wait.h
--- a/Engine/Source/Runtime/AIModule/Classes/BehaviorTree/Tasks/BTTask_Wait.h
+++ b/Engine/Source/Runtime/AIModule/Classes/BehaviorTree/Tasks/BTTask_Wait.h
@@ -19,7 +19,7 @@ class AIMODULE_API UBTTask_Wait : public UBTTaskNode
 	GENERATED_UCLASS_BODY()
 
 	/** wait time in seconds */
-	UPROPERTY(Category = Wait, EditAnywhere)
+	UPROPERTY(Category = Wait, EditAnywhere, meta = (ClampMin = "0.0", UIMin = "0.0"))
 	float WaitTime;
 
 	/** allows adding random time to wait time */
