Exposed parameters to BT move tasks for controlling whether AI move should finish on AI overlapping the goal location #UE4

diff --git a/Engine/Source/Runtime/AIModule/Classes/BehaviorTree/Tasks/BTTask_MoveDirectlyToward.h b/Engine/Source/Runtime/AIModule/Classes/BehaviorTree/Tasks/BTTask_MoveDirectlyToward.h
--- a/Engine/Source/Runtime/AIModule/Classes/BehaviorTree/Tasks/BTTask_MoveDirectlyToward.h
+++ b/Engine/Source/Runtime/AIModule/Classes/BehaviorTree/Tasks/BTTask_MoveDirectlyToward.h
@@ -34,6 +34,11 @@ class AIMODULE_API UBTTask_MoveDirectlyToward : public UBTTask_BlackboardBase
 	UPROPERTY(Category=Node, EditAnywhere)
 	uint32 bAllowStrafe : 1;
 
+	/** if set to true agent's radius will be added to AcceptableRadius for purposes of checking
+	 *	if path's end point has been reached. Will result in AI stopping on contact with destination location*/
+	UPROPERTY(Category = Node, EditAnywhere)
+	uint32 bStopOnOverlap : 1;
+
 	virtual EBTNodeResult::Type ExecuteTask(UBehaviorTreeComponent& OwnerComp, uint8* NodeMemory) override;
 	virtual EBTNodeResult::Type AbortTask(UBehaviorTreeComponent& OwnerComp, uint8* NodeMemory) override;
 	virtual uint16 GetInstanceMemorySize() const override;
diff --git a/Engine/Source/Runtime/AIModule/Classes/BehaviorTree/Tasks/BTTask_MoveTo.h b/Engine/Source/Runtime/AIModule/Classes/BehaviorTree/Tasks/BTTask_MoveTo.h
--- a/Engine/Source/Runtime/AIModule/Classes/BehaviorTree/Tasks/BTTask_MoveTo.h
+++ b/Engine/Source/Runtime/AIModule/Classes/BehaviorTree/Tasks/BTTask_MoveTo.h
@@ -32,6 +32,11 @@ class AIMODULE_API UBTTask_MoveTo : public UBTTask_BlackboardBase
 
 	UPROPERTY(Category=Node, EditAnywhere)
 	uint32 bAllowStrafe : 1;
+
+	/** if set to true agent's radius will be added to AcceptableRadius for purposes of checking 
+	 *	if path's end point has been reached. Will result in AI stopping on contact with destination location*/
+	UPROPERTY(Category=Node, EditAnywhere)
+	uint32 bStopOnOverlap : 1;
 	
 	virtual EBTNodeResult::Type ExecuteTask(UBehaviorTreeComponent& OwnerComp, uint8* NodeMemory) override;
 	virtual EBTNodeResult::Type AbortTask(UBehaviorTreeComponent& OwnerComp, uint8* NodeMemory) override;
diff --git a/Engine/Source/Runtime/AIModule/Private/BehaviorTree/Tasks/BTTask_MoveDirectlyToward.cpp b/Engine/Source/Runtime/AIModule/Private/BehaviorTree/Tasks/BTTask_MoveDirectlyToward.cpp
--- a/Engine/Source/Runtime/AIModule/Private/BehaviorTree/Tasks/BTTask_MoveDirectlyToward.cpp
+++ b/Engine/Source/Runtime/AIModule/Private/BehaviorTree/Tasks/BTTask_MoveDirectlyToward.cpp
@@ -11,6 +11,7 @@ UBTTask_MoveDirectlyToward::UBTTask_MoveDirectlyToward(const FObjectInitializer&
 	, AcceptableRadius(50.f)
 	, bProjectVectorGoalToNavigation(true)
 	, bAllowStrafe(true)
+	, bStopOnOverlap(false)
 {
 	NodeName = "MoveDirectlyToward";
 
@@ -37,14 +38,14 @@ EBTNodeResult::Type UBTTask_MoveDirectlyToward::ExecuteTask(UBehaviorTreeCompone
 			if (TargetActor)
 			{
 				RequestResult = bDisablePathUpdateOnGoalLocationChange ?
-					MyController->MoveToLocation(TargetActor->GetActorLocation(), AcceptableRadius, /*bStopOnOverlap=*/true, /*bUsePathfinding=*/false, /*bProjectDestinationToNavigation=*/bProjectVectorGoalToNavigation, bAllowStrafe) :
-					MyController->MoveToActor(TargetActor, AcceptableRadius, /*bStopOnOverlap=*/true, /*bUsePathfinding=*/false, bAllowStrafe);
+					MyController->MoveToLocation(TargetActor->GetActorLocation(), AcceptableRadius, bStopOnOverlap, /*bUsePathfinding=*/false, /*bProjectDestinationToNavigation=*/bProjectVectorGoalToNavigation, bAllowStrafe) :
+					MyController->MoveToActor(TargetActor, AcceptableRadius, bStopOnOverlap, /*bUsePathfinding=*/false, bAllowStrafe);
 			}
 		}
 		else if (BlackboardKey.SelectedKeyType == UBlackboardKeyType_Vector::StaticClass())
 		{
 			const FVector TargetLocation = MyBlackboard->GetValue<UBlackboardKeyType_Vector>(BlackboardKey.GetSelectedKeyID());
-			RequestResult = MyController->MoveToLocation(TargetLocation, AcceptableRadius, /*bStopOnOverlap=*/true, /*bUsePathfinding=*/false, /*bProjectDestinationToNavigation=*/bProjectVectorGoalToNavigation, bAllowStrafe);
+			RequestResult = MyController->MoveToLocation(TargetLocation, AcceptableRadius, bStopOnOverlap, /*bUsePathfinding=*/false, /*bProjectDestinationToNavigation=*/bProjectVectorGoalToNavigation, bAllowStrafe);
 		}
 
 		if (RequestResult == EPathFollowingRequestResult::RequestSuccessful)
diff --git a/Engine/Source/Runtime/AIModule/Private/BehaviorTree/Tasks/BTTask_MoveTo.cpp b/Engine/Source/Runtime/AIModule/Private/BehaviorTree/Tasks/BTTask_MoveTo.cpp
--- a/Engine/Source/Runtime/AIModule/Private/BehaviorTree/Tasks/BTTask_MoveTo.cpp
+++ b/Engine/Source/Runtime/AIModule/Private/BehaviorTree/Tasks/BTTask_MoveTo.cpp
@@ -11,6 +11,7 @@ UBTTask_MoveTo::UBTTask_MoveTo(const FObjectInitializer& ObjectInitializer)
 	: Super(ObjectInitializer)
 	, AcceptableRadius(50.f)
 	, bAllowStrafe(false)
+	, bStopOnOverlap(true)
 {
 	NodeName = "Move To";
 	bNotifyTick = true;
@@ -58,7 +59,8 @@ EBTNodeResult::Type UBTTask_MoveTo::PerformMoveTask(UBehaviorTreeComponent& Owne
 			AActor* TargetActor = Cast<AActor>(KeyValue);
 			if (TargetActor)
 			{
-				RequestResult = MyController->MoveToActor(TargetActor, AcceptableRadius, true, true, bAllowStrafe, FilterClass);
+				RequestResult = MyController->MoveToActor(TargetActor, AcceptableRadius, bStopOnOverlap, /*bUsePathfinding=*/true
+					, bAllowStrafe, FilterClass);
 			}
 			else
 			{
@@ -68,7 +70,8 @@ EBTNodeResult::Type UBTTask_MoveTo::PerformMoveTask(UBehaviorTreeComponent& Owne
 		else if (BlackboardKey.SelectedKeyType == UBlackboardKeyType_Vector::StaticClass())
 		{
 			const FVector TargetLocation = MyBlackboard->GetValue<UBlackboardKeyType_Vector>(BlackboardKey.GetSelectedKeyID());
-			RequestResult = MyController->MoveToLocation(TargetLocation, AcceptableRadius, true, true, false, bAllowStrafe, FilterClass);
+			RequestResult = MyController->MoveToLocation(TargetLocation, AcceptableRadius, bStopOnOverlap, /*bUsePathfinding=*/true
+				, /*bProjectDestinationToNavigation*/false, bAllowStrafe, FilterClass);
 		}
 
 		if (RequestResult == EPathFollowingRequestResult::RequestSuccessful)
