Added a function to AIPerceptionComponent for querying for all actors perceived with specified sense (or just all perceived) #UE4

diff --git a/Engine/Source/Runtime/AIModule/Classes/Perception/AIPerceptionComponent.h b/Engine/Source/Runtime/AIModule/Classes/Perception/AIPerceptionComponent.h
--- a/Engine/Source/Runtime/AIModule/Classes/Perception/AIPerceptionComponent.h
+++ b/Engine/Source/Runtime/AIModule/Classes/Perception/AIPerceptionComponent.h
@@ -235,6 +235,10 @@ public:
 	//----------------------------------------------------------------------//
 	UFUNCTION(BlueprintCallable, Category = "AI|Perception")
 	void GetPerceivedHostileActors(TArray<AActor*>& OutActors) const;
+
+	/** If SenseToUse is none all actors perceived in any way will get fetched */
+	UFUNCTION(BlueprintCallable, Category = "AI|Perception")
+	void GetPerceivedActors(TSubclassOf<UAISense> SenseToUse, TArray<AActor*>& OutActors) const;
 	
 	/** Retrieves whatever has been sensed about given actor */
 	UFUNCTION(BlueprintCallable, Category = "AI|Perception")
diff --git a/Engine/Source/Runtime/AIModule/Private/Perception/AIPerceptionComponent.cpp b/Engine/Source/Runtime/AIModule/Private/Perception/AIPerceptionComponent.cpp
--- a/Engine/Source/Runtime/AIModule/Private/Perception/AIPerceptionComponent.cpp
+++ b/Engine/Source/Runtime/AIModule/Private/Perception/AIPerceptionComponent.cpp
@@ -563,6 +563,23 @@ void UAIPerceptionComponent::GetPerceivedHostileActors(TArray<AActor*>& OutActor
 	GetHostileActors(OutActors);
 }
 
+void UAIPerceptionComponent::GetPerceivedActors(TSubclassOf<UAISense> SenseToUse, TArray<AActor*>& OutActors) const
+{
+	const FAISenseID SenseID = UAISense::GetSenseID(SenseToUse);
+
+	OutActors.Reserve(PerceptualData.Num());
+	for (TActorPerceptionContainer::TConstIterator DataIt = GetPerceptualDataConstIterator(); DataIt; ++DataIt)
+	{
+		if (SenseToUse == nullptr || DataIt->Value.IsSenseRegistered(SenseID))
+		{
+			if (DataIt->Value.Target.IsValid())
+			{
+				OutActors.Add(DataIt->Value.Target.Get());
+			}
+		}
+	}
+}
+
 bool UAIPerceptionComponent::GetActorsPerception(AActor* Actor, FActorPerceptionBlueprintInfo& Info)
 {
 	bool bInfoFound = false;
