Enabled detour collision resolution

Base implementation does not perform any resolution itself, it is up to derived classes to deal with the supplied resolved position themselves.

diff --git a/Engine/Source/Runtime/AIModule/Classes/Navigation/CrowdFollowingComponent.h b/Engine/Source/Runtime/AIModule/Classes/Navigation/CrowdFollowingComponent.h
--- a/Engine/Source/Runtime/AIModule/Classes/Navigation/CrowdFollowingComponent.h
+++ b/Engine/Source/Runtime/AIModule/Classes/Navigation/CrowdFollowingComponent.h
@@ -57,6 +57,9 @@ class AIMODULE_API UCrowdFollowingComponent : public UPathFollowingComponent, pu
 	/** pass agent velocity to movement component */
 	virtual void ApplyCrowdAgentVelocity(const FVector& NewVelocity, const FVector& DestPathCorner, bool bTraversingLink);
 
+	/** pass desired position to movement component (after resolving collisions between crowd agents) */
+	virtual void ApplyCrowdAgentPosition(const FVector& NewPosition);
+
 	/** master switch for crowd steering & avoidance */
 	virtual void SuspendCrowdSteering(bool bSuspend);
 
diff --git a/Engine/Source/Runtime/AIModule/Classes/Navigation/CrowdManager.h b/Engine/Source/Runtime/AIModule/Classes/Navigation/CrowdManager.h
--- a/Engine/Source/Runtime/AIModule/Classes/Navigation/CrowdManager.h
+++ b/Engine/Source/Runtime/AIModule/Classes/Navigation/CrowdManager.h
@@ -269,6 +269,10 @@ protected:
 
 	uint32 bPruneStartedOffmeshConnections : 1;
 	uint32 bSingleAreaVisibilityOptimization : 1;
+	
+	/** should crowd simulation resolve collisions between agents? if not, this will be handled by their movement components */
+	UPROPERTY(config, EditAnywhere, Category = Config)
+	uint32 bResolveCollisions : 1;
 
 	/** agents registered in crowd manager */
 	TMap<ICrowdAgentInterface*, FCrowdAgentData> ActiveAgents;
diff --git a/Engine/Source/Runtime/AIModule/Private/Navigation/CrowdFollowingComponent.cpp b/Engine/Source/Runtime/AIModule/Private/Navigation/CrowdFollowingComponent.cpp
--- a/Engine/Source/Runtime/AIModule/Private/Navigation/CrowdFollowingComponent.cpp
+++ b/Engine/Source/Runtime/AIModule/Private/Navigation/CrowdFollowingComponent.cpp
@@ -240,6 +240,11 @@ void UCrowdFollowingComponent::ApplyCrowdAgentVelocity(const FVector& NewVelocit
 	}
 }
 
+void UCrowdFollowingComponent::ApplyCrowdAgentPosition(const FVector& NewPosition)
+{
+	// base implementation does nothing
+}
+
 void UCrowdFollowingComponent::SetCrowdSimulation(bool bEnable)
 {
 	if (bEnableCrowdSimulation == bEnable)
diff --git a/Engine/Source/Runtime/AIModule/Private/Navigation/CrowdManager.cpp b/Engine/Source/Runtime/AIModule/Private/Navigation/CrowdManager.cpp
--- a/Engine/Source/Runtime/AIModule/Private/Navigation/CrowdManager.cpp
+++ b/Engine/Source/Runtime/AIModule/Private/Navigation/CrowdManager.cpp
@@ -24,6 +24,7 @@ DECLARE_CYCLE_STAT(TEXT("Step: proximity"), STAT_AI_Crowd_StepProximityTime, STA
 DECLARE_CYCLE_STAT(TEXT("Step: next point"), STAT_AI_Crowd_StepNextPointTime, STATGROUP_AICrowd);
 DECLARE_CYCLE_STAT(TEXT("Step: steering"), STAT_AI_Crowd_StepSteeringTime, STATGROUP_AICrowd);
 DECLARE_CYCLE_STAT(TEXT("Step: avoidance"), STAT_AI_Crowd_StepAvoidanceTime, STATGROUP_AICrowd);
+DECLARE_CYCLE_STAT(TEXT("Step: collisions"), STAT_AI_Crowd_StepCollisionsTime, STATGROUP_AICrowd);
 DECLARE_CYCLE_STAT(TEXT("Step: components"), STAT_AI_Crowd_StepComponentsTime, STATGROUP_AICrowd);
 DECLARE_CYCLE_STAT(TEXT("Step: navlinks"), STAT_AI_Crowd_StepNavLinkTime, STATGROUP_AICrowd);
 DECLARE_CYCLE_STAT(TEXT("Step: movement"), STAT_AI_Crowd_StepMovementTime, STATGROUP_AICrowd);
@@ -139,6 +140,7 @@ UCrowdManager::UCrowdManager(const FObjectInitializer& ObjectInitializer) : Supe
 	PathOptimizationInterval = 0.5f;
 	bSingleAreaVisibilityOptimization = true;
 	bPruneStartedOffmeshConnections = false;
+	bResolveCollisions = false;
 	
 	FCrowdAvoidanceConfig AvoidanceConfig11;		// 11 samples, ECrowdAvoidanceQuality::Low
 	AvoidanceConfig11.VelocityBias = 0.5f;
@@ -237,6 +239,11 @@ void UCrowdManager::Tick(float DeltaTime)
 				SCOPE_CYCLE_COUNTER(STAT_AI_Crowd_StepAvoidanceTime);
 				DetourCrowd->updateStepAvoidance(DeltaTime, DetourAgentDebug);
 			}
+			if (bResolveCollisions)
+			{
+				SCOPE_CYCLE_COUNTER(STAT_AI_Crowd_StepCollisionsTime);
+				DetourCrowd->updateStepMove(DeltaTime, DetourAgentDebug);
+			}
 			{
 				SCOPE_CYCLE_COUNTER(STAT_AI_Crowd_StepComponentsTime);
 				UpdateAgentPaths();
@@ -755,6 +762,12 @@ void UCrowdManager::ApplyVelocity(UCrowdFollowingComponent* AgentComponent, int3
 
 	const FVector DestPathCorner = Recast2UnrealPoint(RcDestCorner);
 	AgentComponent->ApplyCrowdAgentVelocity(NewVelocity, DestPathCorner, anims->active != 0);
+
+	if (bResolveCollisions)
+	{
+		const FVector NewPosition = Recast2UnrealPoint(ag->npos);
+		AgentComponent->ApplyCrowdAgentPosition(NewPosition);
+	}
 }
 
 void UCrowdManager::UpdateAgentPaths()
