Made AIPerceptionComponent not unregister from owner's endplay delegate if component is being destroyed #UE4

- previous behavior resulted in a crash (check fail) while trying to create delegate for comparison purposes of `RemoveDynamic`

diff --git a/Engine/Source/Runtime/AIModule/Private/Perception/AIPerceptionComponent.cpp b/Engine/Source/Runtime/AIModule/Private/Perception/AIPerceptionComponent.cpp
--- a/Engine/Source/Runtime/AIModule/Private/Perception/AIPerceptionComponent.cpp
+++ b/Engine/Source/Runtime/AIModule/Private/Perception/AIPerceptionComponent.cpp
@@ -208,10 +208,13 @@ void UAIPerceptionComponent::CleanUp()
 			AIPerceptionSys->UnregisterListener(*this);
 		}
 
-		AActor* Owner = GetOwner();
-		if (Owner != nullptr)
+		if (HasAnyFlags(RF_BeginDestroyed) == false)
 		{
-			Owner->OnEndPlay.RemoveDynamic(this, &UAIPerceptionComponent::OnOwnerEndPlay);
+			AActor* Owner = GetOwner();
+			if (Owner != nullptr)
+			{
+				Owner->OnEndPlay.RemoveDynamic(this, &UAIPerceptionComponent::OnOwnerEndPlay);
+			}
 		}
 
 		bCleanedUp = true;
