blackboard keys can be instanced, support for custom data allocations
#ue4 UE-8473

diff --git a/Engine/Source/Runtime/AIModule/Classes/BehaviorTree/BehaviorTreeTypes.h b/Engine/Source/Runtime/AIModule/Classes/BehaviorTree/BehaviorTreeTypes.h
--- a/Engine/Source/Runtime/AIModule/Classes/BehaviorTree/BehaviorTreeTypes.h
+++ b/Engine/Source/Runtime/AIModule/Classes/BehaviorTree/BehaviorTreeTypes.h
@@ -423,8 +423,8 @@ struct AIMODULE_API FBlackboardKeySelector
 	{}
 
 	/** array of allowed types with additional properties (e.g. uobject's base class) 
-	  * EditDefaults is required for FBlackboardSelectorDetails::CacheBlackboardData() */
-	UPROPERTY(transient, EditDefaultsOnly, BlueprintReadWrite, Category = Blackboard)
+	  * EditAnywhere is required for FBlackboardSelectorDetails::CacheBlackboardData() */
+	UPROPERTY(transient, EditAnywhere, BlueprintReadWrite, Category = Blackboard)
 	TArray<UBlackboardKeyType*> AllowedTypes;
 
 	/** name of selected key */
@@ -443,7 +443,7 @@ protected:
 	static_assert(sizeof(uint8) == sizeof(FBlackboard::FKey), "FBlackboardKeySelector::SelectedKeyId should be of FBlackboard::FKey-compatible type.");
 
 	// Requires BlueprintReadWrite so that blueprint creators (using MakeBlackboardKeySelector) can specify whether or not None is Allowed.
-	UPROPERTY(transient, EditDefaultsOnly, BlueprintReadWrite, Category=Blackboard, Meta=(Tooltip=""))
+	UPROPERTY(transient, EditAnywhere, BlueprintReadWrite, Category = Blackboard, Meta = (Tooltip = ""))
 	uint32 bNoneIsAllowedValue:1;
 
 public:
diff --git a/Engine/Source/Runtime/AIModule/Classes/BehaviorTree/Blackboard/BlackboardKeyType.h b/Engine/Source/Runtime/AIModule/Classes/BehaviorTree/Blackboard/BlackboardKeyType.h
--- a/Engine/Source/Runtime/AIModule/Classes/BehaviorTree/Blackboard/BlackboardKeyType.h
+++ b/Engine/Source/Runtime/AIModule/Classes/BehaviorTree/Blackboard/BlackboardKeyType.h
@@ -1,5 +1,6 @@
 // Copyright 1998-2015 Epic Games, Inc. All Rights Reserved.
 #pragma once
+#include "BehaviorTree/BehaviorTreeTypes.h"
 #include "BlackboardKeyType.generated.h"
 
 class UBlackboardComponent;
@@ -67,49 +68,95 @@ namespace BlackboardKeyUtils
 	bool CalculateComparisonResult(EArithmeticKeyOperation::Type Operator, float A, float B);
 }
 
+struct FBlackboardInstancedKeyMemory
+{
+	/** index of instanced key in UBlackboardComponent::InstancedKeys */
+	int32 KeyIdx;
+};
+
 UCLASS(EditInlineNew, Abstract, CollapseCategories, AutoExpandCategories=(Blackboard))
 class AIMODULE_API UBlackboardKeyType : public UObject
 {
 	GENERATED_UCLASS_BODY()
-		
-	/** get ValueSize */
-	FORCEINLINE uint16 GetValueSize() const { return ValueSize; }
 
-	/** get test supported by this type */
-	FORCEINLINE EBlackboardKeyOperation::Type GetTestOperation() const { return SupportedOp; }
+	/** handle dynamic data size */
+	virtual void PreInitialize(UBlackboardComponent& OwnerComp);
 
-	/** initialize memory */
-	virtual void Initialize(uint8* MemoryBlock) const;
+	/** handle instancing if needed */
+	void InitializeKey(UBlackboardComponent& OwnerComp, FBlackboard::FKey KeyID);
 
 	/** does it match settings in filter? */
 	virtual bool IsAllowedByFilter(UBlackboardKeyType* FilterOb) const;
 
-	/** extract location from entry */
-	virtual bool GetLocation(const uint8* MemoryBlock, FVector& Location) const;
+	/** extract location from entry, supports instanced keys */
+	bool WrappedGetLocation(const UBlackboardComponent& OwnerComp, const uint8* MemoryBlock, FVector& Location) const;
 
-	/** extract rotation from entry */
-	virtual bool GetRotation(const uint8* MemoryBlock, FRotator& Rotation) const;
+	/** extract rotation from entry, supports instanced keys */
+	bool WrappedGetRotation(const UBlackboardComponent& OwnerComp, const uint8* MemoryBlock, FRotator& Rotation) const;
 
-	/** sets value to the default */
-	virtual bool Clear(uint8* MemoryBlock) const;
+	/** free value before removing from blackboard, supports instanced keys */
+	void WrappedFree(UBlackboardComponent& OwnerComp, uint8* MemoryBlock);
 
-	/** interprets given two memory blocks as "my" type and compares them */
-	virtual EBlackboardCompare::Type Compare(const uint8* MemoryBlockA, const uint8* MemoryBlockB) const;
+	/** sets value to the default, supports instanced keys */
+	void WrappedClear(const UBlackboardComponent& OwnerComp, uint8* MemoryBlock) const;
 
-	/** various value testing, used by decorators */
-	virtual bool TestBasicOperation(const uint8* MemoryBlock, EBasicKeyOperation::Type Op) const;
-	virtual bool TestArithmeticOperation(const uint8* MemoryBlock, EArithmeticKeyOperation::Type Op, int32 OtherIntValue, float OtherFloatValue) const;
-	virtual bool TestTextOperation(const uint8* MemoryBlock, ETextKeyOperation::Type Op, const FString& OtherString) const;
+	/** check if key has stored value, supports instanced keys */
+	bool WrappedIsEmpty(const UBlackboardComponent& OwnerComp, const uint8* MemoryBlock) const;
+
+	/** various value testing, used by decorators, supports instanced keys */
+	bool WrappedTestBasicOperation(const UBlackboardComponent& OwnerComp, const uint8* MemoryBlock, EBasicKeyOperation::Type Op) const;
+	bool WrappedTestArithmeticOperation(const UBlackboardComponent& OwnerComp, const uint8* MemoryBlock, EArithmeticKeyOperation::Type Op, int32 OtherIntValue, float OtherFloatValue) const;
+	bool WrappedTestTextOperation(const UBlackboardComponent& OwnerComp, const uint8* MemoryBlock, ETextKeyOperation::Type Op, const FString& OtherString) const;
 
 	/** describe params of arithmetic test */
 	virtual FString DescribeArithmeticParam(int32 IntValue, float FloatValue) const;
 
-	/** convert value to text */
-	virtual FString DescribeValue(const uint8* MemoryBlock) const;
+	/** convert value to text, supports instanced keys */
+	FString WrappedDescribeValue(const UBlackboardComponent& OwnerComp, const uint8* MemoryBlock) const;
 
 	/** description of params for property view */
 	virtual FString DescribeSelf() const;
 
+	/** @return key instance if bCreateKeyInstance was set */
+	const UBlackboardKeyType* GetKeyInstance(const UBlackboardComponent& OwnerComp, const uint8* MemoryBlock) const;
+	UBlackboardKeyType* GetKeyInstance(UBlackboardComponent& OwnerComp, const uint8* MemoryBlock) const;
+
+	/** compares two values */
+	virtual EBlackboardCompare::Type CompareValues(const UBlackboardComponent& OwnerComp, const uint8* MemoryBlock,
+		const UBlackboardKeyType* OtherKeyOb, const uint8* OtherMemoryBlock) const;
+
+	/** @return true if key wants to be instanced */
+	bool HasInstance() const;
+
+	/** @return true if this object is instanced key */
+	bool IsInstanced() const;
+
+	/** get ValueSize */
+	uint16 GetValueSize() const;
+
+	/** get test supported by this type */
+	EBlackboardKeyOperation::Type GetTestOperation() const;
+
+	// DEPRECATED FUNCTION
+	DEPRECATED(4.8, "This function now requires an extra UBlackboardComponent parameter.")
+	virtual void Initialize(uint8* MemoryBlock) const;
+	DEPRECATED(4.8, "This function is deprecated. Please use WrappedGetLocation for external access and version with extra UBlackboardComponent parameter for overrides.")
+	virtual bool GetLocation(const uint8* MemoryBlock, FVector& Location) const;
+	DEPRECATED(4.8, "This function is deprecated. Please use WrappedGetLocation for external access and version with extra UBlackboardComponent parameter for overrides.")
+	virtual bool GetRotation(const uint8* MemoryBlock, FRotator& Rotation) const;
+	DEPRECATED(4.8, "This function is deprecated. Please use WrappedClear/WrappedIsEmpty for external access and versions with extra UBlackboardComponent parameter for overrides.")
+	virtual bool Clear(uint8* MemoryBlock) const;
+	DEPRECATED(4.8, "This function is deprecated. Please use CompareValues instead.")
+	virtual EBlackboardCompare::Type Compare(const uint8* MemoryBlockA, const uint8* MemoryBlockB) const;
+	DEPRECATED(4.8, "This function is deprecated. Please use WrappedTestBasicOperation for external access and version with extra UBlackboardComponent parameter for overrides.")
+	virtual bool TestBasicOperation(const uint8* MemoryBlock, EBasicKeyOperation::Type Op) const;
+	DEPRECATED(4.8, "This function is deprecated. Please use WrappedTestArithmeticOperation for external access and version with extra UBlackboardComponent parameter for overrides.")
+	virtual bool TestArithmeticOperation(const uint8* MemoryBlock, EArithmeticKeyOperation::Type Op, int32 OtherIntValue, float OtherFloatValue) const;
+	DEPRECATED(4.8, "This function is deprecated. Please use WrappedTestTextOperation for external access and version with extra UBlackboardComponent parameter for overrides.")
+	virtual bool TestTextOperation(const uint8* MemoryBlock, ETextKeyOperation::Type Op, const FString& OtherString) const;
+	DEPRECATED(4.8, "This function is deprecated. Please use WrappedDescribeValue for external access and version with extra UBlackboardComponent parameter for overrides.")
+	virtual FString DescribeValue(const uint8* MemoryBlock) const;
+
 protected:
 
 	/** size of value for this type */
@@ -118,6 +165,12 @@ protected:
 	/** decorator operation supported with this type */
 	TEnumAsByte<EBlackboardKeyOperation::Type> SupportedOp;
 
+	/** set automatically for node instances */
+	uint8 bIsInstanced : 1;
+
+	/** if set, key will be instanced instead of using memory block */
+	uint8 bCreateKeyInstance : 1;
+
 	/** helper function for reading typed data from memory block */
 	template<typename T>
 	static T GetValueFromMemory(const uint8* MemoryBlock)
@@ -151,8 +204,56 @@ protected:
 	}
 
 	friend UBlackboardComponent;
-	void CopyValue(uint8* To, const uint8* From) const
-	{
-		FMemory::Memcpy(To, From, GetValueSize());
-	}
+	
+	/** copy value from other key, works directly on provided memory/properties */
+	virtual void CopyValues(UBlackboardComponent& OwnerComp, uint8* MemoryBlock, const UBlackboardKeyType* SourceKeyOb, const uint8* SourceBlock);
+
+	/** initialize memory, works directly on provided memory/properties */
+	virtual void InitializeMemory(UBlackboardComponent& OwnerComp, uint8* MemoryBlock);
+
+	/** free value before removing from blackboard, works directly on provided memory/properties */
+	virtual void FreeMemory(UBlackboardComponent& OwnerComp, uint8* MemoryBlock);
+
+	/** extract location from entry, works directly on provided memory/properties */
+	virtual bool GetLocation(const UBlackboardComponent& OwnerComp, const uint8* MemoryBlock, FVector& Location) const;
+	
+	/** extract rotation from entry, works directly on provided memory/properties */
+	virtual bool GetRotation(const UBlackboardComponent& OwnerComp, const uint8* MemoryBlock, FRotator& Rotation) const;
+
+	/** sets value to the default, works directly on provided memory/properties */
+	virtual void Clear(UBlackboardComponent& OwnerComp, uint8* MemoryBlock);
+
+	/** check if key has stored value, works directly on provided memory/properties */
+	virtual bool IsEmpty(const UBlackboardComponent& OwnerComp, const uint8* MemoryBlock) const;
+
+	/** various value testing, works directly on provided memory/properties */
+	virtual bool TestBasicOperation(const UBlackboardComponent& OwnerComp, const uint8* MemoryBlock, EBasicKeyOperation::Type Op) const;
+	virtual bool TestArithmeticOperation(const UBlackboardComponent& OwnerComp, const uint8* MemoryBlock, EArithmeticKeyOperation::Type Op, int32 OtherIntValue, float OtherFloatValue) const;
+	virtual bool TestTextOperation(const UBlackboardComponent& OwnerComp, const uint8* MemoryBlock, ETextKeyOperation::Type Op, const FString& OtherString) const;
+
+	/** convert value to text, works directly on provided memory/properties */
+	virtual FString DescribeValue(const UBlackboardComponent& OwnerComp, const uint8* MemoryBlock) const;
 };
+
+//////////////////////////////////////////////////////////////////////////
+// Inlines
+
+FORCEINLINE uint16 UBlackboardKeyType::GetValueSize() const
+{
+	return ValueSize;
+}
+
+FORCEINLINE EBlackboardKeyOperation::Type UBlackboardKeyType::GetTestOperation() const
+{
+	return SupportedOp; 
+}
+
+FORCEINLINE bool UBlackboardKeyType::HasInstance() const
+{
+	return bCreateKeyInstance;
+}
+
+FORCEINLINE bool UBlackboardKeyType::IsInstanced() const
+{
+	return bIsInstanced;
+}
diff --git a/Engine/Source/Runtime/AIModule/Classes/BehaviorTree/Blackboard/BlackboardKeyType_Bool.h b/Engine/Source/Runtime/AIModule/Classes/BehaviorTree/Blackboard/BlackboardKeyType_Bool.h
--- a/Engine/Source/Runtime/AIModule/Classes/BehaviorTree/Blackboard/BlackboardKeyType_Bool.h
+++ b/Engine/Source/Runtime/AIModule/Classes/BehaviorTree/Blackboard/BlackboardKeyType_Bool.h
@@ -11,10 +11,13 @@ class AIMODULE_API UBlackboardKeyType_Bool : public UBlackboardKeyType
 	typedef bool FDataType;
 	static const FDataType InvalidValue;
 
-	static bool GetValue(const uint8* RawData);
-	static bool SetValue(uint8* RawData, bool bValue);
+	static bool GetValue(const UBlackboardKeyType_Bool* KeyOb, const uint8* RawData);
+	static bool SetValue(UBlackboardKeyType_Bool* KeyOb, uint8* RawData, bool bValue);
 
-	virtual FString DescribeValue(const uint8* RawData) const override;
-	virtual EBlackboardCompare::Type Compare(const uint8* MemoryBlockA, const uint8* MemoryBlockB) const override;
-	virtual bool TestBasicOperation(const uint8* MemoryBlock, EBasicKeyOperation::Type Op) const override;
+	virtual EBlackboardCompare::Type CompareValues(const UBlackboardComponent& OwnerComp, const uint8* MemoryBlock,
+		const UBlackboardKeyType* OtherKeyOb, const uint8* OtherMemoryBlock) const override;
+
+protected:
+	virtual FString DescribeValue(const UBlackboardComponent& OwnerComp, const uint8* RawData) const override;
+	virtual bool TestBasicOperation(const UBlackboardComponent& OwnerComp, const uint8* MemoryBlock, EBasicKeyOperation::Type Op) const override;
 };
diff --git a/Engine/Source/Runtime/AIModule/Classes/BehaviorTree/Blackboard/BlackboardKeyType_Class.h b/Engine/Source/Runtime/AIModule/Classes/BehaviorTree/Blackboard/BlackboardKeyType_Class.h
--- a/Engine/Source/Runtime/AIModule/Classes/BehaviorTree/Blackboard/BlackboardKeyType_Class.h
+++ b/Engine/Source/Runtime/AIModule/Classes/BehaviorTree/Blackboard/BlackboardKeyType_Class.h
@@ -14,12 +14,16 @@ class AIMODULE_API UBlackboardKeyType_Class : public UBlackboardKeyType
 	UPROPERTY(Category=Blackboard, EditDefaultsOnly, meta=(AllowAbstract="1"))
 	UClass* BaseClass;
 	
-	static UClass* GetValue(const uint8* RawData);
-	static bool SetValue(uint8* RawData, UClass* Value);
+	static UClass* GetValue(const UBlackboardKeyType_Class* KeyOb, const uint8* RawData);
+	static bool SetValue(UBlackboardKeyType_Class* KeyOb, uint8* RawData, UClass* Value);
+
+	virtual EBlackboardCompare::Type CompareValues(const UBlackboardComponent& OwnerComp, const uint8* MemoryBlock,
+		const UBlackboardKeyType* OtherKeyOb, const uint8* OtherMemoryBlock) const override;
 
-	virtual FString DescribeValue(const uint8* RawData) const override;
 	virtual FString DescribeSelf() const override;
 	virtual bool IsAllowedByFilter(UBlackboardKeyType* FilterOb) const override;
-	virtual EBlackboardCompare::Type Compare(const uint8* MemoryBlockA, const uint8* MemoryBlockB) const override;
-	virtual bool TestBasicOperation(const uint8* MemoryBlock, EBasicKeyOperation::Type Op) const override;
+
+protected:
+	virtual FString DescribeValue(const UBlackboardComponent& OwnerComp, const uint8* RawData) const override;
+	virtual bool TestBasicOperation(const UBlackboardComponent& OwnerComp, const uint8* MemoryBlock, EBasicKeyOperation::Type Op) const override;
 };
diff --git a/Engine/Source/Runtime/AIModule/Classes/BehaviorTree/Blackboard/BlackboardKeyType_Enum.h b/Engine/Source/Runtime/AIModule/Classes/BehaviorTree/Blackboard/BlackboardKeyType_Enum.h
--- a/Engine/Source/Runtime/AIModule/Classes/BehaviorTree/Blackboard/BlackboardKeyType_Enum.h
+++ b/Engine/Source/Runtime/AIModule/Classes/BehaviorTree/Blackboard/BlackboardKeyType_Enum.h
@@ -14,13 +14,17 @@ class AIMODULE_API UBlackboardKeyType_Enum : public UBlackboardKeyType
 	UPROPERTY(Category=Blackboard, EditDefaultsOnly)
 	UEnum* EnumType;
 
-	static uint8 GetValue(const uint8* RawData);
-	static bool SetValue(uint8* RawData, uint8 Value);
+	static uint8 GetValue(const UBlackboardKeyType_Enum* KeyOb, const uint8* RawData);
+	static bool SetValue(UBlackboardKeyType_Enum* KeyOb, uint8* RawData, uint8 Value);
+
+	virtual EBlackboardCompare::Type CompareValues(const UBlackboardComponent& OwnerComp, const uint8* MemoryBlock,
+		const UBlackboardKeyType* OtherKeyOb, const uint8* OtherMemoryBlock) const override;
 
-	virtual FString DescribeValue(const uint8* RawData) const override;
 	virtual FString DescribeSelf() const override;
-	virtual bool IsAllowedByFilter(UBlackboardKeyType* FilterOb) const override;
-	virtual EBlackboardCompare::Type Compare(const uint8* MemoryBlockA, const uint8* MemoryBlockB) const override;
-	virtual bool TestArithmeticOperation(const uint8* MemoryBlock, EArithmeticKeyOperation::Type Op, int32 OtherIntValue, float OtherFloatValue) const override;
 	virtual FString DescribeArithmeticParam(int32 IntValue, float FloatValue) const override;
+	virtual bool IsAllowedByFilter(UBlackboardKeyType* FilterOb) const override;
+
+protected:
+	virtual FString DescribeValue(const UBlackboardComponent& OwnerComp, const uint8* RawData) const override;
+	virtual bool TestArithmeticOperation(const UBlackboardComponent& OwnerComp, const uint8* MemoryBlock, EArithmeticKeyOperation::Type Op, int32 OtherIntValue, float OtherFloatValue) const override;
 };
diff --git a/Engine/Source/Runtime/AIModule/Classes/BehaviorTree/Blackboard/BlackboardKeyType_Float.h b/Engine/Source/Runtime/AIModule/Classes/BehaviorTree/Blackboard/BlackboardKeyType_Float.h
--- a/Engine/Source/Runtime/AIModule/Classes/BehaviorTree/Blackboard/BlackboardKeyType_Float.h
+++ b/Engine/Source/Runtime/AIModule/Classes/BehaviorTree/Blackboard/BlackboardKeyType_Float.h
@@ -11,11 +11,15 @@ class AIMODULE_API UBlackboardKeyType_Float : public UBlackboardKeyType
 	typedef float FDataType;
 	static const FDataType InvalidValue;
 
-	static float GetValue(const uint8* RawData);
-	static bool SetValue(uint8* RawData, float Value);
+	static float GetValue(const UBlackboardKeyType_Float* KeyOb, const uint8* RawData);
+	static bool SetValue(UBlackboardKeyType_Float* KeyOb, uint8* RawData, float Value);
+
+	virtual EBlackboardCompare::Type CompareValues(const UBlackboardComponent& OwnerComp, const uint8* MemoryBlock,
+		const UBlackboardKeyType* OtherKeyOb, const uint8* OtherMemoryBlock) const override;
 
-	virtual FString DescribeValue(const uint8* RawData) const override;
-	virtual EBlackboardCompare::Type Compare(const uint8* MemoryBlockA, const uint8* MemoryBlockB) const override;
-	virtual bool TestArithmeticOperation(const uint8* MemoryBlock, EArithmeticKeyOperation::Type Op, int32 OtherIntValue, float OtherFloatValue) const override;
 	virtual FString DescribeArithmeticParam(int32 IntValue, float FloatValue) const override;
+
+protected:
+	virtual FString DescribeValue(const UBlackboardComponent& OwnerComp, const uint8* RawData) const override;
+	virtual bool TestArithmeticOperation(const UBlackboardComponent& OwnerComp, const uint8* MemoryBlock, EArithmeticKeyOperation::Type Op, int32 OtherIntValue, float OtherFloatValue) const override;
 };
diff --git a/Engine/Source/Runtime/AIModule/Classes/BehaviorTree/Blackboard/BlackboardKeyType_Int.h b/Engine/Source/Runtime/AIModule/Classes/BehaviorTree/Blackboard/BlackboardKeyType_Int.h
--- a/Engine/Source/Runtime/AIModule/Classes/BehaviorTree/Blackboard/BlackboardKeyType_Int.h
+++ b/Engine/Source/Runtime/AIModule/Classes/BehaviorTree/Blackboard/BlackboardKeyType_Int.h
@@ -11,11 +11,14 @@ class AIMODULE_API UBlackboardKeyType_Int : public UBlackboardKeyType
 	typedef int32 FDataType;
 	static const FDataType InvalidValue;
 
-	static int32 GetValue(const uint8* RawData);
-	static bool SetValue(uint8* RawData, int32 Value);
+	static int32 GetValue(const UBlackboardKeyType_Int* KeyOb, const uint8* RawData);
+	static bool SetValue(UBlackboardKeyType_Int* KeyOb, uint8* RawData, int32 Value);
 
-	virtual FString DescribeValue(const uint8* RawData) const override;
-	virtual EBlackboardCompare::Type Compare(const uint8* MemoryBlockA, const uint8* MemoryBlockB) const override;
-	virtual bool TestArithmeticOperation(const uint8* MemoryBlock, EArithmeticKeyOperation::Type Op, int32 OtherIntValue, float OtherFloatValue) const override;
 	virtual FString DescribeArithmeticParam(int32 IntValue, float FloatValue) const override;
+	virtual EBlackboardCompare::Type CompareValues(const UBlackboardComponent& OwnerComp, const uint8* MemoryBlock,
+		const UBlackboardKeyType* OtherKeyOb, const uint8* OtherMemoryBlock) const override;
+
+protected:
+	virtual FString DescribeValue(const UBlackboardComponent& OwnerComp, const uint8* RawData) const override;
+	virtual bool TestArithmeticOperation(const UBlackboardComponent& OwnerComp, const uint8* MemoryBlock, EArithmeticKeyOperation::Type Op, int32 OtherIntValue, float OtherFloatValue) const override;
 };
diff --git a/Engine/Source/Runtime/AIModule/Classes/BehaviorTree/Blackboard/BlackboardKeyType_Name.h b/Engine/Source/Runtime/AIModule/Classes/BehaviorTree/Blackboard/BlackboardKeyType_Name.h
--- a/Engine/Source/Runtime/AIModule/Classes/BehaviorTree/Blackboard/BlackboardKeyType_Name.h
+++ b/Engine/Source/Runtime/AIModule/Classes/BehaviorTree/Blackboard/BlackboardKeyType_Name.h
@@ -11,10 +11,13 @@ class AIMODULE_API UBlackboardKeyType_Name : public UBlackboardKeyType
 	typedef FName FDataType;
 	static const FDataType InvalidValue;
 
-	static FName GetValue(const uint8* RawData);
-	static bool SetValue(uint8* RawData, const FName& Value);
+	static FName GetValue(const UBlackboardKeyType_Name* KeyOb, const uint8* RawData);
+	static bool SetValue(UBlackboardKeyType_Name* KeyOb, uint8* RawData, const FName& Value);
 
-	virtual FString DescribeValue(const uint8* RawData) const override;
-	virtual EBlackboardCompare::Type Compare(const uint8* MemoryBlockA, const uint8* MemoryBlockB) const override;
-	virtual bool TestTextOperation(const uint8* MemoryBlock, ETextKeyOperation::Type Op, const FString& OtherString) const override;
+	virtual EBlackboardCompare::Type CompareValues(const UBlackboardComponent& OwnerComp, const uint8* MemoryBlock,
+		const UBlackboardKeyType* OtherKeyOb, const uint8* OtherMemoryBlock) const override;
+
+protected:
+	virtual FString DescribeValue(const UBlackboardComponent& OwnerComp, const uint8* RawData) const override;
+	virtual bool TestTextOperation(const UBlackboardComponent& OwnerComp, const uint8* MemoryBlock, ETextKeyOperation::Type Op, const FString& OtherString) const override;
 };
diff --git a/Engine/Source/Runtime/AIModule/Classes/BehaviorTree/Blackboard/BlackboardKeyType_NativeEnum.h b/Engine/Source/Runtime/AIModule/Classes/BehaviorTree/Blackboard/BlackboardKeyType_NativeEnum.h
--- a/Engine/Source/Runtime/AIModule/Classes/BehaviorTree/Blackboard/BlackboardKeyType_NativeEnum.h
+++ b/Engine/Source/Runtime/AIModule/Classes/BehaviorTree/Blackboard/BlackboardKeyType_NativeEnum.h
@@ -21,17 +21,21 @@ class AIMODULE_API UBlackboardKeyType_NativeEnum : public UBlackboardKeyType
 	UPROPERTY()
 	UEnum* EnumType;
 
-	static uint8 GetValue(const uint8* RawData);
-	static bool SetValue(uint8* RawData, uint8 Value);
+	static uint8 GetValue(const UBlackboardKeyType_NativeEnum* KeyOb, const uint8* RawData);
+	static bool SetValue(UBlackboardKeyType_NativeEnum* KeyOb, uint8* RawData, uint8 Value);
+
+	virtual EBlackboardCompare::Type CompareValues(const UBlackboardComponent& OwnerComp, const uint8* MemoryBlock,
+		const UBlackboardKeyType* OtherKeyOb, const uint8* OtherMemoryBlock) const override;
 
-	virtual FString DescribeValue(const uint8* RawData) const override;
-	virtual FString DescribeSelf() const override;
 	virtual bool IsAllowedByFilter(UBlackboardKeyType* FilterOb) const override;
-	virtual EBlackboardCompare::Type Compare(const uint8* MemoryBlockA, const uint8* MemoryBlockB) const override;
-	virtual bool TestArithmeticOperation(const uint8* MemoryBlock, EArithmeticKeyOperation::Type Op, int32 OtherIntValue, float OtherFloatValue) const override;
+	virtual FString DescribeSelf() const override;
 	virtual FString DescribeArithmeticParam(int32 IntValue, float FloatValue) const override;
 
 #if WITH_EDITOR
 	virtual void PostEditChangeProperty(struct FPropertyChangedEvent& PropertyChangedEvent) override;
 #endif
+
+protected:
+	virtual FString DescribeValue(const UBlackboardComponent& OwnerComp, const uint8* RawData) const override;
+	virtual bool TestArithmeticOperation(const UBlackboardComponent& OwnerComp, const uint8* MemoryBlock, EArithmeticKeyOperation::Type Op, int32 OtherIntValue, float OtherFloatValue) const override;
 };
diff --git a/Engine/Source/Runtime/AIModule/Classes/BehaviorTree/Blackboard/BlackboardKeyType_Object.h b/Engine/Source/Runtime/AIModule/Classes/BehaviorTree/Blackboard/BlackboardKeyType_Object.h
--- a/Engine/Source/Runtime/AIModule/Classes/BehaviorTree/Blackboard/BlackboardKeyType_Object.h
+++ b/Engine/Source/Runtime/AIModule/Classes/BehaviorTree/Blackboard/BlackboardKeyType_Object.h
@@ -14,14 +14,18 @@ class AIMODULE_API UBlackboardKeyType_Object : public UBlackboardKeyType
 	UPROPERTY(Category=Blackboard, EditDefaultsOnly, meta=(AllowAbstract="1"))
 	UClass* BaseClass;
 
-	static UObject* GetValue(const uint8* RawData);
-	static bool SetValue(uint8* RawData, UObject* Value);
+	static UObject* GetValue(const UBlackboardKeyType_Object* KeyOb, const uint8* RawData);
+	static bool SetValue(UBlackboardKeyType_Object* KeyOb, uint8* RawData, UObject* Value);
+
+	virtual EBlackboardCompare::Type CompareValues(const UBlackboardComponent& OwnerComp, const uint8* MemoryBlock,
+		const UBlackboardKeyType* OtherKeyOb, const uint8* OtherMemoryBlock) const override;
 
-	virtual FString DescribeValue(const uint8* RawData) const override;
 	virtual FString DescribeSelf() const override;
 	virtual bool IsAllowedByFilter(UBlackboardKeyType* FilterOb) const override;
-	virtual bool GetLocation(const uint8* RawData, FVector& Location) const override;
-	virtual bool GetRotation(const uint8* MemoryBlock, FRotator& Rotation) const override;
-	virtual EBlackboardCompare::Type Compare(const uint8* MemoryBlockA, const uint8* MemoryBlockB) const override;
-	virtual bool TestBasicOperation(const uint8* MemoryBlock, EBasicKeyOperation::Type Op) const override;
+
+protected:
+	virtual FString DescribeValue(const UBlackboardComponent& OwnerComp, const uint8* RawData) const override;
+	virtual bool GetLocation(const UBlackboardComponent& OwnerComp, const uint8* RawData, FVector& Location) const override;
+	virtual bool GetRotation(const UBlackboardComponent& OwnerComp, const uint8* MemoryBlock, FRotator& Rotation) const override;
+	virtual bool TestBasicOperation(const UBlackboardComponent& OwnerComp, const uint8* MemoryBlock, EBasicKeyOperation::Type Op) const override;
 };
diff --git a/Engine/Source/Runtime/AIModule/Classes/BehaviorTree/Blackboard/BlackboardKeyType_Rotator.h b/Engine/Source/Runtime/AIModule/Classes/BehaviorTree/Blackboard/BlackboardKeyType_Rotator.h
--- a/Engine/Source/Runtime/AIModule/Classes/BehaviorTree/Blackboard/BlackboardKeyType_Rotator.h
+++ b/Engine/Source/Runtime/AIModule/Classes/BehaviorTree/Blackboard/BlackboardKeyType_Rotator.h
@@ -10,13 +10,17 @@ class AIMODULE_API UBlackboardKeyType_Rotator : public UBlackboardKeyType
 	typedef FRotator FDataType; 
 	static const FDataType InvalidValue;
 
-	static FRotator GetValue(const uint8* RawData);
-	static bool SetValue(uint8* RawData, const FRotator& Value);
+	static FRotator GetValue(const UBlackboardKeyType_Rotator* KeyOb, const uint8* RawData);
+	static bool SetValue(UBlackboardKeyType_Rotator* KeyOb, uint8* RawData, const FRotator& Value);
 
-	virtual void Initialize(uint8* RawData) const override;
-	virtual FString DescribeValue(const uint8* RawData) const override;
-	virtual bool GetRotation(const uint8* RawData, FRotator& Rotation) const override;
-	virtual bool Clear(uint8* MemoryBlock) const override;
-	virtual EBlackboardCompare::Type Compare(const uint8* MemoryBlockA, const uint8* MemoryBlockB) const override;
-	virtual bool TestBasicOperation(const uint8* MemoryBlock, EBasicKeyOperation::Type Op) const override;
+	virtual EBlackboardCompare::Type CompareValues(const UBlackboardComponent& OwnerComp, const uint8* MemoryBlock,
+		const UBlackboardKeyType* OtherKeyOb, const uint8* OtherMemoryBlock) const override;
+
+protected:
+	virtual void InitializeMemory(UBlackboardComponent& OwnerComp, uint8* RawData) override;
+	virtual FString DescribeValue(const UBlackboardComponent& OwnerComp, const uint8* RawData) const override;
+	virtual bool GetRotation(const UBlackboardComponent& OwnerComp, const uint8* RawData, FRotator& Rotation) const override;
+	virtual bool IsEmpty(const UBlackboardComponent& OwnerComp, const uint8* RawData) const override;
+	virtual void Clear(UBlackboardComponent& OwnerComp, uint8* RawData) override;
+	virtual bool TestBasicOperation(const UBlackboardComponent& OwnerComp, const uint8* MemoryBlock, EBasicKeyOperation::Type Op) const override;
 };
diff --git a/Engine/Source/Runtime/AIModule/Classes/BehaviorTree/Blackboard/BlackboardKeyType_String.h b/Engine/Source/Runtime/AIModule/Classes/BehaviorTree/Blackboard/BlackboardKeyType_String.h
--- a/Engine/Source/Runtime/AIModule/Classes/BehaviorTree/Blackboard/BlackboardKeyType_String.h
+++ b/Engine/Source/Runtime/AIModule/Classes/BehaviorTree/Blackboard/BlackboardKeyType_String.h
@@ -11,10 +11,19 @@ class AIMODULE_API UBlackboardKeyType_String : public UBlackboardKeyType
 	typedef FString FDataType;
 	static const FDataType InvalidValue;
 
-	static FString GetValue(const uint8* RawData);
-	static bool SetValue(uint8* RawData, const FString& Value);
+	static FString GetValue(const UBlackboardKeyType_String* KeyOb, const uint8* RawData);
+	static bool SetValue(UBlackboardKeyType_String* KeyOb, uint8* RawData, const FString& Value);
 
-	virtual FString DescribeValue(const uint8* RawData) const override;
-	virtual EBlackboardCompare::Type Compare(const uint8* MemoryBlockA, const uint8* MemoryBlockB) const override;
-	virtual bool TestTextOperation(const uint8* MemoryBlock, ETextKeyOperation::Type Op, const FString& OtherString) const override;
+	virtual EBlackboardCompare::Type CompareValues(const UBlackboardComponent& OwnerComp, const uint8* MemoryBlock,
+		const UBlackboardKeyType* OtherKeyOb, const uint8* OtherMemoryBlock) const override;
+
+	UPROPERTY()
+	FString StringValue;
+
+protected:
+	virtual FString DescribeValue(const UBlackboardComponent& OwnerComp, const uint8* RawData) const override;
+	virtual bool TestTextOperation(const UBlackboardComponent& OwnerComp, const uint8* MemoryBlock, ETextKeyOperation::Type Op, const FString& OtherString) const override;
+	virtual void Clear(UBlackboardComponent& OwnerComp, uint8* MemoryBlock) override;
+	virtual bool IsEmpty(const UBlackboardComponent& OwnerComp, const uint8* MemoryBlock) const override;
+	virtual void CopyValues(UBlackboardComponent& OwnerComp, uint8* MemoryBlock, const UBlackboardKeyType* SourceKeyOb, const uint8* SourceBlock) override;
 };
diff --git a/Engine/Source/Runtime/AIModule/Classes/BehaviorTree/Blackboard/BlackboardKeyType_Vector.h b/Engine/Source/Runtime/AIModule/Classes/BehaviorTree/Blackboard/BlackboardKeyType_Vector.h
--- a/Engine/Source/Runtime/AIModule/Classes/BehaviorTree/Blackboard/BlackboardKeyType_Vector.h
+++ b/Engine/Source/Runtime/AIModule/Classes/BehaviorTree/Blackboard/BlackboardKeyType_Vector.h
@@ -11,13 +11,17 @@ class AIMODULE_API UBlackboardKeyType_Vector : public UBlackboardKeyType
 	typedef FVector FDataType; 
 	static const FDataType InvalidValue;
 	
-	static FVector GetValue(const uint8* RawData);
-	static bool SetValue(uint8* RawData, const FVector& Value);
+	static FVector GetValue(const UBlackboardKeyType_Vector* KeyOb, const uint8* RawData);
+	static bool SetValue(UBlackboardKeyType_Vector* KeyOb, uint8* RawData, const FVector& Value);
 
-	virtual void Initialize(uint8* RawData) const override;
-	virtual FString DescribeValue(const uint8* RawData) const override;
-	virtual bool GetLocation(const uint8* RawData, FVector& Location) const override;
-	virtual bool Clear(uint8* RawData) const override;
-	virtual EBlackboardCompare::Type Compare(const uint8* MemoryBlockA, const uint8* MemoryBlockB) const override;
-	virtual bool TestBasicOperation(const uint8* MemoryBlock, EBasicKeyOperation::Type Op) const override;
+	virtual EBlackboardCompare::Type CompareValues(const UBlackboardComponent& OwnerComp, const uint8* MemoryBlock,
+		const UBlackboardKeyType* OtherKeyOb, const uint8* OtherMemoryBlock) const override;
+
+protected:
+	virtual void InitializeMemory(UBlackboardComponent& OwnerComp, uint8* RawData) override;
+	virtual FString DescribeValue(const UBlackboardComponent& OwnerComp, const uint8* RawData) const override;
+	virtual bool GetLocation(const UBlackboardComponent& OwnerComp, const uint8* RawData, FVector& Location) const override;
+	virtual bool IsEmpty(const UBlackboardComponent& OwnerComp, const uint8* MemoryBlock) const override;
+	virtual void Clear(UBlackboardComponent& OwnerComp, uint8* MemoryBlock) override;
+	virtual bool TestBasicOperation(const UBlackboardComponent& OwnerComp, const uint8* MemoryBlock, EBasicKeyOperation::Type Op) const override;
 };
diff --git a/Engine/Source/Runtime/AIModule/Classes/BehaviorTree/BlackboardComponent.h b/Engine/Source/Runtime/AIModule/Classes/BehaviorTree/BlackboardComponent.h
--- a/Engine/Source/Runtime/AIModule/Classes/BehaviorTree/BlackboardComponent.h
+++ b/Engine/Source/Runtime/AIModule/Classes/BehaviorTree/BlackboardComponent.h
@@ -212,9 +212,8 @@ public:
 	virtual void DescribeSelfToVisLog(struct FVisualLogEntry* Snapshot) const;
 #endif
 
-	//----------------------------------------------------------------------//
-	// DEPRECATED
-	//----------------------------------------------------------------------//
+	//////////////////////////////////////////////////////////////////////////
+	// Deprecated functions
 	DEPRECATED(4.7, "This function is deprecated. Please use ClearValue() instead.")
 	void ClearValueAsRotator(FBlackboard::FKey KeyID);
 	DEPRECATED(4.7, "This function is deprecated. Please use ClearValue() instead.")
@@ -278,6 +277,9 @@ protected:
 	/** offsets in ValueMemory for each key */
 	TArray<uint16> ValueOffsets;
 
+	/** instanced keys with custom data allocations */
+	TArray<UBlackboardKeyType*> KeyInstances;
+
 	/** observers registered for blackboard keys */
 	TMultiMap<uint8, FOnBlackboardChange> Observers;
 
@@ -299,10 +301,15 @@ protected:
 	/** initializes parent chain in asset */
 	void InitializeParentChain(UBlackboardData* NewAsset);
 
+	/** destroy allocated values */
+	void DestroyValues();
+
 	/** populates BB's synchronized entries */
 	void PopulateSynchronizedKeys();
 
 	bool ShouldSyncWithBlackboard(UBlackboardComponent& OtherBlackboardComponent) const;
+
+	friend UBlackboardKeyType;
 };
 
 //////////////////////////////////////////////////////////////////////////
@@ -323,15 +330,18 @@ bool UBlackboardComponent::SetValue(const FName& KeyName, typename TDataClass::F
 template<class TDataClass>
 bool UBlackboardComponent::SetValue(FBlackboard::FKey KeyID, typename TDataClass::FDataType Value)
 {
-	if (GetKeyType(KeyID) != TDataClass::StaticClass())
+	const FBlackboardEntry* EntryInfo = BlackboardAsset ? BlackboardAsset->GetKey(KeyID) : nullptr;
+	if ((EntryInfo == nullptr) || (EntryInfo->KeyType == nullptr) || !EntryInfo->KeyType->IsA(TDataClass::StaticClass()))
 	{
 		return false;
 	}
 
-	uint8* RawData = GetKeyRawData(KeyID);
+	const uint16 DataOffset = EntryInfo->KeyType->IsInstanced() ? sizeof(FBlackboardInstancedKeyMemory) : 0;
+	uint8* RawData = GetKeyRawData(KeyID) + DataOffset;
 	if (RawData)
 	{
-		const bool bChanged = TDataClass::SetValue(RawData, Value);
+		UBlackboardKeyType* KeyOb = EntryInfo->KeyType->IsInstanced() ? KeyInstances[KeyID] : EntryInfo->KeyType;
+		const bool bChanged = TDataClass::SetValue((TDataClass*)KeyOb, RawData, Value);
 		if (bChanged)
 		{
 			NotifyObservers(KeyID);
@@ -343,7 +353,10 @@ bool UBlackboardComponent::SetValue(FBlackboard::FKey KeyID, typename TDataClass
 					UBlackboardComponent* OtherBlackboard = Iter.Value();
 					if (OtherBlackboard != nullptr && ShouldSyncWithBlackboard(*OtherBlackboard))
 					{
-						TDataClass::SetValue(OtherBlackboard->GetKeyRawData(KeyID), Value);
+						UBlackboardKeyType* OtherKeyOb = EntryInfo->KeyType->IsInstanced() ? OtherBlackboard->KeyInstances[KeyID] : EntryInfo->KeyType;
+						uint8* OtherRawData = OtherBlackboard->GetKeyRawData(KeyID) + DataOffset;
+
+						TDataClass::SetValue((TDataClass*)OtherKeyOb, OtherRawData, Value);
 						OtherBlackboard->NotifyObservers(KeyID);
 					}
 				}
@@ -360,14 +373,21 @@ template<class TDataClass>
 typename TDataClass::FDataType UBlackboardComponent::GetValue(const FName& KeyName) const
 {
 	const FBlackboard::FKey KeyID = GetKeyID(KeyName);
-	return (GetKeyType(KeyID) != TDataClass::StaticClass())
-		? TDataClass::InvalidValue
-		: GetValue<TDataClass>(KeyID);
+	return GetValue<TDataClass>(KeyID);
 }
 
 template<class TDataClass>
 typename TDataClass::FDataType UBlackboardComponent::GetValue(FBlackboard::FKey KeyID) const
 {
-	const uint8* RawData = GetKeyRawData(KeyID);
-	return RawData ? TDataClass::GetValue(RawData) : TDataClass::InvalidValue;
+	const FBlackboardEntry* EntryInfo = BlackboardAsset ? BlackboardAsset->GetKey(KeyID) : nullptr;
+	if ((EntryInfo == nullptr) || (EntryInfo->KeyType == nullptr) || !EntryInfo->KeyType->IsA(TDataClass::StaticClass()))
+	{
+		return TDataClass::InvalidValue;
+	}
+
+	UBlackboardKeyType* KeyOb = EntryInfo->KeyType->IsInstanced() ? KeyInstances[KeyID] : EntryInfo->KeyType;
+	const uint16 DataOffset = EntryInfo->KeyType->IsInstanced() ? sizeof(FBlackboardInstancedKeyMemory) : 0;
+
+	const uint8* RawData = GetKeyRawData(KeyID) + DataOffset;
+	return RawData ? TDataClass::GetValue((TDataClass*)KeyOb, RawData) : TDataClass::InvalidValue;
 }
diff --git a/Engine/Source/Runtime/AIModule/Private/BehaviorTree/Blackboard/BlackboardKeyType.cpp b/Engine/Source/Runtime/AIModule/Private/BehaviorTree/Blackboard/BlackboardKeyType.cpp
--- a/Engine/Source/Runtime/AIModule/Private/BehaviorTree/Blackboard/BlackboardKeyType.cpp
+++ b/Engine/Source/Runtime/AIModule/Private/BehaviorTree/Blackboard/BlackboardKeyType.cpp
@@ -7,9 +7,215 @@ UBlackboardKeyType::UBlackboardKeyType(const FObjectInitializer& ObjectInitializ
 {
 	ValueSize = 0;
 	SupportedOp = EBlackboardKeyOperation::Basic;
+	bCreateKeyInstance = false;
+	bIsInstanced = false;
 }
 
-FString UBlackboardKeyType::DescribeValue(const uint8* RawData) const
+void UBlackboardKeyType::PreInitialize(UBlackboardComponent& OwnerComp)
+{
+	// empty in base class
+}
+
+void UBlackboardKeyType::InitializeKey(UBlackboardComponent& OwnerComp, FBlackboard::FKey KeyID)
+{
+	uint8* RawData = OwnerComp.GetKeyRawData(KeyID);
+
+	if (bCreateKeyInstance)
+	{
+		FBlackboardInstancedKeyMemory* MyMemory = (FBlackboardInstancedKeyMemory*)RawData;
+		UBlackboardKeyType* KeyInstance = ConstructObject<UBlackboardKeyType>(GetClass(), &OwnerComp, GetFName(), RF_NoFlags, (UObject*)(this));
+		MyMemory->KeyIdx = KeyID;
+		OwnerComp.KeyInstances[KeyID] = KeyInstance;
+
+		uint8* InstanceMemoryBlock = RawData + sizeof(FBlackboardInstancedKeyMemory);
+		KeyInstance->InitializeMemory(OwnerComp, InstanceMemoryBlock);
+	}
+	else
+	{
+		InitializeMemory(OwnerComp, RawData);
+	}
+}
+
+UBlackboardKeyType* UBlackboardKeyType::GetKeyInstance(UBlackboardComponent& OwnerComp, const uint8* MemoryBlock) const
+{
+	FBlackboardInstancedKeyMemory* MyMemory = (FBlackboardInstancedKeyMemory*)MemoryBlock;
+	return MyMemory && OwnerComp.KeyInstances.IsValidIndex(MyMemory->KeyIdx) ?
+		OwnerComp.KeyInstances[MyMemory->KeyIdx] : NULL;
+}
+
+const UBlackboardKeyType* UBlackboardKeyType::GetKeyInstance(const UBlackboardComponent& OwnerComp, const uint8* MemoryBlock) const
+{
+	FBlackboardInstancedKeyMemory* MyMemory = (FBlackboardInstancedKeyMemory*)MemoryBlock;
+	return MyMemory && OwnerComp.KeyInstances.IsValidIndex(MyMemory->KeyIdx) ?
+		OwnerComp.KeyInstances[MyMemory->KeyIdx] : NULL;
+}
+
+void UBlackboardKeyType::InitializeMemory(UBlackboardComponent& OwnerComp, uint8* MemoryBlock)
+{
+	// empty in base class
+}
+
+void UBlackboardKeyType::WrappedFree(UBlackboardComponent& OwnerComp, uint8* MemoryBlock)
+{
+	if (HasInstance())
+	{
+		UBlackboardKeyType* InstancedKey = GetKeyInstance(OwnerComp, MemoryBlock);
+		uint8* InstanceMemoryBlock = MemoryBlock + sizeof(FBlackboardInstancedKeyMemory);
+		InstancedKey->FreeMemory(OwnerComp, InstanceMemoryBlock);
+	}
+
+	return FreeMemory(OwnerComp, MemoryBlock);
+}
+
+void UBlackboardKeyType::FreeMemory(UBlackboardComponent& OwnerComp, uint8* MemoryBlock)
+{
+	// empty in base class
+}
+
+bool UBlackboardKeyType::WrappedGetLocation(const UBlackboardComponent& OwnerComp, const uint8* MemoryBlock, FVector& Location) const
+{
+	if (HasInstance())
+	{
+		const UBlackboardKeyType* InstancedKey = GetKeyInstance(OwnerComp, MemoryBlock);
+		const uint8* InstanceMemoryBlock = MemoryBlock + sizeof(FBlackboardInstancedKeyMemory);
+		return InstancedKey->GetLocation(OwnerComp, InstanceMemoryBlock, Location);
+	}
+
+	return GetLocation(OwnerComp, MemoryBlock, Location);
+}
+
+bool UBlackboardKeyType::GetLocation(const UBlackboardComponent& OwnerComp, const uint8* RawData, FVector& Location) const
+{
+	return false;
+}
+
+bool UBlackboardKeyType::WrappedGetRotation(const UBlackboardComponent& OwnerComp, const uint8* MemoryBlock, FRotator& Rotation) const
+{
+	if (HasInstance())
+	{
+		const UBlackboardKeyType* InstancedKey = GetKeyInstance(OwnerComp, MemoryBlock);
+		const uint8* InstanceMemoryBlock = MemoryBlock + sizeof(FBlackboardInstancedKeyMemory);
+		return InstancedKey->GetRotation(OwnerComp, InstanceMemoryBlock, Rotation);
+	}
+
+	return GetRotation(OwnerComp, MemoryBlock, Rotation);
+}
+
+bool UBlackboardKeyType::GetRotation(const UBlackboardComponent& OwnerComp, const uint8* MemoryBlock, FRotator& Rotation) const
+{
+	return false;
+}
+
+void UBlackboardKeyType::WrappedClear(const UBlackboardComponent& OwnerComp, uint8* MemoryBlock) const
+{
+	UBlackboardComponent& MutableOwner = (UBlackboardComponent&)OwnerComp;
+	if (HasInstance())
+	{
+		UBlackboardKeyType* InstancedKey = GetKeyInstance(MutableOwner, MemoryBlock);
+		uint8* InstanceMemoryBlock = MemoryBlock + sizeof(FBlackboardInstancedKeyMemory);
+		InstancedKey->Clear(MutableOwner, InstanceMemoryBlock);
+	}
+	else
+	{
+		UBlackboardKeyType* MutableKey = (UBlackboardKeyType*)this;
+		MutableKey->Clear(MutableOwner, MemoryBlock);
+	}
+}
+
+void UBlackboardKeyType::Clear(UBlackboardComponent& OwnerComp, uint8* MemoryBlock)
+{
+	FMemory::Memzero(MemoryBlock, GetValueSize());
+}
+
+bool UBlackboardKeyType::WrappedIsEmpty(const UBlackboardComponent& OwnerComp, const uint8* MemoryBlock) const
+{
+	if (HasInstance())
+	{
+		const UBlackboardKeyType* InstancedKey = GetKeyInstance(OwnerComp, MemoryBlock);
+		const uint8* InstanceMemoryBlock = MemoryBlock + sizeof(FBlackboardInstancedKeyMemory);
+		return InstancedKey->IsEmpty(OwnerComp, InstanceMemoryBlock);
+	}
+
+	return IsEmpty(OwnerComp, MemoryBlock);
+}
+
+bool UBlackboardKeyType::IsEmpty(const UBlackboardComponent& OwnerComp, const uint8* MemoryBlock) const
+{
+	for (int32 ByteIndex = 0; ByteIndex < GetValueSize(); ++ByteIndex)
+	{
+		if (MemoryBlock[ByteIndex] != uint8(0))
+		{
+			return false;
+		}
+	}
+
+	return true;
+}
+
+bool UBlackboardKeyType::WrappedTestBasicOperation(const UBlackboardComponent& OwnerComp, const uint8* MemoryBlock, EBasicKeyOperation::Type Op) const
+{
+	if (HasInstance())
+	{
+		const UBlackboardKeyType* InstancedKey = GetKeyInstance(OwnerComp, MemoryBlock);
+		const uint8* InstanceMemoryBlock = MemoryBlock + sizeof(FBlackboardInstancedKeyMemory);
+		return InstancedKey->TestBasicOperation(OwnerComp, InstanceMemoryBlock, Op);
+	}
+
+	return TestBasicOperation(OwnerComp, MemoryBlock, Op);
+}
+
+bool UBlackboardKeyType::TestBasicOperation(const UBlackboardComponent& OwnerComp, const uint8* MemoryBlock, EBasicKeyOperation::Type Op) const
+{
+	return false;
+}
+
+bool UBlackboardKeyType::WrappedTestArithmeticOperation(const UBlackboardComponent& OwnerComp, const uint8* MemoryBlock, EArithmeticKeyOperation::Type Op, int32 OtherIntValue, float OtherFloatValue) const
+{
+	if (HasInstance())
+	{
+		const UBlackboardKeyType* InstancedKey = GetKeyInstance(OwnerComp, MemoryBlock);
+		const uint8* InstanceMemoryBlock = MemoryBlock + sizeof(FBlackboardInstancedKeyMemory);
+		return InstancedKey->TestArithmeticOperation(OwnerComp, InstanceMemoryBlock, Op, OtherIntValue, OtherFloatValue);
+	}
+
+	return TestArithmeticOperation(OwnerComp, MemoryBlock, Op, OtherIntValue, OtherFloatValue);
+}
+
+bool UBlackboardKeyType::TestArithmeticOperation(const UBlackboardComponent& OwnerComp, const uint8* MemoryBlock, EArithmeticKeyOperation::Type Op, int32 OtherIntValue, float OtherFloatValue) const
+{
+	return false;
+}
+
+bool UBlackboardKeyType::WrappedTestTextOperation(const UBlackboardComponent& OwnerComp, const uint8* MemoryBlock, ETextKeyOperation::Type Op, const FString& OtherString) const
+{
+	if (HasInstance())
+	{
+		const UBlackboardKeyType* InstancedKey = GetKeyInstance(OwnerComp, MemoryBlock);
+		const uint8* InstanceMemoryBlock = MemoryBlock + sizeof(FBlackboardInstancedKeyMemory);
+		return InstancedKey->TestTextOperation(OwnerComp, InstanceMemoryBlock, Op, OtherString);
+	}
+
+	return TestTextOperation(OwnerComp, MemoryBlock, Op, OtherString);
+}
+
+bool UBlackboardKeyType::TestTextOperation(const UBlackboardComponent& OwnerComp, const uint8* MemoryBlock, ETextKeyOperation::Type Op, const FString& OtherString) const
+{
+	return false;
+}
+
+FString UBlackboardKeyType::WrappedDescribeValue(const UBlackboardComponent& OwnerComp, const uint8* MemoryBlock) const
+{
+	if (HasInstance())
+	{
+		const UBlackboardKeyType* InstancedKey = GetKeyInstance(OwnerComp, MemoryBlock);
+		const uint8* InstanceMemoryBlock = MemoryBlock + sizeof(FBlackboardInstancedKeyMemory);
+		return InstancedKey->DescribeValue(OwnerComp, InstanceMemoryBlock);
+	}
+
+	return DescribeValue(OwnerComp, MemoryBlock);
+}
+
+FString UBlackboardKeyType::DescribeValue(const UBlackboardComponent& OwnerComp, const uint8* RawData) const
 {
 	FString DescBytes;
 	for (int32 ValueIndex = 0; ValueIndex < ValueSize; ValueIndex++)
@@ -20,14 +226,16 @@ FString UBlackboardKeyType::DescribeValue(const uint8* RawData) const
 	return DescBytes.Len() ? (FString("0x") + DescBytes) : FString("empty");
 }
 
-FString UBlackboardKeyType::DescribeSelf() const
+EBlackboardCompare::Type UBlackboardKeyType::CompareValues(const UBlackboardComponent& OwnerComp, const uint8* MemoryBlock,
+	const UBlackboardKeyType* OtherKeyOb, const uint8* OtherMemoryBlock) const
 {
-	return FString();
-};
+	const int32 MemCmp = FMemory::Memcmp(MemoryBlock, OtherMemoryBlock, GetValueSize());
+	return MemCmp ? EBlackboardCompare::NotEqual : EBlackboardCompare::Equal;
+}
 
-void UBlackboardKeyType::Initialize(uint8* MemoryBlock) const
+void UBlackboardKeyType::CopyValues(UBlackboardComponent& OwnerComp, uint8* MemoryBlock, const UBlackboardKeyType* SourceKeyOb, const uint8* SourceBlock)
 {
-	// empty in base class
+	FMemory::Memcpy(MemoryBlock, SourceBlock, GetValueSize());
 }
 
 bool UBlackboardKeyType::IsAllowedByFilter(UBlackboardKeyType* FilterOb) const
@@ -35,7 +243,22 @@ bool UBlackboardKeyType::IsAllowedByFilter(UBlackboardKeyType* FilterOb) const
 	return GetClass() == (FilterOb ? FilterOb->GetClass() : NULL);
 }
 
-bool UBlackboardKeyType::GetLocation(const uint8* RawData, FVector& Location) const
+FString UBlackboardKeyType::DescribeSelf() const
+{
+	return FString();
+}
+
+FString UBlackboardKeyType::DescribeArithmeticParam(int32 IntValue, float FloatValue) const
+{
+	return FString();
+}
+
+// DEPRECATED FUNCTIONS
+void UBlackboardKeyType::Initialize(uint8* MemoryBlock) const
+{
+}
+
+bool UBlackboardKeyType::GetLocation(const uint8* MemoryBlock, FVector& Location) const
 {
 	return false;
 }
@@ -47,22 +270,12 @@ bool UBlackboardKeyType::GetRotation(const uint8* MemoryBlock, FRotator& Rotatio
 
 bool UBlackboardKeyType::Clear(uint8* MemoryBlock) const
 {
-	for (int32 ByteIndex = 0; ByteIndex < GetValueSize(); ++ByteIndex)
-	{
-		if (MemoryBlock[ByteIndex] != uint8(0))
-		{
-			FMemory::Memzero(MemoryBlock, GetValueSize());
-			return true;
-		}		
-	}
-
-	// value already "0", return false to indicate it had not changed
 	return false;
 }
 
 EBlackboardCompare::Type UBlackboardKeyType::Compare(const uint8* MemoryBlockA, const uint8* MemoryBlockB) const
-{ 
-	return MemoryBlockA == MemoryBlockB ? EBlackboardCompare::Equal : EBlackboardCompare::NotEqual;
+{
+	return EBlackboardCompare::NotEqual;
 }
 
 bool UBlackboardKeyType::TestBasicOperation(const uint8* MemoryBlock, EBasicKeyOperation::Type Op) const
@@ -80,7 +293,7 @@ bool UBlackboardKeyType::TestTextOperation(const uint8* MemoryBlock, ETextKeyOpe
 	return false;
 }
 
-FString UBlackboardKeyType::DescribeArithmeticParam(int32 IntValue, float FloatValue) const
+FString UBlackboardKeyType::DescribeValue(const uint8* MemoryBlock) const
 {
 	return FString();
 }
diff --git a/Engine/Source/Runtime/AIModule/Private/BehaviorTree/Blackboard/BlackboardKeyType_Bool.cpp b/Engine/Source/Runtime/AIModule/Private/BehaviorTree/Blackboard/BlackboardKeyType_Bool.cpp
--- a/Engine/Source/Runtime/AIModule/Private/BehaviorTree/Blackboard/BlackboardKeyType_Bool.cpp
+++ b/Engine/Source/Runtime/AIModule/Private/BehaviorTree/Blackboard/BlackboardKeyType_Bool.cpp
@@ -11,28 +11,32 @@ UBlackboardKeyType_Bool::UBlackboardKeyType_Bool(const FObjectInitializer& Objec
 	SupportedOp = EBlackboardKeyOperation::Basic;
 }
 
-bool UBlackboardKeyType_Bool::GetValue(const uint8* RawData)
+bool UBlackboardKeyType_Bool::GetValue(const UBlackboardKeyType_Bool* KeyOb, const uint8* RawData)
 {
 	return GetValueFromMemory<uint8>(RawData) != 0;
 }
 
-bool UBlackboardKeyType_Bool::SetValue(uint8* RawData, bool bValue)
+bool UBlackboardKeyType_Bool::SetValue(UBlackboardKeyType_Bool* KeyOb, uint8* RawData, bool bValue)
 {
 	return SetValueInMemory<uint8>(RawData, bValue);
 }
 
-FString UBlackboardKeyType_Bool::DescribeValue(const uint8* RawData) const
+EBlackboardCompare::Type UBlackboardKeyType_Bool::CompareValues(const UBlackboardComponent& OwnerComp, const uint8* MemoryBlock,
+	const UBlackboardKeyType* OtherKeyOb, const uint8* OtherMemoryBlock) const
 {
-	return GetValue(RawData) ? TEXT("true") : TEXT("false");
+	const bool MyValue = GetValue(this, MemoryBlock);
+	const bool OtherValue = GetValue((UBlackboardKeyType_Bool*)OtherKeyOb, OtherMemoryBlock);
+
+	return (MyValue == OtherValue) ? EBlackboardCompare::Equal : EBlackboardCompare::NotEqual;
 }
 
-EBlackboardCompare::Type UBlackboardKeyType_Bool::Compare(const uint8* MemoryBlockA, const uint8* MemoryBlockB) const
+FString UBlackboardKeyType_Bool::DescribeValue(const UBlackboardComponent& OwnerComp, const uint8* RawData) const
 {
-	return GetValueFromMemory<uint8>(MemoryBlockA) == GetValueFromMemory<uint8>(MemoryBlockB) ? EBlackboardCompare::Equal : EBlackboardCompare::NotEqual;
+	return GetValue(this, RawData) ? TEXT("true") : TEXT("false");
 }
 
-bool UBlackboardKeyType_Bool::TestBasicOperation(const uint8* MemoryBlock, EBasicKeyOperation::Type Op) const
+bool UBlackboardKeyType_Bool::TestBasicOperation(const UBlackboardComponent& OwnerComp, const uint8* MemoryBlock, EBasicKeyOperation::Type Op) const
 {
-	const bool Value = GetValue(MemoryBlock);
+	const bool Value = GetValue(this, MemoryBlock);
 	return (Op == EBasicKeyOperation::Set) ? Value : !Value;
 }
diff --git a/Engine/Source/Runtime/AIModule/Private/BehaviorTree/Blackboard/BlackboardKeyType_Class.cpp b/Engine/Source/Runtime/AIModule/Private/BehaviorTree/Blackboard/BlackboardKeyType_Class.cpp
--- a/Engine/Source/Runtime/AIModule/Private/BehaviorTree/Blackboard/BlackboardKeyType_Class.cpp
+++ b/Engine/Source/Runtime/AIModule/Private/BehaviorTree/Blackboard/BlackboardKeyType_Class.cpp
@@ -12,21 +12,30 @@ UBlackboardKeyType_Class::UBlackboardKeyType_Class(const FObjectInitializer& Obj
 	SupportedOp = EBlackboardKeyOperation::Basic;
 }
 
-UClass* UBlackboardKeyType_Class::GetValue(const uint8* RawData)
+UClass* UBlackboardKeyType_Class::GetValue(const UBlackboardKeyType_Class* KeyOb, const uint8* RawData)
 {
 	TWeakObjectPtr<UClass> WeakObjPtr = GetValueFromMemory< TWeakObjectPtr<UClass> >(RawData);
 	return WeakObjPtr.Get();
 }
 
-bool UBlackboardKeyType_Class::SetValue(uint8* RawData, UClass* Value)
+bool UBlackboardKeyType_Class::SetValue(UBlackboardKeyType_Class* KeyOb, uint8* RawData, UClass* Value)
 {
 	TWeakObjectPtr<UClass> WeakObjPtr(Value);
 	return SetWeakObjectInMemory<UClass>(RawData, WeakObjPtr);
 }
 
-FString UBlackboardKeyType_Class::DescribeValue(const uint8* RawData) const
+EBlackboardCompare::Type UBlackboardKeyType_Class::CompareValues(const UBlackboardComponent& OwnerComp, const uint8* MemoryBlock,
+	const UBlackboardKeyType* OtherKeyOb, const uint8* OtherMemoryBlock) const
 {
-	return *GetNameSafe(GetValue(RawData));
+	const UClass* MyValue = GetValue(this, MemoryBlock);
+	const UClass* OtherValue = GetValue((UBlackboardKeyType_Class*)OtherKeyOb, OtherMemoryBlock);
+
+	return (MyValue == OtherValue) ? EBlackboardCompare::Equal : EBlackboardCompare::NotEqual;
+}
+
+FString UBlackboardKeyType_Class::DescribeValue(const UBlackboardComponent& OwnerComp, const uint8* RawData) const
+{
+	return *GetNameSafe(GetValue(this, RawData));
 }
 
 FString UBlackboardKeyType_Class::DescribeSelf() const
@@ -40,12 +49,7 @@ bool UBlackboardKeyType_Class::IsAllowedByFilter(UBlackboardKeyType* FilterOb) c
 	return (FilterClass && (FilterClass->BaseClass == BaseClass || BaseClass->IsChildOf(FilterClass->BaseClass)));
 }
 
-EBlackboardCompare::Type UBlackboardKeyType_Class::Compare(const uint8* MemoryBlockA, const uint8* MemoryBlockB) const
-{
-	return (FMemory::Memcmp(MemoryBlockA, MemoryBlockB, ValueSize) == 0) ? EBlackboardCompare::Equal : EBlackboardCompare::NotEqual;
-}
-
-bool UBlackboardKeyType_Class::TestBasicOperation(const uint8* MemoryBlock, EBasicKeyOperation::Type Op) const
+bool UBlackboardKeyType_Class::TestBasicOperation(const UBlackboardComponent& OwnerComp, const uint8* MemoryBlock, EBasicKeyOperation::Type Op) const
 {
 	TWeakObjectPtr<UClass> WeakObjPtr = GetValueFromMemory< TWeakObjectPtr<UClass> >(MemoryBlock);
 	return (Op == EBasicKeyOperation::Set) ? WeakObjPtr.IsValid() : !WeakObjPtr.IsValid();
diff --git a/Engine/Source/Runtime/AIModule/Private/BehaviorTree/Blackboard/BlackboardKeyType_Enum.cpp b/Engine/Source/Runtime/AIModule/Private/BehaviorTree/Blackboard/BlackboardKeyType_Enum.cpp
--- a/Engine/Source/Runtime/AIModule/Private/BehaviorTree/Blackboard/BlackboardKeyType_Enum.cpp
+++ b/Engine/Source/Runtime/AIModule/Private/BehaviorTree/Blackboard/BlackboardKeyType_Enum.cpp
@@ -11,19 +11,30 @@ UBlackboardKeyType_Enum::UBlackboardKeyType_Enum(const FObjectInitializer& Objec
 	SupportedOp = EBlackboardKeyOperation::Arithmetic;
 }
 
-uint8 UBlackboardKeyType_Enum::GetValue(const uint8* RawData)
+uint8 UBlackboardKeyType_Enum::GetValue(const UBlackboardKeyType_Enum* KeyOb, const uint8* RawData)
 {
 	return GetValueFromMemory<uint8>(RawData);
 }
 
-bool UBlackboardKeyType_Enum::SetValue(uint8* RawData, uint8 Value)
+bool UBlackboardKeyType_Enum::SetValue(UBlackboardKeyType_Enum* KeyOb, uint8* RawData, uint8 Value)
 {
 	return SetValueInMemory<uint8>(RawData, Value);
 }
 
-FString UBlackboardKeyType_Enum::DescribeValue(const uint8* RawData) const
+EBlackboardCompare::Type UBlackboardKeyType_Enum::CompareValues(const UBlackboardComponent& OwnerComp, const uint8* MemoryBlock,
+	const UBlackboardKeyType* OtherKeyOb, const uint8* OtherMemoryBlock) const
 {
-	return EnumType ? EnumType->GetEnumText(GetValue(RawData)).ToString() : FString("UNKNOWN!");
+	const uint8 MyValue = GetValue(this, MemoryBlock);
+	const uint8 OtherValue = GetValue((UBlackboardKeyType_Enum*)OtherKeyOb, OtherMemoryBlock);
+
+	return (MyValue > OtherValue) ? EBlackboardCompare::Greater :
+		(MyValue < OtherValue) ? EBlackboardCompare::Less :
+		EBlackboardCompare::Equal;
+}
+
+FString UBlackboardKeyType_Enum::DescribeValue(const UBlackboardComponent& OwnerComp, const uint8* RawData) const
+{
+	return EnumType ? EnumType->GetEnumText(GetValue(this, RawData)).ToString() : FString("UNKNOWN!");
 }
 
 FString UBlackboardKeyType_Enum::DescribeSelf() const
@@ -37,15 +48,9 @@ bool UBlackboardKeyType_Enum::IsAllowedByFilter(UBlackboardKeyType* FilterOb) co
 	return (FilterEnum && FilterEnum->EnumType == EnumType);
 }
 
-EBlackboardCompare::Type UBlackboardKeyType_Enum::Compare(const uint8* MemoryBlockA, const uint8* MemoryBlockB) const
-{
-	const int8 Diff = GetValueFromMemory<uint8>(MemoryBlockA) - GetValueFromMemory<uint8>(MemoryBlockB);
-	return Diff > 0 ? EBlackboardCompare::Greater : (Diff < 0 ? EBlackboardCompare::Less : EBlackboardCompare::Equal);
-}
-
-bool UBlackboardKeyType_Enum::TestArithmeticOperation(const uint8* MemoryBlock, EArithmeticKeyOperation::Type Op, int32 OtherIntValue, float OtherFloatValue) const
+bool UBlackboardKeyType_Enum::TestArithmeticOperation(const UBlackboardComponent& OwnerComp, const uint8* MemoryBlock, EArithmeticKeyOperation::Type Op, int32 OtherIntValue, float OtherFloatValue) const
 {
-	const uint8 Value = GetValue(MemoryBlock);
+	const uint8 Value = GetValue(this, MemoryBlock);
 	switch (Op)
 	{
 	case EArithmeticKeyOperation::Equal:			return (Value == OtherIntValue);
diff --git a/Engine/Source/Runtime/AIModule/Private/BehaviorTree/Blackboard/BlackboardKeyType_Float.cpp b/Engine/Source/Runtime/AIModule/Private/BehaviorTree/Blackboard/BlackboardKeyType_Float.cpp
--- a/Engine/Source/Runtime/AIModule/Private/BehaviorTree/Blackboard/BlackboardKeyType_Float.cpp
+++ b/Engine/Source/Runtime/AIModule/Private/BehaviorTree/Blackboard/BlackboardKeyType_Float.cpp
@@ -11,31 +11,35 @@ UBlackboardKeyType_Float::UBlackboardKeyType_Float(const FObjectInitializer& Obj
 	SupportedOp = EBlackboardKeyOperation::Arithmetic;
 }
 
-float UBlackboardKeyType_Float::GetValue(const uint8* RawData)
+float UBlackboardKeyType_Float::GetValue(const UBlackboardKeyType_Float* KeyOb, const uint8* RawData)
 {
 	return GetValueFromMemory<float>(RawData);
 }
 
-bool UBlackboardKeyType_Float::SetValue(uint8* RawData, float Value)
+bool UBlackboardKeyType_Float::SetValue(UBlackboardKeyType_Float* KeyOb, uint8* RawData, float Value)
 {
 	return SetValueInMemory<float>(RawData, Value);
 }
 
-FString UBlackboardKeyType_Float::DescribeValue(const uint8* RawData) const
+EBlackboardCompare::Type UBlackboardKeyType_Float::CompareValues(const UBlackboardComponent& OwnerComp, const uint8* MemoryBlock,
+	const UBlackboardKeyType* OtherKeyOb, const uint8* OtherMemoryBlock) const
 {
-	return FString::Printf(TEXT("%f"), GetValue(RawData));
+	const float MyValue = GetValue(this, MemoryBlock);
+	const float OtherValue = GetValue((UBlackboardKeyType_Float*)OtherKeyOb, OtherMemoryBlock);
+
+	return (FMath::Abs(MyValue - OtherValue) < KINDA_SMALL_NUMBER) ? EBlackboardCompare::Equal :
+		(MyValue > OtherValue) ? EBlackboardCompare::Greater :
+		EBlackboardCompare::Less;
 }
 
-EBlackboardCompare::Type UBlackboardKeyType_Float::Compare(const uint8* MemoryBlockA, const uint8* MemoryBlockB) const
+FString UBlackboardKeyType_Float::DescribeValue(const UBlackboardComponent& OwnerComp, const uint8* RawData) const
 {
-	const float Diff = GetValueFromMemory<float>(MemoryBlockA) - GetValueFromMemory<float>(MemoryBlockB);
-	return FMath::Abs(Diff) < KINDA_SMALL_NUMBER ? EBlackboardCompare::Equal 
-		: (Diff > 0 ? EBlackboardCompare::Greater : EBlackboardCompare::Less);
+	return FString::Printf(TEXT("%f"), GetValue(this, RawData));
 }
 
-bool UBlackboardKeyType_Float::TestArithmeticOperation(const uint8* MemoryBlock, EArithmeticKeyOperation::Type Op, int32 OtherIntValue, float OtherFloatValue) const
+bool UBlackboardKeyType_Float::TestArithmeticOperation(const UBlackboardComponent& OwnerComp, const uint8* MemoryBlock, EArithmeticKeyOperation::Type Op, int32 OtherIntValue, float OtherFloatValue) const
 {
-	const float Value = GetValue(MemoryBlock);
+	const float Value = GetValue(this, MemoryBlock);
 	switch (Op)
 	{
 	case EArithmeticKeyOperation::Equal:			return (Value == OtherFloatValue);
diff --git a/Engine/Source/Runtime/AIModule/Private/BehaviorTree/Blackboard/BlackboardKeyType_Int.cpp b/Engine/Source/Runtime/AIModule/Private/BehaviorTree/Blackboard/BlackboardKeyType_Int.cpp
--- a/Engine/Source/Runtime/AIModule/Private/BehaviorTree/Blackboard/BlackboardKeyType_Int.cpp
+++ b/Engine/Source/Runtime/AIModule/Private/BehaviorTree/Blackboard/BlackboardKeyType_Int.cpp
@@ -11,30 +11,35 @@ UBlackboardKeyType_Int::UBlackboardKeyType_Int(const FObjectInitializer& ObjectI
 	SupportedOp = EBlackboardKeyOperation::Arithmetic;
 }
 
-int32 UBlackboardKeyType_Int::GetValue(const uint8* RawData)
+int32 UBlackboardKeyType_Int::GetValue(const UBlackboardKeyType_Int* KeyOb, const uint8* RawData)
 {
 	return GetValueFromMemory<int32>(RawData);
 }
 
-bool UBlackboardKeyType_Int::SetValue(uint8* RawData, int32 Value)
+bool UBlackboardKeyType_Int::SetValue(UBlackboardKeyType_Int* KeyOb, uint8* RawData, int32 Value)
 {
 	return SetValueInMemory<int32>(RawData, Value);
 }
 
-FString UBlackboardKeyType_Int::DescribeValue(const uint8* RawData) const
+EBlackboardCompare::Type UBlackboardKeyType_Int::CompareValues(const UBlackboardComponent& OwnerComp, const uint8* MemoryBlock,
+	const UBlackboardKeyType* OtherKeyOb, const uint8* OtherMemoryBlock) const
 {
-	return FString::Printf(TEXT("%d"), GetValue(RawData));
+	const int32 MyValue = GetValue(this, MemoryBlock);
+	const int32 OtherValue = GetValue((UBlackboardKeyType_Int*)OtherKeyOb, OtherMemoryBlock);
+
+	return (MyValue > OtherValue) ? EBlackboardCompare::Greater :
+		(MyValue < OtherValue) ? EBlackboardCompare::Less :
+		EBlackboardCompare::Equal;
 }
 
-EBlackboardCompare::Type UBlackboardKeyType_Int::Compare(const uint8* MemoryBlockA, const uint8* MemoryBlockB) const
+FString UBlackboardKeyType_Int::DescribeValue(const UBlackboardComponent& OwnerComp, const uint8* RawData) const
 {
-	const int32 Diff = GetValueFromMemory<int32>(MemoryBlockA) - GetValueFromMemory<int32>(MemoryBlockB);
-	return Diff > 0 ? EBlackboardCompare::Greater : (Diff < 0 ? EBlackboardCompare::Less : EBlackboardCompare::Equal);
+	return FString::Printf(TEXT("%d"), GetValue(this, RawData));
 }
 
-bool UBlackboardKeyType_Int::TestArithmeticOperation(const uint8* MemoryBlock, EArithmeticKeyOperation::Type Op, int32 OtherIntValue, float OtherFloatValue) const
+bool UBlackboardKeyType_Int::TestArithmeticOperation(const UBlackboardComponent& OwnerComp, const uint8* MemoryBlock, EArithmeticKeyOperation::Type Op, int32 OtherIntValue, float OtherFloatValue) const
 {
-	const int32 Value = GetValue(MemoryBlock);
+	const int32 Value = GetValue(this, MemoryBlock);
 	switch (Op)
 	{
 	case EArithmeticKeyOperation::Equal:			return (Value == OtherIntValue);
diff --git a/Engine/Source/Runtime/AIModule/Private/BehaviorTree/Blackboard/BlackboardKeyType_Name.cpp b/Engine/Source/Runtime/AIModule/Private/BehaviorTree/Blackboard/BlackboardKeyType_Name.cpp
--- a/Engine/Source/Runtime/AIModule/Private/BehaviorTree/Blackboard/BlackboardKeyType_Name.cpp
+++ b/Engine/Source/Runtime/AIModule/Private/BehaviorTree/Blackboard/BlackboardKeyType_Name.cpp
@@ -11,30 +11,33 @@ UBlackboardKeyType_Name::UBlackboardKeyType_Name(const FObjectInitializer& Objec
 	SupportedOp = EBlackboardKeyOperation::Text;
 }
 
-FName UBlackboardKeyType_Name::GetValue(const uint8* RawData)
+FName UBlackboardKeyType_Name::GetValue(const UBlackboardKeyType_Name* KeyOb, const uint8* RawData)
 {
 	return GetValueFromMemory<FName>(RawData);
 }
 
-bool UBlackboardKeyType_Name::SetValue(uint8* RawData, const FName& Value)
+bool UBlackboardKeyType_Name::SetValue(UBlackboardKeyType_Name* KeyOb, uint8* RawData, const FName& Value)
 {
 	return SetValueInMemory<FName>(RawData, Value);
 }
 
-FString UBlackboardKeyType_Name::DescribeValue(const uint8* RawData) const
+EBlackboardCompare::Type UBlackboardKeyType_Name::CompareValues(const UBlackboardComponent& OwnerComp, const uint8* MemoryBlock,
+	const UBlackboardKeyType* OtherKeyOb, const uint8* OtherMemoryBlock) const
 {
-	return GetValue(RawData).ToString();
+	const FName MyValue = GetValue(this, MemoryBlock);
+	const FName OtherValue = GetValue((UBlackboardKeyType_Name*)OtherKeyOb, OtherMemoryBlock);
+
+	return (MyValue == OtherValue) ? EBlackboardCompare::Equal : EBlackboardCompare::NotEqual;
 }
 
-EBlackboardCompare::Type UBlackboardKeyType_Name::Compare(const uint8* MemoryBlockA, const uint8* MemoryBlockB) const
+FString UBlackboardKeyType_Name::DescribeValue(const UBlackboardComponent& OwnerComp, const uint8* RawData) const
 {
-	return GetValueFromMemory<FName>(MemoryBlockA) == GetValueFromMemory<FName>(MemoryBlockB)
-		? EBlackboardCompare::Equal : EBlackboardCompare::NotEqual;
+	return GetValue(this, RawData).ToString();
 }
 
-bool UBlackboardKeyType_Name::TestTextOperation(const uint8* MemoryBlock, ETextKeyOperation::Type Op, const FString& OtherString) const
+bool UBlackboardKeyType_Name::TestTextOperation(const UBlackboardComponent& OwnerComp, const uint8* MemoryBlock, ETextKeyOperation::Type Op, const FString& OtherString) const
 {
-	const FString Value = GetValue(MemoryBlock).ToString();
+	const FString Value = GetValue(this, MemoryBlock).ToString();
 	switch (Op)
 	{
 	case ETextKeyOperation::Equal:			return (Value == OtherString);
diff --git a/Engine/Source/Runtime/AIModule/Private/BehaviorTree/Blackboard/BlackboardKeyType_NativeEnum.cpp b/Engine/Source/Runtime/AIModule/Private/BehaviorTree/Blackboard/BlackboardKeyType_NativeEnum.cpp
--- a/Engine/Source/Runtime/AIModule/Private/BehaviorTree/Blackboard/BlackboardKeyType_NativeEnum.cpp
+++ b/Engine/Source/Runtime/AIModule/Private/BehaviorTree/Blackboard/BlackboardKeyType_NativeEnum.cpp
@@ -11,19 +11,30 @@ UBlackboardKeyType_NativeEnum::UBlackboardKeyType_NativeEnum(const FObjectInitia
 	SupportedOp = EBlackboardKeyOperation::Arithmetic;
 }
 
-uint8 UBlackboardKeyType_NativeEnum::GetValue(const uint8* RawData)
+uint8 UBlackboardKeyType_NativeEnum::GetValue(const UBlackboardKeyType_NativeEnum* KeyOb, const uint8* RawData)
 {
 	return GetValueFromMemory<uint8>(RawData);
 }
 
-bool UBlackboardKeyType_NativeEnum::SetValue(uint8* RawData, uint8 Value)
+bool UBlackboardKeyType_NativeEnum::SetValue(UBlackboardKeyType_NativeEnum* KeyOb, uint8* RawData, uint8 Value)
 {
 	return SetValueInMemory<uint8>(RawData, Value);
 }
 
-FString UBlackboardKeyType_NativeEnum::DescribeValue(const uint8* RawData) const
+EBlackboardCompare::Type UBlackboardKeyType_NativeEnum::CompareValues(const UBlackboardComponent& OwnerComp, const uint8* MemoryBlock,
+	const UBlackboardKeyType* OtherKeyOb, const uint8* OtherMemoryBlock) const
 {
-	return EnumType ? EnumType->GetEnumName(GetValue(RawData)) : FString("UNKNOWN!");
+	const uint8 MyValue = GetValue(this, MemoryBlock);
+	const uint8 OtherValue = GetValue((UBlackboardKeyType_NativeEnum*)OtherKeyOb, OtherMemoryBlock);
+
+	return (MyValue > OtherValue) ? EBlackboardCompare::Greater :
+		(MyValue < OtherValue) ? EBlackboardCompare::Less :
+		EBlackboardCompare::Equal;
+}
+
+FString UBlackboardKeyType_NativeEnum::DescribeValue(const UBlackboardComponent& OwnerComp, const uint8* RawData) const
+{
+	return EnumType ? EnumType->GetEnumName(GetValue(this, RawData)) : FString("UNKNOWN!");
 }
 
 FString UBlackboardKeyType_NativeEnum::DescribeSelf() const
@@ -51,15 +62,9 @@ void UBlackboardKeyType_NativeEnum::PostEditChangeProperty(FPropertyChangedEvent
 }
 #endif
 
-EBlackboardCompare::Type UBlackboardKeyType_NativeEnum::Compare(const uint8* MemoryBlockA, const uint8* MemoryBlockB) const
-{
-	const int8 Diff = GetValueFromMemory<uint8>(MemoryBlockA) - GetValueFromMemory<uint8>(MemoryBlockB);
-	return Diff > 0 ? EBlackboardCompare::Greater : (Diff < 0 ? EBlackboardCompare::Less : EBlackboardCompare::Equal);
-}
-
-bool UBlackboardKeyType_NativeEnum::TestArithmeticOperation(const uint8* MemoryBlock, EArithmeticKeyOperation::Type Op, int32 OtherIntValue, float OtherFloatValue) const
+bool UBlackboardKeyType_NativeEnum::TestArithmeticOperation(const UBlackboardComponent& OwnerComp, const uint8* MemoryBlock, EArithmeticKeyOperation::Type Op, int32 OtherIntValue, float OtherFloatValue) const
 {
-	const uint8 Value = GetValue(MemoryBlock);
+	const uint8 Value = GetValue(this, MemoryBlock);
 	switch (Op)
 	{
 	case EArithmeticKeyOperation::Equal:			return (Value == OtherIntValue);
diff --git a/Engine/Source/Runtime/AIModule/Private/BehaviorTree/Blackboard/BlackboardKeyType_Object.cpp b/Engine/Source/Runtime/AIModule/Private/BehaviorTree/Blackboard/BlackboardKeyType_Object.cpp
--- a/Engine/Source/Runtime/AIModule/Private/BehaviorTree/Blackboard/BlackboardKeyType_Object.cpp
+++ b/Engine/Source/Runtime/AIModule/Private/BehaviorTree/Blackboard/BlackboardKeyType_Object.cpp
@@ -12,21 +12,30 @@ UBlackboardKeyType_Object::UBlackboardKeyType_Object(const FObjectInitializer& O
 	SupportedOp = EBlackboardKeyOperation::Basic;
 }
 
-UObject* UBlackboardKeyType_Object::GetValue(const uint8* RawData)
+UObject* UBlackboardKeyType_Object::GetValue(const UBlackboardKeyType_Object* KeyOb, const uint8* RawData)
 {
 	FWeakObjectPtr WeakObjPtr = GetValueFromMemory<FWeakObjectPtr>(RawData);
 	return WeakObjPtr.Get();
 }
 
-bool UBlackboardKeyType_Object::SetValue(uint8* RawData, UObject* Value)
+bool UBlackboardKeyType_Object::SetValue(UBlackboardKeyType_Object* KeyOb, uint8* RawData, UObject* Value)
 {
 	TWeakObjectPtr<UObject> WeakObjPtr(Value);
 	return SetWeakObjectInMemory<UObject>(RawData, WeakObjPtr);
 }
 
-FString UBlackboardKeyType_Object::DescribeValue(const uint8* RawData) const
+EBlackboardCompare::Type UBlackboardKeyType_Object::CompareValues(const UBlackboardComponent& OwnerComp, const uint8* MemoryBlock,
+	const UBlackboardKeyType* OtherKeyOb, const uint8* OtherMemoryBlock) const
 {
-	return *GetNameSafe(GetValue(RawData));
+	const UObject* MyValue = GetValue(this, MemoryBlock);
+	const UObject* OtherValue = GetValue((UBlackboardKeyType_Object*)OtherKeyOb, OtherMemoryBlock);
+
+	return (MyValue == OtherValue) ? EBlackboardCompare::Equal : EBlackboardCompare::NotEqual;
+}
+
+FString UBlackboardKeyType_Object::DescribeValue(const UBlackboardComponent& OwnerComp, const uint8* RawData) const
+{
+	return *GetNameSafe(GetValue(this, RawData));
 }
 
 FString UBlackboardKeyType_Object::DescribeSelf() const
@@ -40,9 +49,9 @@ bool UBlackboardKeyType_Object::IsAllowedByFilter(UBlackboardKeyType* FilterOb)
 	return (FilterObject && (FilterObject->BaseClass == BaseClass || BaseClass->IsChildOf(FilterObject->BaseClass)));
 }
 
-bool UBlackboardKeyType_Object::GetLocation(const uint8* RawData, FVector& Location) const
+bool UBlackboardKeyType_Object::GetLocation(const UBlackboardComponent& OwnerComp, const uint8* RawData, FVector& Location) const
 {
-	AActor* MyActor = Cast<AActor>(GetValue(RawData));
+	AActor* MyActor = Cast<AActor>(GetValue(this, RawData));
 	if (MyActor)
 	{
 		Location = MyActor->GetActorLocation();
@@ -52,9 +61,9 @@ bool UBlackboardKeyType_Object::GetLocation(const uint8* RawData, FVector& Locat
 	return false;
 }
 
-bool UBlackboardKeyType_Object::GetRotation(const uint8* RawData, FRotator& Rotation) const
+bool UBlackboardKeyType_Object::GetRotation(const UBlackboardComponent& OwnerComp, const uint8* RawData, FRotator& Rotation) const
 {
-	AActor* MyActor = Cast<AActor>(GetValue(RawData));
+	AActor* MyActor = Cast<AActor>(GetValue(this, RawData));
 	if (MyActor)
 	{
 		Rotation = MyActor->GetActorRotation();
@@ -64,12 +73,7 @@ bool UBlackboardKeyType_Object::GetRotation(const uint8* RawData, FRotator& Rota
 	return false;
 }
 
-EBlackboardCompare::Type UBlackboardKeyType_Object::Compare(const uint8* MemoryBlockA, const uint8* MemoryBlockB) const
-{
-	return (FMemory::Memcmp(MemoryBlockA, MemoryBlockB, ValueSize) == 0) ? EBlackboardCompare::Equal : EBlackboardCompare::NotEqual;
-}
-
-bool UBlackboardKeyType_Object::TestBasicOperation(const uint8* MemoryBlock, EBasicKeyOperation::Type Op) const
+bool UBlackboardKeyType_Object::TestBasicOperation(const UBlackboardComponent& OwnerComp, const uint8* MemoryBlock, EBasicKeyOperation::Type Op) const
 {
 	FWeakObjectPtr WeakObjPtr = GetValueFromMemory<FWeakObjectPtr>(MemoryBlock);
 	return (Op == EBasicKeyOperation::Set) ? WeakObjPtr.IsValid() : !WeakObjPtr.IsValid();
diff --git a/Engine/Source/Runtime/AIModule/Private/BehaviorTree/Blackboard/BlackboardKeyType_Rotator.cpp b/Engine/Source/Runtime/AIModule/Private/BehaviorTree/Blackboard/BlackboardKeyType_Rotator.cpp
--- a/Engine/Source/Runtime/AIModule/Private/BehaviorTree/Blackboard/BlackboardKeyType_Rotator.cpp
+++ b/Engine/Source/Runtime/AIModule/Private/BehaviorTree/Blackboard/BlackboardKeyType_Rotator.cpp
@@ -11,46 +11,55 @@ UBlackboardKeyType_Rotator::UBlackboardKeyType_Rotator(const FObjectInitializer&
 	SupportedOp = EBlackboardKeyOperation::Basic;
 }
 
-FRotator UBlackboardKeyType_Rotator::GetValue(const uint8* RawData)
+FRotator UBlackboardKeyType_Rotator::GetValue(const UBlackboardKeyType_Rotator* KeyOb, const uint8* RawData)
 {
 	return GetValueFromMemory<FRotator>(RawData);
 }
 
-bool UBlackboardKeyType_Rotator::SetValue(uint8* RawData, const FRotator& Value)
+bool UBlackboardKeyType_Rotator::SetValue(UBlackboardKeyType_Rotator* KeyOb, uint8* RawData, const FRotator& Value)
 {
 	return SetValueInMemory<FRotator>(RawData, Value);
 }
 
-bool UBlackboardKeyType_Rotator::Clear(uint8* RawData) const 
+EBlackboardCompare::Type UBlackboardKeyType_Rotator::CompareValues(const UBlackboardComponent& OwnerComp, const uint8* MemoryBlock,
+	const UBlackboardKeyType* OtherKeyOb, const uint8* OtherMemoryBlock) const
 {
-	return SetValueInMemory<FRotator>(RawData, FAISystem::InvalidRotation);
+	const FRotator MyValue = GetValue(this, MemoryBlock);
+	const FRotator OtherValue = GetValue((UBlackboardKeyType_Rotator*)OtherKeyOb, OtherMemoryBlock);
+
+	return MyValue.Equals(OtherValue) ? EBlackboardCompare::Equal : EBlackboardCompare::NotEqual;
 }
 
-FString UBlackboardKeyType_Rotator::DescribeValue(const uint8* RawData) const
+void UBlackboardKeyType_Rotator::Clear(UBlackboardComponent& OwnerComp, uint8* RawData)
 {
-	const FRotator Rotation = GetValue(RawData);
-	return FAISystem::IsValidRotation(Rotation) ? Rotation.ToString() : TEXT("(invalid)");
+	SetValueInMemory<FRotator>(RawData, FAISystem::InvalidRotation);
 }
 
-bool UBlackboardKeyType_Rotator::GetRotation(const uint8* RawData, FRotator& Rotation) const
+bool UBlackboardKeyType_Rotator::IsEmpty(const UBlackboardComponent& OwnerComp, const uint8* RawData) const
 {
-	Rotation = GetValue(RawData);
-	return FAISystem::IsValidRotation(Rotation);
+	const FRotator Rotation = GetValue(this, RawData);
+	return !FAISystem::IsValidRotation(Rotation);
+}
+
+FString UBlackboardKeyType_Rotator::DescribeValue(const UBlackboardComponent& OwnerComp, const uint8* RawData) const
+{
+	const FRotator Rotation = GetValue(this, RawData);
+	return FAISystem::IsValidRotation(Rotation) ? Rotation.ToString() : TEXT("(invalid)");
 }
 
-void UBlackboardKeyType_Rotator::Initialize(uint8* RawData) const
+bool UBlackboardKeyType_Rotator::GetRotation(const UBlackboardComponent& OwnerComp, const uint8* RawData, FRotator& Rotation) const
 {
-	SetValue(RawData, FAISystem::InvalidRotation);
+	Rotation = GetValue(this, RawData);
+	return FAISystem::IsValidRotation(Rotation);
 }
 
-EBlackboardCompare::Type UBlackboardKeyType_Rotator::Compare(const uint8* MemoryBlockA, const uint8* MemoryBlockB) const
+void UBlackboardKeyType_Rotator::InitializeMemory(UBlackboardComponent& OwnerComp, uint8* RawData)
 {
-	return GetValueFromMemory<FRotator>(MemoryBlockA).Equals(GetValueFromMemory<FRotator>(MemoryBlockB))
-		? EBlackboardCompare::Equal : EBlackboardCompare::NotEqual;
+	SetValue(this, RawData, FAISystem::InvalidRotation);
 }
 
-bool UBlackboardKeyType_Rotator::TestBasicOperation(const uint8* MemoryBlock, EBasicKeyOperation::Type Op) const
+bool UBlackboardKeyType_Rotator::TestBasicOperation(const UBlackboardComponent& OwnerComp, const uint8* MemoryBlock, EBasicKeyOperation::Type Op) const
 {
-	const FRotator Rotation = GetValue(MemoryBlock);
+	const FRotator Rotation = GetValue(this, MemoryBlock);
 	return (Op == EBasicKeyOperation::Set) ? FAISystem::IsValidRotation(Rotation) : !FAISystem::IsValidRotation(Rotation);
 }
diff --git a/Engine/Source/Runtime/AIModule/Private/BehaviorTree/Blackboard/BlackboardKeyType_String.cpp b/Engine/Source/Runtime/AIModule/Private/BehaviorTree/Blackboard/BlackboardKeyType_String.cpp
--- a/Engine/Source/Runtime/AIModule/Private/BehaviorTree/Blackboard/BlackboardKeyType_String.cpp
+++ b/Engine/Source/Runtime/AIModule/Private/BehaviorTree/Blackboard/BlackboardKeyType_String.cpp
@@ -7,42 +7,63 @@ const UBlackboardKeyType_String::FDataType UBlackboardKeyType_String::InvalidVal
 
 UBlackboardKeyType_String::UBlackboardKeyType_String(const FObjectInitializer& ObjectInitializer) : Super(ObjectInitializer)
 {
-	ValueSize = sizeof(FString);
+	ValueSize = 0;
+	bCreateKeyInstance = true;
+
 	SupportedOp = EBlackboardKeyOperation::Text;
 }
 
-FString UBlackboardKeyType_String::GetValue(const uint8* RawData)
+FString UBlackboardKeyType_String::GetValue(const UBlackboardKeyType_String* KeyOb, const uint8* RawData)
+{
+	return KeyOb->StringValue;
+}
+
+bool UBlackboardKeyType_String::SetValue(UBlackboardKeyType_String* KeyOb, uint8* RawData, const FString& Value)
 {
-	return GetValueFromMemory<FString>(RawData);
+	const bool bChanged = !KeyOb->StringValue.Equals(Value);
+	KeyOb->StringValue = Value;
+	return bChanged;
 }
 
-bool UBlackboardKeyType_String::SetValue(uint8* RawData, const FString& Value)
+EBlackboardCompare::Type UBlackboardKeyType_String::CompareValues(const UBlackboardComponent& OwnerComp, const uint8* MemoryBlock,
+	const UBlackboardKeyType* OtherKeyOb, const uint8* OtherMemoryBlock) const
 {
-	return SetValueInMemory<FString>(RawData, Value);
+	const FString MyValue = GetValue(this, MemoryBlock);
+	const FString OtherValue = GetValue((UBlackboardKeyType_String*)OtherKeyOb, OtherMemoryBlock);
+
+	return MyValue.Equals(OtherValue) ? EBlackboardCompare::Equal : EBlackboardCompare::NotEqual;
 }
 
-FString UBlackboardKeyType_String::DescribeValue(const uint8* RawData) const
+void UBlackboardKeyType_String::CopyValues(UBlackboardComponent& OwnerComp, uint8* MemoryBlock, const UBlackboardKeyType* SourceKeyOb, const uint8* SourceBlock)
 {
-	return GetValue(RawData);
+	StringValue = ((UBlackboardKeyType_String*)SourceKeyOb)->StringValue;
 }
 
-EBlackboardCompare::Type UBlackboardKeyType_String::Compare(const uint8* MemoryBlockA, const uint8* MemoryBlockB) const
+FString UBlackboardKeyType_String::DescribeValue(const UBlackboardComponent& OwnerComp, const uint8* RawData) const
 {
-	return GetValueFromMemory<FString>(MemoryBlockA) == GetValueFromMemory<FString>(MemoryBlockB)
-		? EBlackboardCompare::Equal : EBlackboardCompare::NotEqual;
+	return StringValue;
 }
 
-bool UBlackboardKeyType_String::TestTextOperation(const uint8* MemoryBlock, ETextKeyOperation::Type Op, const FString& OtherString) const
+bool UBlackboardKeyType_String::TestTextOperation(const UBlackboardComponent& OwnerComp, const uint8* MemoryBlock, ETextKeyOperation::Type Op, const FString& OtherString) const
 {
-	const FString Value = GetValue(MemoryBlock);
 	switch (Op)
 	{
-	case ETextKeyOperation::Equal:			return (Value == OtherString);
-	case ETextKeyOperation::NotEqual:		return (Value != OtherString);
-	case ETextKeyOperation::Contain:		return (Value.Contains(OtherString) == true);
-	case ETextKeyOperation::NotContain:		return (Value.Contains(OtherString) == false);
+	case ETextKeyOperation::Equal:			return (StringValue == OtherString);
+	case ETextKeyOperation::NotEqual:		return (StringValue != OtherString);
+	case ETextKeyOperation::Contain:		return (StringValue.Contains(OtherString) == true);
+	case ETextKeyOperation::NotContain:		return (StringValue.Contains(OtherString) == false);
 	default: break;
 	}
 
 	return false;
 }
+
+void UBlackboardKeyType_String::Clear(UBlackboardComponent& OwnerComp, uint8* MemoryBlock)
+{
+	StringValue.Empty();
+}
+
+bool UBlackboardKeyType_String::IsEmpty(const UBlackboardComponent& OwnerComp, const uint8* MemoryBlock) const
+{
+	return StringValue.IsEmpty();
+}
diff --git a/Engine/Source/Runtime/AIModule/Private/BehaviorTree/Blackboard/BlackboardKeyType_Vector.cpp b/Engine/Source/Runtime/AIModule/Private/BehaviorTree/Blackboard/BlackboardKeyType_Vector.cpp
--- a/Engine/Source/Runtime/AIModule/Private/BehaviorTree/Blackboard/BlackboardKeyType_Vector.cpp
+++ b/Engine/Source/Runtime/AIModule/Private/BehaviorTree/Blackboard/BlackboardKeyType_Vector.cpp
@@ -11,46 +11,55 @@ UBlackboardKeyType_Vector::UBlackboardKeyType_Vector(const FObjectInitializer& O
 	SupportedOp = EBlackboardKeyOperation::Basic;
 }
 
-FVector UBlackboardKeyType_Vector::GetValue(const uint8* RawData)
+FVector UBlackboardKeyType_Vector::GetValue(const UBlackboardKeyType_Vector* KeyOb, const uint8* RawData)
 {
 	return GetValueFromMemory<FVector>(RawData);
 }
 
-bool UBlackboardKeyType_Vector::SetValue(uint8* RawData, const FVector& Value)
+bool UBlackboardKeyType_Vector::SetValue(UBlackboardKeyType_Vector* KeyOb, uint8* RawData, const FVector& Value)
 {
 	return SetValueInMemory<FVector>(RawData, Value);
 }
 
-bool UBlackboardKeyType_Vector::Clear(uint8* RawData) const 
+EBlackboardCompare::Type UBlackboardKeyType_Vector::CompareValues(const UBlackboardComponent& OwnerComp, const uint8* MemoryBlock,
+	const UBlackboardKeyType* OtherKeyOb, const uint8* OtherMemoryBlock) const
 {
-	return SetValueInMemory<FVector>(RawData, FAISystem::InvalidLocation);
+	const FVector MyValue = GetValue(this, MemoryBlock);
+	const FVector OtherValue = GetValue((UBlackboardKeyType_Vector*)OtherKeyOb, OtherMemoryBlock);
+
+	return MyValue.Equals(OtherValue) ? EBlackboardCompare::Equal : EBlackboardCompare::NotEqual;
 }
 
-FString UBlackboardKeyType_Vector::DescribeValue(const uint8* RawData) const
+void UBlackboardKeyType_Vector::Clear(UBlackboardComponent& OwnerComp, uint8* RawData)
 {
-	const FVector Location = GetValue(RawData);
-	return FAISystem::IsValidLocation(Location) ? Location.ToString() : TEXT("(invalid)");
+	SetValueInMemory<FVector>(RawData, FAISystem::InvalidLocation);
 }
 
-bool UBlackboardKeyType_Vector::GetLocation(const uint8* RawData, FVector& Location) const
+bool UBlackboardKeyType_Vector::IsEmpty(const UBlackboardComponent& OwnerComp, const uint8* RawData) const
 {
-	Location = GetValue(RawData);
-	return FAISystem::IsValidLocation(Location);
+	const FVector Location = GetValue(this, RawData);
+	return !FAISystem::IsValidLocation(Location);
+}
+
+FString UBlackboardKeyType_Vector::DescribeValue(const UBlackboardComponent& OwnerComp, const uint8* RawData) const
+{
+	const FVector Location = GetValue(this, RawData);
+	return FAISystem::IsValidLocation(Location) ? Location.ToString() : TEXT("(invalid)");
 }
 
-void UBlackboardKeyType_Vector::Initialize(uint8* RawData) const
+bool UBlackboardKeyType_Vector::GetLocation(const UBlackboardComponent& OwnerComp, const uint8* RawData, FVector& Location) const
 {
-	SetValue(RawData, FAISystem::InvalidLocation);
+	Location = GetValue(this, RawData);
+	return FAISystem::IsValidLocation(Location);
 }
 
-EBlackboardCompare::Type UBlackboardKeyType_Vector::Compare(const uint8* MemoryBlockA, const uint8* MemoryBlockB) const
+void UBlackboardKeyType_Vector::InitializeMemory(UBlackboardComponent& OwnerComp, uint8* RawData)
 {
-	return GetValueFromMemory<FVector>(MemoryBlockA).Equals(GetValueFromMemory<FVector>(MemoryBlockB))
-		? EBlackboardCompare::Equal : EBlackboardCompare::NotEqual;
+	SetValue(this, RawData, FAISystem::InvalidLocation);
 }
 
-bool UBlackboardKeyType_Vector::TestBasicOperation(const uint8* MemoryBlock, EBasicKeyOperation::Type Op) const
+bool UBlackboardKeyType_Vector::TestBasicOperation(const UBlackboardComponent& OwnerComp, const uint8* MemoryBlock, EBasicKeyOperation::Type Op) const
 {
-	const FVector Location = GetValue(MemoryBlock);
+	const FVector Location = GetValue(this, MemoryBlock);
 	return (Op == EBasicKeyOperation::Set) ? FAISystem::IsValidLocation(Location) : !FAISystem::IsValidLocation(Location);
 }
diff --git a/Engine/Source/Runtime/AIModule/Private/BehaviorTree/BlackboardComponent.cpp b/Engine/Source/Runtime/AIModule/Private/BehaviorTree/BlackboardComponent.cpp
--- a/Engine/Source/Runtime/AIModule/Private/BehaviorTree/BlackboardComponent.cpp
+++ b/Engine/Source/Runtime/AIModule/Private/BehaviorTree/BlackboardComponent.cpp
@@ -45,6 +45,8 @@ void UBlackboardComponent::UninitializeComponent()
 			AISystem->UnregisterBlackboardComponent(*BlackboardAsset, *this);
 		}
 	}
+
+	DestroyValues();
 	Super::UninitializeComponent();
 }
 
@@ -103,6 +105,7 @@ bool UBlackboardComponent::InitializeBlackboard(UBlackboardData& NewAsset)
 	if (BlackboardAsset)
 	{
 		AISystem->UnregisterBlackboardComponent(*BlackboardAsset, *this);
+		DestroyValues();
 	}
 
 	BlackboardAsset = &NewAsset;
@@ -125,9 +128,13 @@ bool UBlackboardComponent::InitializeBlackboard(UBlackboardData& NewAsset)
 		{
 			for (int32 KeyIndex = 0; KeyIndex < It->Keys.Num(); KeyIndex++)
 			{
-				if (It->Keys[KeyIndex].KeyType)
+				UBlackboardKeyType* KeyType = It->Keys[KeyIndex].KeyType;
+				if (KeyType)
 				{
-					InitList.Add(FBlackboardInitializationData(KeyIndex + It->GetFirstKeyID(), It->Keys[KeyIndex].KeyType->GetValueSize()));
+					KeyType->PreInitialize(*this);
+
+					const uint16 KeyMemory = KeyType->GetValueSize() + (KeyType->HasInstance() ? sizeof(FBlackboardInstancedKeyMemory) : 0);
+					InitList.Add(FBlackboardInitializationData(KeyIndex + It->GetFirstKeyID(), KeyMemory));
 				}
 			}
 		}
@@ -146,12 +153,11 @@ bool UBlackboardComponent::InitializeBlackboard(UBlackboardData& NewAsset)
 		ValueMemory.AddZeroed(MemoryOffset);
 
 		// initialize memory
+		KeyInstances.AddZeroed(InitList.Num());
 		for (int32 Index = 0; Index < InitList.Num(); Index++)
 		{
 			const FBlackboardEntry* KeyData = BlackboardAsset->GetKey(InitList[Index].KeyID);
-			uint8* RawData = GetKeyRawData(InitList[Index].KeyID);
-
-			KeyData->KeyType->Initialize(RawData);
+			KeyData->KeyType->InitializeKey(*this, InitList[Index].KeyID);
 		}
 	
 		// naive initial synchronization with one of already instantiated blackboards using the same BB asset
@@ -169,6 +175,26 @@ bool UBlackboardComponent::InitializeBlackboard(UBlackboardData& NewAsset)
 	return bSuccess;
 }
 
+void UBlackboardComponent::DestroyValues()
+{
+	for (UBlackboardData* It = BlackboardAsset; It; It = It->Parent)
+	{
+		for (int32 KeyIndex = 0; KeyIndex < It->Keys.Num(); KeyIndex++)
+		{
+			UBlackboardKeyType* KeyType = It->Keys[KeyIndex].KeyType;
+			if (KeyType)
+			{
+				const int32 UseIdx = KeyIndex + It->GetFirstKeyID();
+				uint8* KeyMemory = GetKeyRawData(UseIdx);
+				KeyType->WrappedFree(*this, KeyMemory);
+			}
+		}
+	}
+
+	ValueOffsets.Reset();
+	ValueMemory.Reset();
+}
+
 void UBlackboardComponent::PopulateSynchronizedKeys()
 {
 	ensure(bSynchronizedKeyPopulated == false);
@@ -187,12 +213,19 @@ void UBlackboardComponent::PopulateSynchronizedKeys()
 				if (Key.bInstanceSynced)
 				{
 					const int32 KeyID = BlackboardAsset->GetKeyID(Key.EntryName);
-					check(BlackboardAsset->GetKeyID(Key.EntryName) == OtherBlackboard->BlackboardAsset->GetKeyID(Key.EntryName));
-					check(BlackboardAsset->GetKeyType(KeyID) == OtherBlackboard->BlackboardAsset->GetKeyType(KeyID));
+					const FBlackboardEntry* OtherKey = OtherBlackboard->GetBlackboardAsset()->GetKey(KeyID);
+
+					check(Key.EntryName == OtherKey->EntryName);
+					check(Key.KeyType == OtherKey->KeyType);
 
-					const TSubclassOf<UBlackboardKeyType> ValueType = BlackboardAsset->GetKeyType(KeyID);
-					const UBlackboardKeyType* ValueTypeInstance = GetDefault<UBlackboardKeyType>(*ValueType);
-					ValueTypeInstance->CopyValue(GetKeyRawData(KeyID), OtherBlackboard->GetKeyRawData(KeyID));
+					const uint16 DataOffset = Key.KeyType->IsInstanced() ? sizeof(FBlackboardInstancedKeyMemory) : 0;
+					uint8* RawData = GetKeyRawData(KeyID) + DataOffset;
+					uint8* RawSource = OtherBlackboard->GetKeyRawData(KeyID) + DataOffset;
+
+					UBlackboardKeyType* KeyOb = Key.KeyType->IsInstanced() ? KeyInstances[KeyID] : Key.KeyType;
+					const UBlackboardKeyType* SourceKeyOb = Key.KeyType->IsInstanced() ? OtherBlackboard->KeyInstances[KeyID] : Key.KeyType;
+
+					KeyOb->CopyValues(*this, RawData, SourceKeyOb, RawSource);
 				}
 			}
 			break;
@@ -374,7 +407,11 @@ bool UBlackboardComponent::IsCompatibleWith(UBlackboardData* TestAsset) const
 
 EBlackboardCompare::Type UBlackboardComponent::CompareKeyValues(TSubclassOf<UBlackboardKeyType> KeyType, FBlackboard::FKey KeyA, FBlackboard::FKey KeyB) const
 {	
-	return GetDefault<UBlackboardKeyType>()->Compare(GetKeyRawData(KeyA), GetKeyRawData(KeyB));
+	const uint8* KeyAMemory = GetKeyRawData(KeyA) + (KeyInstances[KeyA] ? sizeof(FBlackboardInstancedKeyMemory) : 0);
+	const uint8* KeyBMemory = GetKeyRawData(KeyB) + (KeyInstances[KeyB] ? sizeof(FBlackboardInstancedKeyMemory) : 0);
+
+	const UBlackboardKeyType* KeyAOb = KeyInstances[KeyA] ? KeyInstances[KeyA] : GetDefault<UBlackboardKeyType>();
+	return KeyAOb->CompareValues(*this, KeyAMemory, KeyInstances[KeyB], KeyBMemory);
 }
 
 FString UBlackboardComponent::GetDebugInfoString(EBlackboardDescription::Type Mode) const 
@@ -436,7 +473,7 @@ FString UBlackboardComponent::DescribeKeyValue(FBlackboard::FKey KeyID, EBlackbo
 	if (Key)
 	{
 		const uint8* ValueData = GetKeyRawData(KeyID);
-		FString ValueDesc = Key->KeyType ? *(Key->KeyType->DescribeValue(ValueData)) : TEXT("empty");
+		FString ValueDesc = Key->KeyType ? *(Key->KeyType->WrappedDescribeValue(*this, ValueData)) : TEXT("empty");
 
 		if (Mode == EBlackboardDescription::OnlyValue)
 		{
@@ -473,7 +510,7 @@ void UBlackboardComponent::DescribeSelfToVisLog(FVisualLogEntry* Snapshot) const
 			const FBlackboardEntry& Key = It->Keys[KeyIndex];
 
 			const uint8* ValueData = GetKeyRawData(It->GetFirstKeyID() + KeyIndex);
-			FString ValueDesc = Key.KeyType ? *(Key.KeyType->DescribeValue(ValueData)) : TEXT("empty");
+			FString ValueDesc = Key.KeyType ? *(Key.KeyType->WrappedDescribeValue(*this, ValueData)) : TEXT("empty");
 
 			Category.Add(Key.EntryName.ToString(), ValueDesc);
 		}
@@ -647,22 +684,23 @@ void UBlackboardComponent::ClearValue(const FName& KeyName)
 void UBlackboardComponent::ClearValue(FBlackboard::FKey KeyID)
 {
 	check(BlackboardAsset);
+	const FBlackboardEntry* EntryInfo = BlackboardAsset->GetKey(KeyID);
 
 	uint8* RawData = GetKeyRawData(KeyID);
 	if (RawData)
 	{
-		TSubclassOf<UBlackboardKeyType> ValueType = GetKeyType(KeyID);
-		check(ValueType);
-		
-		const UBlackboardKeyType* ValueTypeInstance = GetDefault<UBlackboardKeyType>(*ValueType);
-
-		const bool bChanged = ValueTypeInstance->Clear(RawData);
-		if (bChanged)
+		const bool bHasData = EntryInfo->KeyType->WrappedIsEmpty(*this, RawData);
+		if (bHasData)
 		{
+			EntryInfo->KeyType->WrappedClear(*this, RawData);
 			NotifyObservers(KeyID);
 
 			if (BlackboardAsset->HasSyncronizedKeys() && IsKeyInstanceSynced(KeyID))
 			{
+				UBlackboardKeyType* KeyOb = EntryInfo->KeyType->IsInstanced() ? KeyInstances[KeyID] : EntryInfo->KeyType;
+				const uint16 DataOffset = EntryInfo->KeyType->IsInstanced() ? sizeof(FBlackboardInstancedKeyMemory) : 0;
+				uint8* InstancedRawData = RawData + DataOffset;
+
 				// grab the value set and apply the same to synchronized keys
 				// to avoid virtual function call overhead
 				UAISystem* AISystem = UAISystem::GetCurrentSafe(GetWorld());
@@ -671,7 +709,11 @@ void UBlackboardComponent::ClearValue(FBlackboard::FKey KeyID)
 					UBlackboardComponent* OtherBlackboard = Iter.Value();
 					if (OtherBlackboard != nullptr && ShouldSyncWithBlackboard(*OtherBlackboard))
 					{
-						ValueTypeInstance->CopyValue(OtherBlackboard->GetKeyRawData(KeyID), RawData);
+						const FBlackboardEntry* OtherEntryInfo = OtherBlackboard->BlackboardAsset->GetKey(KeyID);
+						UBlackboardKeyType* OtherKeyOb = EntryInfo->KeyType->IsInstanced() ? OtherBlackboard->KeyInstances[KeyID] : EntryInfo->KeyType;
+						uint8* OtherRawData = OtherBlackboard->GetKeyRawData(KeyID) + DataOffset;
+
+						OtherKeyOb->CopyValues(*OtherBlackboard, OtherRawData, KeyOb, InstancedRawData);
 						OtherBlackboard->NotifyObservers(KeyID);
 					}
 				}
@@ -694,7 +736,7 @@ bool UBlackboardComponent::GetLocationFromEntry(FBlackboard::FKey KeyID, FVector
 		if (EntryInfo && EntryInfo->KeyType)
 		{
 			const uint8* ValueData = ValueMemory.GetData() + ValueOffsets[KeyID];
-			return EntryInfo->KeyType->GetLocation(ValueData, ResultLocation);
+			return EntryInfo->KeyType->WrappedGetLocation(*this, ValueData, ResultLocation);
 		}
 	}
 
@@ -715,7 +757,7 @@ bool UBlackboardComponent::GetRotationFromEntry(FBlackboard::FKey KeyID, FRotato
 		if (EntryInfo && EntryInfo->KeyType)
 		{
 			const uint8* ValueData = ValueMemory.GetData() + ValueOffsets[KeyID];
-			return EntryInfo->KeyType->GetRotation(ValueData, ResultRotation);
+			return EntryInfo->KeyType->WrappedGetRotation(*this, ValueData, ResultRotation);
 		}
 	}
 
diff --git a/Engine/Source/Runtime/AIModule/Private/BehaviorTree/Decorators/BTDecorator_Blackboard.cpp b/Engine/Source/Runtime/AIModule/Private/BehaviorTree/Decorators/BTDecorator_Blackboard.cpp
--- a/Engine/Source/Runtime/AIModule/Private/BehaviorTree/Decorators/BTDecorator_Blackboard.cpp
+++ b/Engine/Source/Runtime/AIModule/Private/BehaviorTree/Decorators/BTDecorator_Blackboard.cpp
@@ -32,15 +32,15 @@ bool UBTDecorator_Blackboard::EvaluateOnBlackboard(const UBlackboardComponent& B
 		switch (Op)
 		{
 			case EBlackboardKeyOperation::Basic:
-				bResult = KeyCDO->TestBasicOperation(KeyMemory, (EBasicKeyOperation::Type)OperationType);
+				bResult = KeyCDO->WrappedTestBasicOperation(BlackboardComp, KeyMemory, (EBasicKeyOperation::Type)OperationType);
 				break;
 
 			case EBlackboardKeyOperation::Arithmetic:
-				bResult = KeyCDO->TestArithmeticOperation(KeyMemory, (EArithmeticKeyOperation::Type)OperationType, IntValue, FloatValue);
+				bResult = KeyCDO->WrappedTestArithmeticOperation(BlackboardComp, KeyMemory, (EArithmeticKeyOperation::Type)OperationType, IntValue, FloatValue);
 				break;
 
 			case EBlackboardKeyOperation::Text:
-				bResult = KeyCDO->TestTextOperation(KeyMemory, (ETextKeyOperation::Type)OperationType, StringValue);
+				bResult = KeyCDO->WrappedTestTextOperation(BlackboardComp, KeyMemory, (ETextKeyOperation::Type)OperationType, StringValue);
 				break;
 
 			default:
