fixed disappearing AI messages spawned from existing message observers
#ue4

diff --git a/Engine/Source/Runtime/AIModule/Private/BrainComponent.cpp b/Engine/Source/Runtime/AIModule/Private/BrainComponent.cpp
--- a/Engine/Source/Runtime/AIModule/Private/BrainComponent.cpp
+++ b/Engine/Source/Runtime/AIModule/Private/BrainComponent.cpp
@@ -267,6 +267,7 @@ void UBrainComponent::TickComponent(float DeltaTime, enum ELevelTick TickType, F
 {
 	if (MessagesToProcess.Num() > 0)
 	{
+		const int32 NumMessages = MessagesToProcess.Num();
 		for (auto Message : MessagesToProcess)
 		{
 			for (int32 ObserverIndex = 0; ObserverIndex < MessageObservers.Num(); ObserverIndex++)
@@ -274,7 +275,7 @@ void UBrainComponent::TickComponent(float DeltaTime, enum ELevelTick TickType, F
 				MessageObservers[ObserverIndex]->OnMessage(Message);
 			}
 		}
-		MessagesToProcess.Reset();
+		MessagesToProcess.RemoveAt(0, NumMessages, false);
 	}
 }
 
