Fix for BB component registering/unregistering missmatch merged over from dev branch #UE4

- the missmatch was a result of always unregistering, while registering was being called only when given BB asset had any synchronized keys

diff --git a/Engine/Source/Runtime/AIModule/Classes/BehaviorTree/BlackboardComponent.h b/Engine/Source/Runtime/AIModule/Classes/BehaviorTree/BlackboardComponent.h
--- a/Engine/Source/Runtime/AIModule/Classes/BehaviorTree/BlackboardComponent.h
+++ b/Engine/Source/Runtime/AIModule/Classes/BehaviorTree/BlackboardComponent.h
@@ -345,7 +345,7 @@ bool UBlackboardComponent::SetValue(FBlackboard::FKey KeyID, typename TDataClass
 		if (bChanged)
 		{
 			NotifyObservers(KeyID);
-			if (BlackboardAsset->HasSyncronizedKeys() && IsKeyInstanceSynced(KeyID))
+			if (BlackboardAsset->HasSynchronizedKeys() && IsKeyInstanceSynced(KeyID))
 			{
 				UAISystem* AISystem = UAISystem::GetCurrentSafe(GetWorld());
 				for (auto Iter = AISystem->CreateBlackboardDataToComponentsIterator(*BlackboardAsset); Iter; ++Iter)
diff --git a/Engine/Source/Runtime/AIModule/Classes/BehaviorTree/BlackboardData.h b/Engine/Source/Runtime/AIModule/Classes/BehaviorTree/BlackboardData.h
--- a/Engine/Source/Runtime/AIModule/Classes/BehaviorTree/BlackboardData.h
+++ b/Engine/Source/Runtime/AIModule/Classes/BehaviorTree/BlackboardData.h
@@ -54,11 +54,11 @@ class AIMODULE_API UBlackboardData : public UDataAsset
 
 private:
 	UPROPERTY()
-	uint32 bHasSyncronizedKeys : 1;
+	uint32 bHasSynchronizedKeys : 1;
 
 public:
 
-	FORCEINLINE bool HasSyncronizedKeys() const { return bHasSyncronizedKeys; }
+	FORCEINLINE bool HasSynchronizedKeys() const { return bHasSynchronizedKeys; }
 
 	/** @return true if the key is instance synced */
 	bool IsKeyInstanceSynced(FBlackboard::FKey KeyID) const;
diff --git a/Engine/Source/Runtime/AIModule/Private/BehaviorTree/BlackboardComponent.cpp b/Engine/Source/Runtime/AIModule/Private/BehaviorTree/BlackboardComponent.cpp
--- a/Engine/Source/Runtime/AIModule/Private/BehaviorTree/BlackboardComponent.cpp
+++ b/Engine/Source/Runtime/AIModule/Private/BehaviorTree/BlackboardComponent.cpp
@@ -37,7 +37,7 @@ void UBlackboardComponent::InitializeComponent()
 
 void UBlackboardComponent::UninitializeComponent()
 {
-	if (BlackboardAsset && BlackboardAsset->HasSyncronizedKeys())
+	if (BlackboardAsset && BlackboardAsset->HasSynchronizedKeys())
 	{
 		UAISystem* AISystem = UAISystem::GetCurrentSafe(GetWorld());
 		if (AISystem)
@@ -102,7 +102,7 @@ bool UBlackboardComponent::InitializeBlackboard(UBlackboardData& NewAsset)
 		return false;
 	}
 
-	if (BlackboardAsset)
+	if (BlackboardAsset && BlackboardAsset->HasSynchronizedKeys())
 	{
 		AISystem->UnregisterBlackboardComponent(*BlackboardAsset, *this);
 		DestroyValues();
@@ -159,9 +159,9 @@ bool UBlackboardComponent::InitializeBlackboard(UBlackboardData& NewAsset)
 			const FBlackboardEntry* KeyData = BlackboardAsset->GetKey(InitList[Index].KeyID);
 			KeyData->KeyType->InitializeKey(*this, InitList[Index].KeyID);
 		}
-	
+
 		// naive initial synchronization with one of already instantiated blackboards using the same BB asset
-		if (BlackboardAsset->HasSyncronizedKeys())
+		if (BlackboardAsset->HasSynchronizedKeys())
 		{
 			PopulateSynchronizedKeys();
 		}
@@ -201,7 +201,6 @@ void UBlackboardComponent::PopulateSynchronizedKeys()
 	
 	UAISystem* AISystem = UAISystem::GetCurrentSafe(GetWorld());
 	check(AISystem);
-	AISystem->RegisterBlackboardComponent(*BlackboardAsset, *this);
 
 	for (auto Iter = AISystem->CreateBlackboardDataToComponentsIterator(*BlackboardAsset); Iter; ++Iter)
 	{
@@ -695,7 +694,7 @@ void UBlackboardComponent::ClearValue(FBlackboard::FKey KeyID)
 			EntryInfo->KeyType->WrappedClear(*this, RawData);
 			NotifyObservers(KeyID);
 
-			if (BlackboardAsset->HasSyncronizedKeys() && IsKeyInstanceSynced(KeyID))
+			if (BlackboardAsset->HasSynchronizedKeys() && IsKeyInstanceSynced(KeyID))
 			{
 				UBlackboardKeyType* KeyOb = EntryInfo->KeyType->IsInstanced() ? KeyInstances[KeyID] : EntryInfo->KeyType;
 				const uint16 DataOffset = EntryInfo->KeyType->IsInstanced() ? sizeof(FBlackboardInstancedKeyMemory) : 0;
diff --git a/Engine/Source/Runtime/AIModule/Private/BehaviorTree/BlackboardData.cpp b/Engine/Source/Runtime/AIModule/Private/BehaviorTree/BlackboardData.cpp
--- a/Engine/Source/Runtime/AIModule/Private/BehaviorTree/BlackboardData.cpp
+++ b/Engine/Source/Runtime/AIModule/Private/BehaviorTree/BlackboardData.cpp
@@ -109,10 +109,10 @@ bool UBlackboardData::IsValid() const
 
 void UBlackboardData::UpdateIfHasSynchronizedKeys()
 {
-	bHasSyncronizedKeys = Parent != nullptr && Parent->bHasSyncronizedKeys;
-	for (int32 KeyIndex = 0; KeyIndex < Keys.Num() && bHasSyncronizedKeys == false; ++KeyIndex)
+	bHasSynchronizedKeys = Parent != nullptr && Parent->bHasSynchronizedKeys;
+	for (int32 KeyIndex = 0; KeyIndex < Keys.Num() && bHasSynchronizedKeys == false; ++KeyIndex)
 	{
-		bHasSyncronizedKeys |= Keys[KeyIndex].bInstanceSynced;
+		bHasSynchronizedKeys |= Keys[KeyIndex].bInstanceSynced;
 	}
 }
 
