added missing uproperty markup for instanced blackboard keys
#ue4

diff --git a/Engine/Source/Runtime/AIModule/Classes/BehaviorTree/BlackboardComponent.h b/Engine/Source/Runtime/AIModule/Classes/BehaviorTree/BlackboardComponent.h
--- a/Engine/Source/Runtime/AIModule/Classes/BehaviorTree/BlackboardComponent.h
+++ b/Engine/Source/Runtime/AIModule/Classes/BehaviorTree/BlackboardComponent.h
@@ -278,6 +278,7 @@ protected:
 	TArray<uint16> ValueOffsets;
 
 	/** instanced keys with custom data allocations */
+	UPROPERTY(transient)
 	TArray<UBlackboardKeyType*> KeyInstances;
 
 	/** observers registered for blackboard keys */
