fixed crash on collecting vislog snapshot from uninitialized blackboard
#ue4

diff --git a/Engine/Source/Runtime/AIModule/Private/BehaviorTree/BlackboardComponent.cpp b/Engine/Source/Runtime/AIModule/Private/BehaviorTree/BlackboardComponent.cpp
--- a/Engine/Source/Runtime/AIModule/Private/BehaviorTree/BlackboardComponent.cpp
+++ b/Engine/Source/Runtime/AIModule/Private/BehaviorTree/BlackboardComponent.cpp
@@ -502,20 +502,25 @@ void UBlackboardComponent::DescribeSelfToVisLog(FVisualLogEntry* Snapshot) const
 	FVisualLogStatusCategory Category;
 	Category.Category = FString::Printf(TEXT("Blackboard (asset: %s)"), *GetNameSafe(BlackboardAsset));
 
-	for (UBlackboardData* It = BlackboardAsset; It; It = It->Parent)
+	// describe only when values are initialized
+	if (ValueMemory.Num())
 	{
-		for (int32 KeyIndex = 0; KeyIndex < It->Keys.Num(); KeyIndex++)
+		for (UBlackboardData* It = BlackboardAsset; It; It = It->Parent)
 		{
-			const FBlackboardEntry& Key = It->Keys[KeyIndex];
+			for (int32 KeyIndex = 0; KeyIndex < It->Keys.Num(); KeyIndex++)
+			{
+				const FBlackboardEntry& Key = It->Keys[KeyIndex];
 
-			const uint8* ValueData = GetKeyRawData(It->GetFirstKeyID() + KeyIndex);
-			FString ValueDesc = Key.KeyType ? *(Key.KeyType->WrappedDescribeValue(*this, ValueData)) : TEXT("empty");
+				const uint8* ValueData = GetKeyRawData(It->GetFirstKeyID() + KeyIndex);
+				FString ValueDesc = Key.KeyType ? *(Key.KeyType->WrappedDescribeValue(*this, ValueData)) : TEXT("empty");
 
-			Category.Add(Key.EntryName.ToString(), ValueDesc);
+				Category.Add(Key.EntryName.ToString(), ValueDesc);
+			}
 		}
+
+		Category.Data.Sort();
 	}
 
-	Category.Data.Sort();
 	Snapshot->Status.Add(Category);
 }
 
