behavior's tree stop will call Abort on active tasks
#ue4

diff --git a/Engine/Source/Developer/AITestSuite/Classes/MockAI_BT.h b/Engine/Source/Developer/AITestSuite/Classes/MockAI_BT.h
--- a/Engine/Source/Developer/AITestSuite/Classes/MockAI_BT.h
+++ b/Engine/Source/Developer/AITestSuite/Classes/MockAI_BT.h
@@ -16,5 +16,5 @@ class UMockAI_BT : public UMockAI
 	TArray<int32> ExpectedResult;
 
 	bool IsRunning() const;
-	bool RunBT(UBehaviorTree& BTAsset, EBTExecutionMode::Type RunType = EBTExecutionMode::SingleRun);
+	void RunBT(UBehaviorTree& BTAsset, EBTExecutionMode::Type RunType = EBTExecutionMode::SingleRun);
 };
diff --git a/Engine/Source/Developer/AITestSuite/Private/MockAI/MockAI_BT.cpp b/Engine/Source/Developer/AITestSuite/Private/MockAI/MockAI_BT.cpp
--- a/Engine/Source/Developer/AITestSuite/Private/MockAI/MockAI_BT.cpp
+++ b/Engine/Source/Developer/AITestSuite/Private/MockAI/MockAI_BT.cpp
@@ -25,7 +25,7 @@ bool UMockAI_BT::IsRunning() const
 	return BTComp && BTComp->IsRunning() && BTComp->GetRootTree();
 }
 
-bool UMockAI_BT::RunBT(UBehaviorTree& BTAsset, EBTExecutionMode::Type RunType)
+void UMockAI_BT::RunBT(UBehaviorTree& BTAsset, EBTExecutionMode::Type RunType)
 {
 	if (BTAsset.BlackboardAsset)
 	{
@@ -39,5 +39,5 @@ bool UMockAI_BT::RunBT(UBehaviorTree& BTAsset, EBTExecutionMode::Type RunType)
 	BBComp->RegisterComponentWithWorld(World);
 	BTComp->RegisterComponentWithWorld(World);
 
-	return BTComp->StartTree(BTAsset, RunType);
+	BTComp->StartTree(BTAsset, RunType);
 }
diff --git a/Engine/Source/Runtime/AIModule/Classes/BehaviorTree/BehaviorTreeComponent.h b/Engine/Source/Runtime/AIModule/Classes/BehaviorTree/BehaviorTreeComponent.h
--- a/Engine/Source/Runtime/AIModule/Classes/BehaviorTree/BehaviorTreeComponent.h
+++ b/Engine/Source/Runtime/AIModule/Classes/BehaviorTree/BehaviorTreeComponent.h
@@ -60,6 +60,15 @@ struct FBTPendingExecutionInfo
 	bool IsSet() const { return NextTask || bOutOfNodes; }
 };
 
+struct FBTPendingInitializeInfo
+{
+	UBehaviorTree* Asset;
+	EBTExecutionMode::Type ExecuteMode;
+
+	FBTPendingInitializeInfo() : Asset(nullptr), ExecuteMode(EBTExecutionMode::Looped) {}
+	bool IsSet() const { return Asset != nullptr; }
+};
+
 UCLASS()
 class AIMODULE_API UBehaviorTreeComponent : public UBrainComponent
 {
@@ -86,10 +95,10 @@ public:
 	// End UActorComponent overrides
 
 	/** starts execution from root */
-	bool StartTree(UBehaviorTree& Asset, EBTExecutionMode::Type ExecuteMode = EBTExecutionMode::Looped);
+	void StartTree(UBehaviorTree& Asset, EBTExecutionMode::Type ExecuteMode = EBTExecutionMode::Looped);
 
 	/** stops execution */
-	void StopTree();
+	void StopTree(EBTStopMode::Type StopMode = EBTStopMode::Safe);
 
 	/** restarts execution from root */
 	void RestartTree();
@@ -205,6 +214,9 @@ protected:
 	/** result of ExecutionRequest, will be applied when current task finish aborting */
 	FBTPendingExecutionInfo PendingExecution;
 
+	/** stored data for starting new tree, waits until previously running finishes aborting */
+	FBTPendingInitializeInfo PendingInitialize;
+
 	/** message observers mapped by instance & execution index */
 	TMultiMap<FBTNodeIndex,FAIMessageObserverHandle> TaskMessageObservers;
 
@@ -237,6 +249,9 @@ protected:
 	/** set when execution update is scheduled for next tick */
 	uint8 bRequestedFlowUpdate : 1;
 
+	/** set when tree stop was called */
+	uint8 bRequestedStop : 1;
+
 	/** if set, tree execution is allowed */
 	uint8 bIsRunning : 1;
 
@@ -289,6 +304,9 @@ protected:
 	/** apply pending execution from last task search */
 	void ProcessPendingExecution();
 
+	/** apply pending tree initialization */
+	void ProcessPendingInitialize();
+
 	/** make a snapshot for debugger */
 	void StoreDebuggerExecutionStep(EBTExecutionSnap::Type SnapType);
 
diff --git a/Engine/Source/Runtime/AIModule/Classes/BehaviorTree/BehaviorTreeTypes.h b/Engine/Source/Runtime/AIModule/Classes/BehaviorTree/BehaviorTreeTypes.h
--- a/Engine/Source/Runtime/AIModule/Classes/BehaviorTree/BehaviorTreeTypes.h
+++ b/Engine/Source/Runtime/AIModule/Classes/BehaviorTree/BehaviorTreeTypes.h
@@ -75,6 +75,15 @@ namespace EBTExecutionMode
 	};
 }
 
+namespace EBTStopMode
+{
+	enum Type
+	{
+		Safe,
+		Forced,
+	};
+}
+
 namespace EBTMemoryInit
 {
 	enum Type
diff --git a/Engine/Source/Runtime/AIModule/Private/BehaviorTree/BehaviorTreeComponent.cpp b/Engine/Source/Runtime/AIModule/Private/BehaviorTree/BehaviorTreeComponent.cpp
--- a/Engine/Source/Runtime/AIModule/Private/BehaviorTree/BehaviorTreeComponent.cpp
+++ b/Engine/Source/Runtime/AIModule/Private/BehaviorTree/BehaviorTreeComponent.cpp
@@ -109,7 +109,7 @@ bool UBehaviorTreeComponent::IsPaused() const
 	return bIsPaused;
 }
 
-bool UBehaviorTreeComponent::StartTree(UBehaviorTree& Asset, EBTExecutionMode::Type ExecuteMode /*= EBTExecutionMode::Looped*/)
+void UBehaviorTreeComponent::StartTree(UBehaviorTree& Asset, EBTExecutionMode::Type ExecuteMode /*= EBTExecutionMode::Looped*/)
 {
 	// clear instance stack, start should always run new tree from root
 	UBehaviorTree* CurrentRoot = GetRootTree();
@@ -117,7 +117,7 @@ bool UBehaviorTreeComponent::StartTree(UBehaviorTree& Asset, EBTExecutionMode::T
 	if (CurrentRoot == &Asset && TreeHasBeenStarted())
 	{
 		UE_VLOG(GetOwner(), LogBehaviorTree, Log, TEXT("Skipping behavior start request - it's already running"));
-		return true;
+		return;
 	}
 	else if (CurrentRoot)
 	{
@@ -125,16 +125,32 @@ bool UBehaviorTreeComponent::StartTree(UBehaviorTree& Asset, EBTExecutionMode::T
 			*GetNameSafe(CurrentRoot), *Asset.GetName());
 	}
 
-	StopTree();
+	StopTree(EBTStopMode::Safe);
+
+	PendingInitialize.Asset = &Asset;
+	PendingInitialize.ExecuteMode = ExecuteMode;
+
+	ProcessPendingInitialize();
+}
+
+void UBehaviorTreeComponent::ProcessPendingInitialize()
+{
+	StopTree(EBTStopMode::Safe);
+	if (bWaitingForAbortingTasks)
+	{
+		return;
+	}
+
+	// finish cleanup
 	RemoveAllInstances();
 
-	bLoopExecution = (ExecuteMode == EBTExecutionMode::Looped);
+	bLoopExecution = (PendingInitialize.ExecuteMode == EBTExecutionMode::Looped);
 	bIsRunning = true;
 
 #if USE_BEHAVIORTREE_DEBUGGER
 	DebuggerSteps.Reset();
 #endif
-	
+
 	UBehaviorTreeManager* BTManager = UBehaviorTreeManager::GetCurrent(GetWorld());
 	if (BTManager)
 	{
@@ -142,48 +158,99 @@ bool UBehaviorTreeComponent::StartTree(UBehaviorTree& Asset, EBTExecutionMode::T
 	}
 
 	// push new instance
-	const bool bPushed = PushInstance(Asset);
-	return bPushed;
+	const bool bPushed = PushInstance(*PendingInitialize.Asset);
+
+	PendingInitialize = FBTPendingInitializeInfo();
 }
 
-void UBehaviorTreeComponent::StopTree()
+void UBehaviorTreeComponent::StopTree(EBTStopMode::Type StopMode)
 {
-	// clear current state, don't touch debugger data
-	for (int32 InstanceIndex = 0; InstanceIndex < InstanceStack.Num(); InstanceIndex++)
+	if (!bRequestedStop)
 	{
-		FBehaviorTreeInstance& InstanceInfo = InstanceStack[InstanceIndex];
+		bRequestedStop = true;
 
-		// notify active aux nodes
-		for (int32 AuxIndex = 0; AuxIndex < InstanceInfo.ActiveAuxNodes.Num(); AuxIndex++)
+		for (int32 InstanceIndex = InstanceStack.Num() - 1; InstanceIndex >= 0; InstanceIndex--)
 		{
-			const UBTAuxiliaryNode* AuxNode = InstanceInfo.ActiveAuxNodes[AuxIndex];
-			check(AuxNode != NULL);
-			uint8* NodeMemory = AuxNode->GetNodeMemory<uint8>(InstanceInfo);
-			AuxNode->WrappedOnCeaseRelevant(*this, NodeMemory);
-		}
+			FBehaviorTreeInstance& InstanceInfo = InstanceStack[InstanceIndex];
 
-		// notify active parallel tasks
-		for (int32 ParallelIndex = 0; ParallelIndex < InstanceInfo.ParallelTasks.Num(); ParallelIndex++)
-		{
-			const FBehaviorTreeParallelTask& ParalleInfo = InstanceInfo.ParallelTasks[ParallelIndex];
-			if (ParalleInfo.TaskNode && (ParalleInfo.Status != EBTTaskStatus::Inactive))
+			// notify active aux nodes
+			for (int32 AuxIndex = 0; AuxIndex < InstanceInfo.ActiveAuxNodes.Num(); AuxIndex++)
+			{
+				const UBTAuxiliaryNode* AuxNode = InstanceInfo.ActiveAuxNodes[AuxIndex];
+				check(AuxNode != NULL);
+				uint8* NodeMemory = AuxNode->GetNodeMemory<uint8>(InstanceInfo);
+				AuxNode->WrappedOnCeaseRelevant(*this, NodeMemory);
+			}
+			InstanceInfo.ActiveAuxNodes.Reset();
+
+			// notify active parallel tasks
+			for (int32 ParallelIndex = 0; ParallelIndex < InstanceInfo.ParallelTasks.Num(); ParallelIndex++)
+			{
+				FBehaviorTreeParallelTask& ParalleInfo = InstanceInfo.ParallelTasks[ParallelIndex];
+				if (ParalleInfo.TaskNode && (ParalleInfo.Status == EBTTaskStatus::Active))
+				{
+					// remove all message observers from node to abort to avoid calling OnTaskFinished from AbortTask
+					UnregisterMessageObserversFrom(ParalleInfo.TaskNode);
+
+					uint8* NodeMemory = ParalleInfo.TaskNode->GetNodeMemory<uint8>(InstanceInfo);
+					EBTNodeResult::Type NodeResult = ParalleInfo.TaskNode->WrappedAbortTask(*this, NodeMemory);
+
+					UE_VLOG(GetOwner(), LogBehaviorTree, Log, TEXT("Parallel task aborted: %s (%s)"),
+						*UBehaviorTreeTypes::DescribeNodeHelper(ParalleInfo.TaskNode),
+						(NodeResult == EBTNodeResult::InProgress) ? TEXT("in progress") : TEXT("instant"));
+
+					// mark as pending abort
+					if (NodeResult == EBTNodeResult::InProgress)
+					{
+						ParalleInfo.Status = EBTTaskStatus::Aborting;
+						bWaitingForAbortingTasks = true;
+					}
+
+					OnTaskFinished(ParalleInfo.TaskNode, NodeResult);
+				}
+			}
+
+			// notify active task
+			if (InstanceInfo.ActiveNodeType == EBTActiveNode::ActiveTask)
 			{
-				uint8* NodeMemory = ParalleInfo.TaskNode->GetNodeMemory<uint8>(InstanceInfo);
-				ParalleInfo.TaskNode->WrappedOnTaskFinished(*this, NodeMemory, EBTNodeResult::InProgress);
+				const UBTTaskNode* TaskNode = Cast<const UBTTaskNode>(InstanceInfo.ActiveNode);
+				check(TaskNode != NULL);
+
+				// remove all observers before requesting abort
+				UnregisterMessageObserversFrom(TaskNode);
+				InstanceInfo.ActiveNodeType = EBTActiveNode::AbortingTask;
+
+				UE_VLOG(GetOwner(), LogBehaviorTree, Log, TEXT("Abort task: %s"), *UBehaviorTreeTypes::DescribeNodeHelper(TaskNode));
+
+				// abort task using current state of tree
+				uint8* NodeMemory = TaskNode->GetNodeMemory<uint8>(InstanceInfo);
+				EBTNodeResult::Type TaskResult = TaskNode->WrappedAbortTask(*this, NodeMemory);
+
+				// pass task finished if wasn't already notified (FinishLatentAbort)
+				if (InstanceInfo.ActiveNodeType == EBTActiveNode::AbortingTask)
+				{
+					OnTaskFinished(TaskNode, TaskResult);
+				}
 			}
 		}
+	}
 
-		// notify active task
-		if (InstanceInfo.ActiveNodeType == EBTActiveNode::ActiveTask || InstanceInfo.ActiveNodeType == EBTActiveNode::AbortingTask)
+	if (bWaitingForAbortingTasks)
+	{
+		if (StopMode == EBTStopMode::Safe)
 		{
-			const UBTTaskNode* TaskNode = Cast<const UBTTaskNode>(InstanceInfo.ActiveNode);
-			check(TaskNode != NULL);
-			uint8* NodeMemory = TaskNode->GetNodeMemory<uint8>(InstanceInfo);
-			TaskNode->WrappedOnTaskFinished(*this, NodeMemory, EBTNodeResult::InProgress);
+			UE_VLOG(GetOwner(), LogBehaviorTree, Log, TEXT("StopTree is waiting for aborting tasks to finish..."));
+			return;
 		}
-		
-		// clear instance
-		InstanceStack[InstanceIndex].Cleanup(*this, EBTMemoryClear::StoreSubtree);
+
+		UE_VLOG(GetOwner(), LogBehaviorTree, Warning, TEXT("StopTree was forced while waiting for tasks to finish aborting!"));
+	}
+
+	// clear current state, don't touch debugger data
+	for (int32 InstanceIndex = 0; InstanceIndex < InstanceStack.Num(); InstanceIndex++)
+	{
+		FBehaviorTreeInstance& InstanceInfo = InstanceStack[InstanceIndex];
+		InstanceStack[InstanceIndex].Cleanup(*this, EBTMemoryClear::Destroy);
 	}
 
 	InstanceStack.Reset();
@@ -194,6 +261,7 @@ void UBehaviorTreeComponent::StopTree()
 
 	// make sure to allow new execution requests
 	bRequestedFlowUpdate = false;
+	bRequestedStop = false;
 }
 
 void UBehaviorTreeComponent::RestartTree()
@@ -209,7 +277,7 @@ void UBehaviorTreeComponent::RestartTree()
 
 void UBehaviorTreeComponent::Cleanup()
 {
-	StopTree();
+	StopTree(EBTStopMode::Forced);
 	RemoveAllInstances();
 
 	KnownInstances.Reset();
@@ -276,7 +344,7 @@ void UBehaviorTreeComponent::OnTaskFinished(const UBTTaskNode* TaskNode, EBTNode
 		UpdateAbortingTasks();
 
 		// make sure that we continue execution after all pending latent aborts finished
-		if (!bWaitingForAbortingTasks && bWasWaitingForAbort)
+		if (!bWaitingForAbortingTasks && bWasWaitingForAbort && !bRequestedStop)
 		{
 			ScheduleExecutionUpdate();
 		}
@@ -286,6 +354,11 @@ void UBehaviorTreeComponent::OnTaskFinished(const UBTTaskNode* TaskNode, EBTNode
 		// always update state of aborting tasks
 		UpdateAbortingTasks();
 	}
+
+	if (PendingInitialize.IsSet())
+	{
+		ProcessPendingInitialize();
+	}
 }
 
 void UBehaviorTreeComponent::OnTreeFinished()
@@ -314,7 +387,7 @@ void UBehaviorTreeComponent::OnTreeFinished()
 	}
 	else
 	{
-		StopTree();
+		StopTree(EBTStopMode::Safe);
 	}
 }
 
@@ -817,7 +890,7 @@ void UBehaviorTreeComponent::ApplySearchUpdates(const TArray<FBehaviorTreeSearch
 				UnregisterMessageObserversFrom(UpdateInfo.TaskNode);
 
 				uint8* NodeMemory = (uint8*)UpdateNode->GetNodeMemory<uint8>(UpdateInstance);
-				EBTNodeResult::Type NodeResult = UpdateInfo.TaskNode->AbortTask(*this, NodeMemory);
+				EBTNodeResult::Type NodeResult = UpdateInfo.TaskNode->WrappedAbortTask(*this, NodeMemory);
 
 				UE_VLOG(GetOwner(), LogBehaviorTree, Log, TEXT("Parallel task aborted: %s (%s)"),
 					*UBehaviorTreeTypes::DescribeNodeHelper(UpdateInfo.TaskNode),
@@ -1646,7 +1719,7 @@ void UBehaviorTreeComponent::RemoveAllInstances()
 {
 	if (InstanceStack.Num())
 	{
-		StopTree();
+		StopTree(EBTStopMode::Forced);
 	}
 
 	FBehaviorTreeInstance DummyInstance;
