fixed blackboard key selector not having unique filter objects for each selector property
#ue4

diff --git a/Engine/Source/Runtime/AIModule/Classes/BehaviorTree/BehaviorTreeTypes.h b/Engine/Source/Runtime/AIModule/Classes/BehaviorTree/BehaviorTreeTypes.h
--- a/Engine/Source/Runtime/AIModule/Classes/BehaviorTree/BehaviorTreeTypes.h
+++ b/Engine/Source/Runtime/AIModule/Classes/BehaviorTree/BehaviorTreeTypes.h
@@ -467,16 +467,39 @@ public:
 	FORCEINLINE FBlackboard::FKey GetSelectedKeyID() const { return SelectedKeyID; }
 
 	/** helper functions for setting basic filters */
+	void AddObjectFilter(UObject* Owner, FName PropertyName, TSubclassOf<UObject> AllowedClass);
+	void AddClassFilter(UObject* Owner, FName PropertyName, TSubclassOf<UClass> AllowedClass);
+	void AddEnumFilter(UObject* Owner, FName PropertyName, UEnum* AllowedEnum);
+	void AddNativeEnumFilter(UObject* Owner, FName PropertyName, const FString& AllowedEnumName);
+	void AddIntFilter(UObject* Owner, FName PropertyName);
+	void AddFloatFilter(UObject* Owner, FName PropertyName);
+	void AddBoolFilter(UObject* Owner, FName PropertyName);
+	void AddVectorFilter(UObject* Owner, FName PropertyName);
+	void AddRotatorFilter(UObject* Owner, FName PropertyName);
+	void AddStringFilter(UObject* Owner, FName PropertyName);
+	void AddNameFilter(UObject* Owner, FName PropertyName);
+
+	DEPRECATED(4.8, "This version is deprecated, please use override with PropertyName.")
 	void AddObjectFilter(UObject* Owner, TSubclassOf<UObject> AllowedClass);
+	DEPRECATED(4.8, "This version is deprecated, please use override with PropertyName.")
 	void AddClassFilter(UObject* Owner, TSubclassOf<UClass> AllowedClass);
+	DEPRECATED(4.8, "This version is deprecated, please use override with PropertyName.")
 	void AddEnumFilter(UObject* Owner, UEnum* AllowedEnum);
+	DEPRECATED(4.8, "This version is deprecated, please use override with PropertyName.")
 	void AddNativeEnumFilter(UObject* Owner, const FString& AllowedEnumName);
+	DEPRECATED(4.8, "This version is deprecated, please use override with PropertyName.")
 	void AddIntFilter(UObject* Owner);
+	DEPRECATED(4.8, "This version is deprecated, please use override with PropertyName.")
 	void AddFloatFilter(UObject* Owner);
+	DEPRECATED(4.8, "This version is deprecated, please use override with PropertyName.")
 	void AddBoolFilter(UObject* Owner);
+	DEPRECATED(4.8, "This version is deprecated, please use override with PropertyName.")
 	void AddVectorFilter(UObject* Owner);
+	DEPRECATED(4.8, "This version is deprecated, please use override with PropertyName.")
 	void AddRotatorFilter(UObject* Owner);
+	DEPRECATED(4.8, "This version is deprecated, please use override with PropertyName.")
 	void AddStringFilter(UObject* Owner);
+	DEPRECATED(4.8, "This version is deprecated, please use override with PropertyName.")
 	void AddNameFilter(UObject* Owner);
 
 	FORCEINLINE bool IsNone() const { return bNoneIsAllowedValue && SelectedKeyID == FBlackboard::InvalidKey; }
diff --git a/Engine/Source/Runtime/AIModule/Private/BehaviorTree/BehaviorTreeTypes.cpp b/Engine/Source/Runtime/AIModule/Private/BehaviorTree/BehaviorTreeTypes.cpp
--- a/Engine/Source/Runtime/AIModule/Private/BehaviorTree/BehaviorTreeTypes.cpp
+++ b/Engine/Source/Runtime/AIModule/Private/BehaviorTree/BehaviorTreeTypes.cpp
@@ -226,65 +226,131 @@ void FBlackboardKeySelector::InitSelectedKey(UBlackboardData* BlackboardAsset)
 
 void FBlackboardKeySelector::AddObjectFilter(UObject* Owner, TSubclassOf<UObject> AllowedClass)
 {
-	UBlackboardKeyType_Object* FilterOb = Owner->CreateDefaultSubobject<UBlackboardKeyType_Object>(TEXT("BlackboardKeyType_Object"));
-	FilterOb->BaseClass = AllowedClass;
-	AllowedTypes.Add(FilterOb);
+	AddObjectFilter(Owner, TEXT("BlackboardKeyType"), AllowedClass);
 }
 
 void FBlackboardKeySelector::AddClassFilter(UObject* Owner, TSubclassOf<UClass> AllowedClass)
 {
-	UBlackboardKeyType_Class* FilterOb = Owner->CreateDefaultSubobject<UBlackboardKeyType_Class>(TEXT("BlackboardKeyType_Class"));
-	FilterOb->BaseClass = AllowedClass;
-	AllowedTypes.Add(FilterOb);
+	AddClassFilter(Owner, TEXT("BlackboardKeyType"), AllowedClass);
 }
 
 void FBlackboardKeySelector::AddEnumFilter(UObject* Owner, UEnum* AllowedEnum)
 {
-	UBlackboardKeyType_Enum* FilterOb = Owner->CreateDefaultSubobject<UBlackboardKeyType_Enum>(TEXT("BlackboardKeyType_Enum"));
-	FilterOb->EnumType = AllowedEnum;
-	AllowedTypes.Add(FilterOb);
+	AddEnumFilter(Owner, TEXT("BlackboardKeyType"), AllowedEnum);
 }
 
 void FBlackboardKeySelector::AddNativeEnumFilter(UObject* Owner, const FString& AllowedEnumName)
 {
-	UBlackboardKeyType_NativeEnum* FilterOb = Owner->CreateDefaultSubobject<UBlackboardKeyType_NativeEnum>(TEXT("BlackboardKeyType_NativeEnum"));
-	FilterOb->EnumName = AllowedEnumName;
-	AllowedTypes.Add(FilterOb);
+	AddNativeEnumFilter(Owner, TEXT("BlackboardKeyType"), AllowedEnumName);
 }
 
 void FBlackboardKeySelector::AddIntFilter(UObject* Owner)
 {
-	AllowedTypes.Add(Owner->CreateDefaultSubobject<UBlackboardKeyType_Int>(TEXT("BlackboardKeyType_Int")));
+	AddIntFilter(Owner, TEXT("BlackboardKeyType"));
 }
 
 void FBlackboardKeySelector::AddFloatFilter(UObject* Owner)
 {
-	AllowedTypes.Add(Owner->CreateDefaultSubobject<UBlackboardKeyType_Float>(TEXT("BlackboardKeyType_Float")));
+	AddFloatFilter(Owner, TEXT("BlackboardKeyType"));
 }
 
 void FBlackboardKeySelector::AddBoolFilter(UObject* Owner)
 {
-	AllowedTypes.Add(Owner->CreateDefaultSubobject<UBlackboardKeyType_Bool>(TEXT("BlackboardKeyType_Bool")));
+	AddBoolFilter(Owner, TEXT("BlackboardKeyType"));
 }
 
 void FBlackboardKeySelector::AddVectorFilter(UObject* Owner)
 {
-	AllowedTypes.Add(Owner->CreateDefaultSubobject<UBlackboardKeyType_Vector>(TEXT("BlackboardKeyType_Vector")));
+	AddVectorFilter(Owner, TEXT("BlackboardKeyType"));
 }
 
 void FBlackboardKeySelector::AddRotatorFilter(UObject* Owner)
 {
-	AllowedTypes.Add(Owner->CreateDefaultSubobject<UBlackboardKeyType_Rotator>(TEXT("BlackboardKeyType_Rotator")));
+	AddRotatorFilter(Owner, TEXT("BlackboardKeyType"));
 }
 
 void FBlackboardKeySelector::AddStringFilter(UObject* Owner)
 {
-	AllowedTypes.Add(Owner->CreateDefaultSubobject<UBlackboardKeyType_String>(TEXT("BlackboardKeyType_String")));
+	AddStringFilter(Owner, TEXT("BlackboardKeyType"));
 }
 
 void FBlackboardKeySelector::AddNameFilter(UObject* Owner)
 {
-	AllowedTypes.Add(Owner->CreateDefaultSubobject<UBlackboardKeyType_Name>(TEXT("BlackboardKeyType_Name")));
+	AddNameFilter(Owner, TEXT("BlackboardKeyType"));
+}
+
+void FBlackboardKeySelector::AddObjectFilter(UObject* Owner, FName PropertyName, TSubclassOf<UObject> AllowedClass)
+{
+	const FString FilterName = PropertyName.ToString() + TEXT("_Object");
+	UBlackboardKeyType_Object* FilterOb = Owner->CreateDefaultSubobject<UBlackboardKeyType_Object>(*FilterName);
+	FilterOb->BaseClass = AllowedClass;
+	AllowedTypes.Add(FilterOb);
+}
+
+void FBlackboardKeySelector::AddClassFilter(UObject* Owner, FName PropertyName, TSubclassOf<UClass> AllowedClass)
+{
+	const FString FilterName = PropertyName.ToString() + TEXT("_Class");
+	UBlackboardKeyType_Class* FilterOb = Owner->CreateDefaultSubobject<UBlackboardKeyType_Class>(*FilterName);
+	FilterOb->BaseClass = AllowedClass;
+	AllowedTypes.Add(FilterOb);
+}
+
+void FBlackboardKeySelector::AddEnumFilter(UObject* Owner, FName PropertyName, UEnum* AllowedEnum)
+{
+	const FString FilterName = PropertyName.ToString() + TEXT("_Enum");
+	UBlackboardKeyType_Enum* FilterOb = Owner->CreateDefaultSubobject<UBlackboardKeyType_Enum>(*FilterName);
+	FilterOb->EnumType = AllowedEnum;
+	AllowedTypes.Add(FilterOb);
+}
+
+void FBlackboardKeySelector::AddNativeEnumFilter(UObject* Owner, FName PropertyName, const FString& AllowedEnumName)
+{
+	const FString FilterName = PropertyName.ToString() + TEXT("_NativeEnum");
+	UBlackboardKeyType_NativeEnum* FilterOb = Owner->CreateDefaultSubobject<UBlackboardKeyType_NativeEnum>(*FilterName);
+	FilterOb->EnumName = AllowedEnumName;
+	AllowedTypes.Add(FilterOb);
+}
+
+void FBlackboardKeySelector::AddIntFilter(UObject* Owner, FName PropertyName)
+{
+	const FString FilterName = PropertyName.ToString() + TEXT("_Int");
+	AllowedTypes.Add(Owner->CreateDefaultSubobject<UBlackboardKeyType_Int>(*FilterName));
+}
+
+void FBlackboardKeySelector::AddFloatFilter(UObject* Owner, FName PropertyName)
+{
+	const FString FilterName = PropertyName.ToString() + TEXT("_Float");
+	AllowedTypes.Add(Owner->CreateDefaultSubobject<UBlackboardKeyType_Float>(*FilterName));
+}
+
+void FBlackboardKeySelector::AddBoolFilter(UObject* Owner, FName PropertyName)
+{
+	const FString FilterName = PropertyName.ToString() + TEXT("_Bool");
+	AllowedTypes.Add(Owner->CreateDefaultSubobject<UBlackboardKeyType_Bool>(*FilterName));
+}
+
+void FBlackboardKeySelector::AddVectorFilter(UObject* Owner, FName PropertyName)
+{
+	const FString FilterName = PropertyName.ToString() + TEXT("_Vector");
+	AllowedTypes.Add(Owner->CreateDefaultSubobject<UBlackboardKeyType_Vector>(*FilterName));
+}
+
+void FBlackboardKeySelector::AddRotatorFilter(UObject* Owner, FName PropertyName)
+{
+	const FString FilterName = PropertyName.ToString() + TEXT("_Rotator");
+	AllowedTypes.Add(Owner->CreateDefaultSubobject<UBlackboardKeyType_Rotator>(*FilterName));
+}
+
+void FBlackboardKeySelector::AddStringFilter(UObject* Owner, FName PropertyName)
+{
+	const FString FilterName = PropertyName.ToString() + TEXT("_String");
+	AllowedTypes.Add(Owner->CreateDefaultSubobject<UBlackboardKeyType_String>(*FilterName));
+}
+
+void FBlackboardKeySelector::AddNameFilter(UObject* Owner, FName PropertyName)
+{
+	const FString FilterName = PropertyName.ToString() + TEXT("_Name");
+	AllowedTypes.Add(Owner->CreateDefaultSubobject<UBlackboardKeyType_Name>(*FilterName));
 }
 
 //----------------------------------------------------------------------//
diff --git a/Engine/Source/Runtime/AIModule/Private/BehaviorTree/Decorators/BTDecorator_CheckGameplayTagsOnActor.cpp b/Engine/Source/Runtime/AIModule/Private/BehaviorTree/Decorators/BTDecorator_CheckGameplayTagsOnActor.cpp
--- a/Engine/Source/Runtime/AIModule/Private/BehaviorTree/Decorators/BTDecorator_CheckGameplayTagsOnActor.cpp
+++ b/Engine/Source/Runtime/AIModule/Private/BehaviorTree/Decorators/BTDecorator_CheckGameplayTagsOnActor.cpp
@@ -10,7 +10,7 @@ UBTDecorator_CheckGameplayTagsOnActor::UBTDecorator_CheckGameplayTagsOnActor(con
 	NodeName = "Gameplay Tag Condition";
 
 	// Accept only actors
-	ActorToCheck.AddObjectFilter(this, AActor::StaticClass());
+	ActorToCheck.AddObjectFilter(this, GET_MEMBER_NAME_CHECKED(UBTDecorator_CheckGameplayTagsOnActor, ActorToCheck), AActor::StaticClass());
 
 	// Default to using Self Actor
 	ActorToCheck.SelectedKeyName = FBlackboard::KeySelf;
diff --git a/Engine/Source/Runtime/AIModule/Private/BehaviorTree/Decorators/BTDecorator_ConeCheck.cpp b/Engine/Source/Runtime/AIModule/Private/BehaviorTree/Decorators/BTDecorator_ConeCheck.cpp
--- a/Engine/Source/Runtime/AIModule/Private/BehaviorTree/Decorators/BTDecorator_ConeCheck.cpp
+++ b/Engine/Source/Runtime/AIModule/Private/BehaviorTree/Decorators/BTDecorator_ConeCheck.cpp
@@ -9,13 +9,13 @@ UBTDecorator_ConeCheck::UBTDecorator_ConeCheck(const FObjectInitializer& ObjectI
 	NodeName = "Cone Check";
 
 	// accept only actors and vectors
-	ConeOrigin.AddObjectFilter(this, AActor::StaticClass());
-	ConeOrigin.AddVectorFilter(this);
-	ConeDirection.AddObjectFilter(this, AActor::StaticClass());
-	ConeDirection.AddVectorFilter(this);
+	ConeOrigin.AddObjectFilter(this, GET_MEMBER_NAME_CHECKED(UBTDecorator_ConeCheck, ConeOrigin), AActor::StaticClass());
+	ConeOrigin.AddVectorFilter(this, GET_MEMBER_NAME_CHECKED(UBTDecorator_ConeCheck, ConeOrigin));
+	ConeDirection.AddObjectFilter(this, GET_MEMBER_NAME_CHECKED(UBTDecorator_ConeCheck, ConeDirection), AActor::StaticClass());
+	ConeDirection.AddVectorFilter(this, GET_MEMBER_NAME_CHECKED(UBTDecorator_ConeCheck, ConeDirection));
 	ConeDirection.AllowNoneAsValue(true);
-	Observed.AddObjectFilter(this, AActor::StaticClass());
-	Observed.AddVectorFilter(this);
+	Observed.AddObjectFilter(this, GET_MEMBER_NAME_CHECKED(UBTDecorator_ConeCheck, Observed), AActor::StaticClass());
+	Observed.AddVectorFilter(this, GET_MEMBER_NAME_CHECKED(UBTDecorator_ConeCheck, Observed));
 
 	bNotifyBecomeRelevant = true;
 	bNotifyTick = true;
diff --git a/Engine/Source/Runtime/AIModule/Private/BehaviorTree/Decorators/BTDecorator_DoesPathExist.cpp b/Engine/Source/Runtime/AIModule/Private/BehaviorTree/Decorators/BTDecorator_DoesPathExist.cpp
--- a/Engine/Source/Runtime/AIModule/Private/BehaviorTree/Decorators/BTDecorator_DoesPathExist.cpp
+++ b/Engine/Source/Runtime/AIModule/Private/BehaviorTree/Decorators/BTDecorator_DoesPathExist.cpp
@@ -12,10 +12,10 @@ UBTDecorator_DoesPathExist::UBTDecorator_DoesPathExist(const FObjectInitializer&
 	NodeName = "Does path exist";
 
 	// accept only actors and vectors
-	BlackboardKeyA.AddObjectFilter(this, AActor::StaticClass());
-	BlackboardKeyA.AddVectorFilter(this);
-	BlackboardKeyB.AddObjectFilter(this, AActor::StaticClass());
-	BlackboardKeyB.AddVectorFilter(this);
+	BlackboardKeyA.AddObjectFilter(this, GET_MEMBER_NAME_CHECKED(UBTDecorator_DoesPathExist, BlackboardKeyA), AActor::StaticClass());
+	BlackboardKeyA.AddVectorFilter(this, GET_MEMBER_NAME_CHECKED(UBTDecorator_DoesPathExist, BlackboardKeyA));
+	BlackboardKeyB.AddObjectFilter(this, GET_MEMBER_NAME_CHECKED(UBTDecorator_DoesPathExist, BlackboardKeyB), AActor::StaticClass());
+	BlackboardKeyB.AddVectorFilter(this, GET_MEMBER_NAME_CHECKED(UBTDecorator_DoesPathExist, BlackboardKeyB));
 
 	bAllowAbortLowerPri = false;
 	bAllowAbortNone = true;
diff --git a/Engine/Source/Runtime/AIModule/Private/BehaviorTree/Decorators/BTDecorator_IsBBEntryOfClass.cpp b/Engine/Source/Runtime/AIModule/Private/BehaviorTree/Decorators/BTDecorator_IsBBEntryOfClass.cpp
--- a/Engine/Source/Runtime/AIModule/Private/BehaviorTree/Decorators/BTDecorator_IsBBEntryOfClass.cpp
+++ b/Engine/Source/Runtime/AIModule/Private/BehaviorTree/Decorators/BTDecorator_IsBBEntryOfClass.cpp
@@ -8,7 +8,7 @@ UBTDecorator_IsBBEntryOfClass::UBTDecorator_IsBBEntryOfClass(const FObjectInitia
 {
 	NodeName = TEXT("Is BlackBoard value of given Class");
 	
-	BlackboardKey.AddObjectFilter(this, UObject::StaticClass());
+	BlackboardKey.AddObjectFilter(this, GET_MEMBER_NAME_CHECKED(UBTDecorator_IsBBEntryOfClass, BlackboardKey), UObject::StaticClass());
 }
 
 // @todo refactor CalculateRawConditionValue to take const UBlackboardComponent& instead
diff --git a/Engine/Source/Runtime/AIModule/Private/BehaviorTree/Decorators/BTDecorator_KeepInCone.cpp b/Engine/Source/Runtime/AIModule/Private/BehaviorTree/Decorators/BTDecorator_KeepInCone.cpp
--- a/Engine/Source/Runtime/AIModule/Private/BehaviorTree/Decorators/BTDecorator_KeepInCone.cpp
+++ b/Engine/Source/Runtime/AIModule/Private/BehaviorTree/Decorators/BTDecorator_KeepInCone.cpp
@@ -9,10 +9,10 @@ UBTDecorator_KeepInCone::UBTDecorator_KeepInCone(const FObjectInitializer& Objec
 	NodeName = "Keep in Cone";
 
 	// accept only actors and vectors
-	ConeOrigin.AddObjectFilter(this, AActor::StaticClass());
-	ConeOrigin.AddVectorFilter(this);
-	Observed.AddObjectFilter(this, AActor::StaticClass());
-	Observed.AddVectorFilter(this);
+	ConeOrigin.AddObjectFilter(this, GET_MEMBER_NAME_CHECKED(UBTDecorator_KeepInCone, ConeOrigin), AActor::StaticClass());
+	ConeOrigin.AddVectorFilter(this, GET_MEMBER_NAME_CHECKED(UBTDecorator_KeepInCone, ConeOrigin));
+	Observed.AddObjectFilter(this, GET_MEMBER_NAME_CHECKED(UBTDecorator_KeepInCone, Observed), AActor::StaticClass());
+	Observed.AddVectorFilter(this, GET_MEMBER_NAME_CHECKED(UBTDecorator_KeepInCone, Observed));
 
 	bNotifyBecomeRelevant = true;
 	bNotifyTick = true;
diff --git a/Engine/Source/Runtime/AIModule/Private/BehaviorTree/Services/BTService_DefaultFocus.cpp b/Engine/Source/Runtime/AIModule/Private/BehaviorTree/Services/BTService_DefaultFocus.cpp
--- a/Engine/Source/Runtime/AIModule/Private/BehaviorTree/Services/BTService_DefaultFocus.cpp
+++ b/Engine/Source/Runtime/AIModule/Private/BehaviorTree/Services/BTService_DefaultFocus.cpp
@@ -17,8 +17,8 @@ UBTService_DefaultFocus::UBTService_DefaultFocus(const FObjectInitializer& Objec
 	FocusPriority = EAIFocusPriority::Default;
 
 	// accept only actors and vectors
-	BlackboardKey.AddObjectFilter(this, AActor::StaticClass());
-	BlackboardKey.AddVectorFilter(this);
+	BlackboardKey.AddObjectFilter(this, GET_MEMBER_NAME_CHECKED(UBTService_DefaultFocus, BlackboardKey), AActor::StaticClass());
+	BlackboardKey.AddVectorFilter(this, GET_MEMBER_NAME_CHECKED(UBTService_DefaultFocus, BlackboardKey));
 }
 
 void UBTService_DefaultFocus::OnBecomeRelevant(UBehaviorTreeComponent& OwnerComp, uint8* NodeMemory)
diff --git a/Engine/Source/Runtime/AIModule/Private/BehaviorTree/Tasks/BTTask_MoveDirectlyToward.cpp b/Engine/Source/Runtime/AIModule/Private/BehaviorTree/Tasks/BTTask_MoveDirectlyToward.cpp
--- a/Engine/Source/Runtime/AIModule/Private/BehaviorTree/Tasks/BTTask_MoveDirectlyToward.cpp
+++ b/Engine/Source/Runtime/AIModule/Private/BehaviorTree/Tasks/BTTask_MoveDirectlyToward.cpp
@@ -16,8 +16,8 @@ UBTTask_MoveDirectlyToward::UBTTask_MoveDirectlyToward(const FObjectInitializer&
 	NodeName = "MoveDirectlyToward";
 
 	// accept only actors and vectors
-	BlackboardKey.AddObjectFilter(this, AActor::StaticClass());
-	BlackboardKey.AddVectorFilter(this);
+	BlackboardKey.AddObjectFilter(this, GET_MEMBER_NAME_CHECKED(UBTTask_MoveDirectlyToward, BlackboardKey), AActor::StaticClass());
+	BlackboardKey.AddVectorFilter(this, GET_MEMBER_NAME_CHECKED(UBTTask_MoveDirectlyToward, BlackboardKey));
 }
 
 EBTNodeResult::Type UBTTask_MoveDirectlyToward::ExecuteTask(UBehaviorTreeComponent& OwnerComp, uint8* NodeMemory)
diff --git a/Engine/Source/Runtime/AIModule/Private/BehaviorTree/Tasks/BTTask_MoveTo.cpp b/Engine/Source/Runtime/AIModule/Private/BehaviorTree/Tasks/BTTask_MoveTo.cpp
--- a/Engine/Source/Runtime/AIModule/Private/BehaviorTree/Tasks/BTTask_MoveTo.cpp
+++ b/Engine/Source/Runtime/AIModule/Private/BehaviorTree/Tasks/BTTask_MoveTo.cpp
@@ -17,8 +17,8 @@ UBTTask_MoveTo::UBTTask_MoveTo(const FObjectInitializer& ObjectInitializer)
 	bNotifyTick = true;
 
 	// accept only actors and vectors
-	BlackboardKey.AddObjectFilter(this, AActor::StaticClass());
-	BlackboardKey.AddVectorFilter(this);
+	BlackboardKey.AddObjectFilter(this, GET_MEMBER_NAME_CHECKED(UBTTask_MoveTo, BlackboardKey), AActor::StaticClass());
+	BlackboardKey.AddVectorFilter(this, GET_MEMBER_NAME_CHECKED(UBTTask_MoveTo, BlackboardKey));
 }
 
 EBTNodeResult::Type UBTTask_MoveTo::ExecuteTask(UBehaviorTreeComponent& OwnerComp, uint8* NodeMemory)
diff --git a/Engine/Source/Runtime/AIModule/Private/BehaviorTree/Tasks/BTTask_RotateToFaceBBEntry.cpp b/Engine/Source/Runtime/AIModule/Private/BehaviorTree/Tasks/BTTask_RotateToFaceBBEntry.cpp
--- a/Engine/Source/Runtime/AIModule/Private/BehaviorTree/Tasks/BTTask_RotateToFaceBBEntry.cpp
+++ b/Engine/Source/Runtime/AIModule/Private/BehaviorTree/Tasks/BTTask_RotateToFaceBBEntry.cpp
@@ -12,9 +12,9 @@ UBTTask_RotateToFaceBBEntry::UBTTask_RotateToFaceBBEntry(const FObjectInitialize
 	bNotifyTick = true;
 	
 	// accept only actors and vectors
-	BlackboardKey.AddObjectFilter(this, AActor::StaticClass());
-	BlackboardKey.AddVectorFilter(this);
-	BlackboardKey.AddRotatorFilter(this);
+	BlackboardKey.AddObjectFilter(this, GET_MEMBER_NAME_CHECKED(UBTTask_RotateToFaceBBEntry, BlackboardKey), AActor::StaticClass());
+	BlackboardKey.AddVectorFilter(this, GET_MEMBER_NAME_CHECKED(UBTTask_RotateToFaceBBEntry, BlackboardKey));
+	BlackboardKey.AddRotatorFilter(this, GET_MEMBER_NAME_CHECKED(UBTTask_RotateToFaceBBEntry, BlackboardKey));
 }
 
 void UBTTask_RotateToFaceBBEntry::PostInitProperties()
diff --git a/Engine/Source/Runtime/AIModule/Private/BehaviorTree/Tasks/BTTask_WaitBlackboardTime.cpp b/Engine/Source/Runtime/AIModule/Private/BehaviorTree/Tasks/BTTask_WaitBlackboardTime.cpp
--- a/Engine/Source/Runtime/AIModule/Private/BehaviorTree/Tasks/BTTask_WaitBlackboardTime.cpp
+++ b/Engine/Source/Runtime/AIModule/Private/BehaviorTree/Tasks/BTTask_WaitBlackboardTime.cpp
@@ -9,7 +9,7 @@ UBTTask_WaitBlackboardTime::UBTTask_WaitBlackboardTime(const FObjectInitializer&
 {
 	NodeName = "Wait Blackboard Time";
 
-	BlackboardKey.AddFloatFilter(this);
+	BlackboardKey.AddFloatFilter(this, GET_MEMBER_NAME_CHECKED(UBTTask_WaitBlackboardTime, BlackboardKey));
 }
 
 void UBTTask_WaitBlackboardTime::InitializeFromAsset(UBehaviorTree& Asset)
diff --git a/Engine/Source/Runtime/AIModule/Private/EnvironmentQuery/Items/EnvQueryItemType_ActorBase.cpp b/Engine/Source/Runtime/AIModule/Private/EnvironmentQuery/Items/EnvQueryItemType_ActorBase.cpp
--- a/Engine/Source/Runtime/AIModule/Private/EnvironmentQuery/Items/EnvQueryItemType_ActorBase.cpp
+++ b/Engine/Source/Runtime/AIModule/Private/EnvironmentQuery/Items/EnvQueryItemType_ActorBase.cpp
@@ -12,7 +12,7 @@ UEnvQueryItemType_ActorBase::UEnvQueryItemType_ActorBase(const FObjectInitialize
 void UEnvQueryItemType_ActorBase::AddBlackboardFilters(FBlackboardKeySelector& KeySelector, UObject* FilterOwner) const
 {
 	Super::AddBlackboardFilters(KeySelector, FilterOwner);
-	KeySelector.AddObjectFilter(FilterOwner, AActor::StaticClass());
+	KeySelector.AddObjectFilter(FilterOwner, TEXT("EnvQueryActor"), AActor::StaticClass());
 }
 
 bool UEnvQueryItemType_ActorBase::StoreInBlackboard(FBlackboardKeySelector& KeySelector, UBlackboardComponent* Blackboard, const uint8* RawData) const
diff --git a/Engine/Source/Runtime/AIModule/Private/EnvironmentQuery/Items/EnvQueryItemType_VectorBase.cpp b/Engine/Source/Runtime/AIModule/Private/EnvironmentQuery/Items/EnvQueryItemType_VectorBase.cpp
--- a/Engine/Source/Runtime/AIModule/Private/EnvironmentQuery/Items/EnvQueryItemType_VectorBase.cpp
+++ b/Engine/Source/Runtime/AIModule/Private/EnvironmentQuery/Items/EnvQueryItemType_VectorBase.cpp
@@ -11,7 +11,7 @@ UEnvQueryItemType_VectorBase::UEnvQueryItemType_VectorBase(const FObjectInitiali
 
 void UEnvQueryItemType_VectorBase::AddBlackboardFilters(FBlackboardKeySelector& KeySelector, UObject* FilterOwner) const
 {
-	KeySelector.AddVectorFilter(FilterOwner);
+	KeySelector.AddVectorFilter(FilterOwner, TEXT("EnvQueryVector"));
 }
 
 bool UEnvQueryItemType_VectorBase::StoreInBlackboard(FBlackboardKeySelector& KeySelector, UBlackboardComponent* Blackboard, const uint8* RawData) const
