Fixed UBlackboardComponent::ClearValue not working due to a wrong condition check #UE4

As pointed out by user serioussam909 on a related AnswerHub thread. Thanks!

diff --git a/Engine/Source/Runtime/AIModule/Private/BehaviorTree/BlackboardComponent.cpp b/Engine/Source/Runtime/AIModule/Private/BehaviorTree/BlackboardComponent.cpp
--- a/Engine/Source/Runtime/AIModule/Private/BehaviorTree/BlackboardComponent.cpp
+++ b/Engine/Source/Runtime/AIModule/Private/BehaviorTree/BlackboardComponent.cpp
@@ -693,7 +693,7 @@ void UBlackboardComponent::ClearValue(FBlackboard::FKey KeyID)
 	uint8* RawData = GetKeyRawData(KeyID);
 	if (RawData)
 	{
-		const bool bHasData = EntryInfo->KeyType->WrappedIsEmpty(*this, RawData);
+		const bool bHasData = (EntryInfo->KeyType->WrappedIsEmpty(*this, RawData) == false);
 		if (bHasData)
 		{
 			EntryInfo->KeyType->WrappedClear(*this, RawData);
