unified shared parts of AI editors (nested nodes, copy/paste/drag operations)

diff --git a/Engine/Source/Runtime/AIModule/Private/EnvironmentQuery/EnvQueryDebugHelpers.cpp b/Engine/Source/Runtime/AIModule/Private/EnvironmentQuery/EnvQueryDebugHelpers.cpp
--- a/Engine/Source/Runtime/AIModule/Private/EnvironmentQuery/EnvQueryDebugHelpers.cpp
+++ b/Engine/Source/Runtime/AIModule/Private/EnvironmentQuery/EnvQueryDebugHelpers.cpp
@@ -88,7 +88,7 @@ void UEnvQueryDebugHelpers::QueryToDebugData(FEnvQueryInstance& Query, EQSDebug:
 		EQSDebug::FTestData TestInfo;
 
 		UEnvQueryTest* TestOb = Query.Options[Query.OptionIndex].Tests[TestIdx];
-		TestInfo.ShortName = TestOb->GetDescriptionTitle();
+		TestInfo.ShortName = TestOb->GetDescriptionTitle().ToString();
 		TestInfo.Detailed = TestOb->GetDescriptionDetails().ToString().Replace(TEXT("\n"), TEXT(", "));
 
 		EQSLocalData.Tests.Add(TestInfo);
diff --git a/Engine/Source/Runtime/AIModule/Private/EnvironmentQuery/EnvQueryGenerator.cpp b/Engine/Source/Runtime/AIModule/Private/EnvironmentQuery/EnvQueryGenerator.cpp
--- a/Engine/Source/Runtime/AIModule/Private/EnvironmentQuery/EnvQueryGenerator.cpp
+++ b/Engine/Source/Runtime/AIModule/Private/EnvironmentQuery/EnvQueryGenerator.cpp
@@ -8,33 +8,13 @@ UEnvQueryGenerator::UEnvQueryGenerator(const FObjectInitializer& ObjectInitializ
 {
 }
 
-void UEnvQueryGenerator::UpdateGeneratorVersion()
+void UEnvQueryGenerator::UpdateNodeVersion()
 {
 	VerNum = EnvQueryGeneratorVersion::Latest;
 }
 
-FText UEnvQueryGenerator::GetDescriptionTitle() const
-{
-	return UEnvQueryTypes::GetShortTypeName(this);
-}
-
-FText UEnvQueryGenerator::GetDescriptionDetails() const
-{
-	return FText::GetEmpty();
-}
-
-#if WITH_EDITOR && USE_EQS_DEBUGGER
-void UEnvQueryGenerator::PostEditChangeProperty(FPropertyChangedEvent& PropertyChangedEvent) 
-{
-	Super::PostEditChangeProperty(PropertyChangedEvent);
-#if USE_EQS_DEBUGGER
-	UEnvQueryManager::NotifyAssetUpdate(NULL);
-#endif
-}
-#endif //WITH_EDITOR && USE_EQS_DEBUGGER
-
 void UEnvQueryGenerator::PostLoad()
 {
 	Super::PostLoad();
-	UpdateGeneratorVersion();
-}
\ No newline at end of file
+	UpdateNodeVersion();
+}
diff --git a/Engine/Source/Runtime/AIModule/Private/EnvironmentQuery/EnvQueryTest.cpp b/Engine/Source/Runtime/AIModule/Private/EnvironmentQuery/EnvQueryTest.cpp
--- a/Engine/Source/Runtime/AIModule/Private/EnvironmentQuery/EnvQueryTest.cpp
+++ b/Engine/Source/Runtime/AIModule/Private/EnvironmentQuery/EnvQueryTest.cpp
@@ -182,16 +182,6 @@ AActor* UEnvQueryTest::GetItemActor(FEnvQueryInstance& QueryInstance, int32 Item
 		NULL;
 }
 
-FString UEnvQueryTest::GetDescriptionTitle() const
-{
-	return UEnvQueryTypes::GetShortTypeName(this).ToString();
-}
-
-FText UEnvQueryTest::GetDescriptionDetails() const
-{
-	return FText::GetEmpty();
-}
-
 void UEnvQueryTest::PostLoad()
 {
 	Super::PostLoad();
@@ -206,24 +196,14 @@ void UEnvQueryTest::PostLoad()
 		Weight.Convert(this, ScoringFactor);
 	}
 
-	UpdateTestVersion();
+	UpdateNodeVersion();
 }
 
-void UEnvQueryTest::UpdateTestVersion()
+void UEnvQueryTest::UpdateNodeVersion()
 {
 	VerNum = EnvQueryTestVersion::Latest;
 }
 
-#if WITH_EDITOR && USE_EQS_DEBUGGER
-void UEnvQueryTest::PostEditChangeProperty(FPropertyChangedEvent& PropertyChangedEvent) 
-{
-	Super::PostEditChangeProperty(PropertyChangedEvent);
-#if USE_EQS_DEBUGGER
-	UEnvQueryManager::NotifyAssetUpdate(NULL);
-#endif
-}
-#endif //WITH_EDITOR && USE_EQS_DEBUGGER
-
 FText UEnvQueryTest::DescribeFloatTestParams() const
 {
 	FText FilterDesc;
diff --git a/Engine/Source/Runtime/AIModule/Private/EnvironmentQuery/Tests/EnvQueryTest_Distance.cpp b/Engine/Source/Runtime/AIModule/Private/EnvironmentQuery/Tests/EnvQueryTest_Distance.cpp
--- a/Engine/Source/Runtime/AIModule/Private/EnvironmentQuery/Tests/EnvQueryTest_Distance.cpp
+++ b/Engine/Source/Runtime/AIModule/Private/EnvironmentQuery/Tests/EnvQueryTest_Distance.cpp
@@ -88,7 +88,7 @@ void UEnvQueryTest_Distance::RunTest(FEnvQueryInstance& QueryInstance) const
 	}
 }
 
-FString UEnvQueryTest_Distance::GetDescriptionTitle() const
+FText UEnvQueryTest_Distance::GetDescriptionTitle() const
 {
 	FString ModeDesc;
 	switch (TestMode)
@@ -109,9 +109,9 @@ FString UEnvQueryTest_Distance::GetDescriptionTitle() const
 			break;
 	}
 
-	return FString::Printf(TEXT("%s%s: to %s"), 
-		*Super::GetDescriptionTitle(), *ModeDesc,
-		*UEnvQueryTypes::DescribeContext(DistanceTo).ToString());
+	return FText::FromString(FString::Printf(TEXT("%s%s: to %s"), 
+		*Super::GetDescriptionTitle().ToString(), *ModeDesc,
+		*UEnvQueryTypes::DescribeContext(DistanceTo).ToString()));
 }
 
 FText UEnvQueryTest_Distance::GetDescriptionDetails() const
diff --git a/Engine/Source/Runtime/AIModule/Private/EnvironmentQuery/Tests/EnvQueryTest_Dot.cpp b/Engine/Source/Runtime/AIModule/Private/EnvironmentQuery/Tests/EnvQueryTest_Dot.cpp
--- a/Engine/Source/Runtime/AIModule/Private/EnvironmentQuery/Tests/EnvQueryTest_Dot.cpp
+++ b/Engine/Source/Runtime/AIModule/Private/EnvironmentQuery/Tests/EnvQueryTest_Dot.cpp
@@ -185,7 +185,7 @@ bool UEnvQueryTest_Dot::RequiresPerItemUpdates(TSubclassOf<UEnvQueryContext> Lin
 	return bRequirePerItemUpdate;
 }
 
-FString UEnvQueryTest_Dot::GetDescriptionTitle() const
+FText UEnvQueryTest_Dot::GetDescriptionTitle() const
 {
 	FString ModeDesc;
 	switch (TestMode)
@@ -202,8 +202,8 @@ FString UEnvQueryTest_Dot::GetDescriptionTitle() const
 			break;
 	}
 
-	return FString::Printf(TEXT("%s%s%s: %s and %s"), bAbsoluteValue ? TEXT("Absolute ") : TEXT(""),
-		*Super::GetDescriptionTitle(), *ModeDesc, *LineA.ToText().ToString(), *LineB.ToText().ToString());
+	return FText::FromString(FString::Printf(TEXT("%s%s%s: %s and %s"), bAbsoluteValue ? TEXT("Absolute ") : TEXT(""),
+		*Super::GetDescriptionTitle().ToString(), *ModeDesc, *LineA.ToText().ToString(), *LineB.ToText().ToString()));
 }
 
 FText UEnvQueryTest_Dot::GetDescriptionDetails() const
diff --git a/Engine/Source/Runtime/AIModule/Private/EnvironmentQuery/Tests/EnvQueryTest_GameplayTags.cpp b/Engine/Source/Runtime/AIModule/Private/EnvironmentQuery/Tests/EnvQueryTest_GameplayTags.cpp
--- a/Engine/Source/Runtime/AIModule/Private/EnvironmentQuery/Tests/EnvQueryTest_GameplayTags.cpp
+++ b/Engine/Source/Runtime/AIModule/Private/EnvironmentQuery/Tests/EnvQueryTest_GameplayTags.cpp
@@ -75,11 +75,6 @@ void UEnvQueryTest_GameplayTags::RunTest(FEnvQueryInstance& QueryInstance) const
 	}
 }
 
-FString UEnvQueryTest_GameplayTags::GetDescriptionTitle() const
-{
-	return Super::GetDescriptionTitle();
-}
-
 FText UEnvQueryTest_GameplayTags::GetDescriptionDetails() const
 {
 	return GameplayTags.ToMatchingText(TagsToMatch, !BoolValue.DefaultValue);
diff --git a/Engine/Source/Runtime/AIModule/Private/EnvironmentQuery/Tests/EnvQueryTest_Pathfinding.cpp b/Engine/Source/Runtime/AIModule/Private/EnvironmentQuery/Tests/EnvQueryTest_Pathfinding.cpp
--- a/Engine/Source/Runtime/AIModule/Private/EnvironmentQuery/Tests/EnvQueryTest_Pathfinding.cpp
+++ b/Engine/Source/Runtime/AIModule/Private/EnvironmentQuery/Tests/EnvQueryTest_Pathfinding.cpp
@@ -116,7 +116,7 @@ void UEnvQueryTest_Pathfinding::RunTest(FEnvQueryInstance& QueryInstance) const
 	}
 }
 
-FString UEnvQueryTest_Pathfinding::GetDescriptionTitle() const
+FText UEnvQueryTest_Pathfinding::GetDescriptionTitle() const
 {
 	FString ModeDesc[] = { TEXT("PathExist"), TEXT("PathCost"), TEXT("PathLength") };
 
@@ -124,7 +124,7 @@ FString UEnvQueryTest_Pathfinding::GetDescriptionTitle() const
 		FString::Printf(TEXT("%s, direction: %s"), *UEnvQueryTypes::DescribeContext(Context).ToString(), *PathFromContext.ToString()) :
 		FString::Printf(TEXT("%s %s"), PathFromContext.DefaultValue ? TEXT("from") : TEXT("to"), *UEnvQueryTypes::DescribeContext(Context).ToString());
 
-	return FString::Printf(TEXT("%s: %s"), *ModeDesc[TestMode], *DirectionDesc);
+	return FText::FromString(FString::Printf(TEXT("%s: %s"), *ModeDesc[TestMode], *DirectionDesc));
 }
 
 FText UEnvQueryTest_Pathfinding::GetDescriptionDetails() const
diff --git a/Engine/Source/Runtime/AIModule/Private/EnvironmentQuery/Tests/EnvQueryTest_Random.cpp b/Engine/Source/Runtime/AIModule/Private/EnvironmentQuery/Tests/EnvQueryTest_Random.cpp
--- a/Engine/Source/Runtime/AIModule/Private/EnvironmentQuery/Tests/EnvQueryTest_Random.cpp
+++ b/Engine/Source/Runtime/AIModule/Private/EnvironmentQuery/Tests/EnvQueryTest_Random.cpp
@@ -27,11 +27,6 @@ void UEnvQueryTest_Random::RunTest(FEnvQueryInstance& QueryInstance) const
 	}
 }
 
-FString UEnvQueryTest_Random::GetDescriptionTitle() const
-{
-	return Super::GetDescriptionTitle();
-}
-
 FText UEnvQueryTest_Random::GetDescriptionDetails() const
 {
 	return DescribeFloatTestParams();
diff --git a/Engine/Source/Runtime/AIModule/Private/EnvironmentQuery/Tests/EnvQueryTest_Trace.cpp b/Engine/Source/Runtime/AIModule/Private/EnvironmentQuery/Tests/EnvQueryTest_Trace.cpp
--- a/Engine/Source/Runtime/AIModule/Private/EnvironmentQuery/Tests/EnvQueryTest_Trace.cpp
+++ b/Engine/Source/Runtime/AIModule/Private/EnvironmentQuery/Tests/EnvQueryTest_Trace.cpp
@@ -88,7 +88,7 @@ void UEnvQueryTest_Trace::RunTest(FEnvQueryInstance& QueryInstance) const
 	}
 }
 
-FString UEnvQueryTest_Trace::GetDescriptionTitle() const
+FText UEnvQueryTest_Trace::GetDescriptionTitle() const
 {
 	UEnum* ChannelEnum = FindObject<UEnum>(ANY_PACKAGE, TEXT("ETraceTypeQuery"), true);
 	FString ChannelDesc = ChannelEnum->GetEnumText(TraceData.TraceChannel).ToString();
@@ -97,8 +97,8 @@ FString UEnvQueryTest_Trace::GetDescriptionTitle() const
 		FString::Printf(TEXT("%s, direction: %s"), *UEnvQueryTypes::DescribeContext(Context).ToString(), *TraceFromContext.ToString()) :
 		FString::Printf(TEXT("%s %s"), TraceFromContext.DefaultValue ? TEXT("from") : TEXT("to"), *UEnvQueryTypes::DescribeContext(Context).ToString());
 
-	return FString::Printf(TEXT("%s: %s on %s"), 
-		*Super::GetDescriptionTitle(), *DirectionDesc, *ChannelDesc);
+	return FText::FromString(FString::Printf(TEXT("%s: %s on %s"), 
+		*Super::GetDescriptionTitle().ToString(), *DirectionDesc, *ChannelDesc));
 }
 
 FText UEnvQueryTest_Trace::GetDescriptionDetails() const
diff --git a/Engine/Source/Runtime/Engine/Public/AI/BlueprintNodeHelpers.h b/Engine/Source/Runtime/Engine/Public/AI/BlueprintNodeHelpers.h
--- a/Engine/Source/Runtime/Engine/Public/AI/BlueprintNodeHelpers.h
+++ b/Engine/Source/Runtime/Engine/Public/AI/BlueprintNodeHelpers.h
@@ -3,6 +3,7 @@
 #pragma once
 
 class UBTNode;
+class UBehaviorTreeComponent;
 
 namespace BlueprintNodeHelpers
 {
