PR #860: Add SightStrength as out parameter to SightTargetInterface to pass as th... (Contributed by [REDACTED_PERSON])

https://github.com/EpicGames/UnrealEngine/pull/860

diff --git a/Engine/Source/Runtime/AIModule/Classes/Perception/AISightTargetInterface.h b/Engine/Source/Runtime/AIModule/Classes/Perception/AISightTargetInterface.h
--- a/Engine/Source/Runtime/AIModule/Classes/Perception/AISightTargetInterface.h
+++ b/Engine/Source/Runtime/AIModule/Classes/Perception/AISightTargetInterface.h
@@ -17,11 +17,20 @@ class AIMODULE_API IAISightTargetInterface
 	/**	Implementation should check whether from given ObserverLocation
 	 *	implementer can be seen. If so OutSeenLocation should contain
 	 *	first visible location
+	 *  Return sight strength for how well the target is seen.
 	 */
-	virtual bool CanBeSeenFrom(const FVector& ObserverLocation, FVector& OutSeenLocation, int32& NumberOfLoSChecksPerformed, const AActor* IgnoreActor = NULL) const
+	virtual bool CanBeSeenFrom(const FVector& ObserverLocation, FVector& OutSeenLocation, int32& NumberOfLoSChecksPerformed, float& OutSightStrength, const AActor* IgnoreActor = NULL) const
 	{ 
 		NumberOfLoSChecksPerformed = 0;
+		OutSightStrength = 0;
 		return false; 
 	}
+
+	DEPRECATED(4.8, "This function is deprecated. Please use the other CanBeSeenFrom version.")
+	virtual bool CanBeSeenFrom(const FVector& ObserverLocation, FVector& OutSeenLocation, int32& NumberOfLoSChecksPerformed, const AActor* IgnoreActor = NULL) const
+	{
+		float SightStrength = 1.f;
+		return CanBeSeenFrom(ObserverLocation, OutSeenLocation, NumberOfLoSChecksPerformed, SightStrength, IgnoreActor);
+	}
 };
 
diff --git a/Engine/Source/Runtime/AIModule/Private/Perception/AISense_Sight.cpp b/Engine/Source/Runtime/AIModule/Private/Perception/AISense_Sight.cpp
--- a/Engine/Source/Runtime/AIModule/Private/Perception/AISense_Sight.cpp
+++ b/Engine/Source/Runtime/AIModule/Private/Perception/AISense_Sight.cpp
@@ -5,6 +5,16 @@
 #include "Perception/AISense_Sight.h"
 #include "Perception/AISenseConfig_Sight.h"
 
+#define DO_SIGHT_VLOGGING (0 && ENABLE_VISUAL_LOG)
+
+#if DO_SIGHT_VLOGGING
+	#define SIGHT_LOG_SEGMENT UE_VLOG_SEGMENT
+	#define SIGHT_LOG_LOCATION UE_VLOG_LOCATION
+#else
+	#define SIGHT_LOG_SEGMENT(...)
+	#define SIGHT_LOG_LOCATION(...)
+#endif // DO_SIGHT_VLOGGING
+
 DECLARE_CYCLE_STAT(TEXT("Perception Sense: Sight"),STAT_AI_Sense_Sight,STATGROUP_AI);
 DECLARE_CYCLE_STAT(TEXT("Perception Sense: Sight, Listener Update"), STAT_AI_Sense_Sight_ListenerUpdate, STATGROUP_AI);
 
@@ -135,21 +145,23 @@ float UAISense_Sight::Update()
 
 				if (CheckIsTargetInSightPie(Listener, PropDigest, TargetLocation, SightRadiusSq))
 				{
-//					UE_VLOG_SEGMENT(Listener.Listener.Get()->GetOwner(), Listener.CachedLocation, TargetLocation, FColor::Green, TEXT("%s"), *(Target.TargetId.ToString()));
+					SIGHT_LOG_SEGMENT(Listener.Listener.Get()->GetOwner(), Listener.CachedLocation, TargetLocation, FColor::Green, TEXT("%s"), *(Target.TargetId.ToString()));
 
 					FVector OutSeenLocation(0.f);
 					// do line checks
 					if (Target.SightTargetInterface != NULL)
 					{
 						int32 NumberOfLoSChecksPerformed = 0;
-						if (Target.SightTargetInterface->CanBeSeenFrom(Listener.CachedLocation, OutSeenLocation, NumberOfLoSChecksPerformed, Listener.Listener->GetBodyActor()) == true)
+						// defaulting to 1 to have "full strength" by default instead of "no strength"
+						float SightStrength = 1.f;
+						if (Target.SightTargetInterface->CanBeSeenFrom(Listener.CachedLocation, OutSeenLocation, NumberOfLoSChecksPerformed, SightStrength, Listener.Listener->GetBodyActor()) == true)
 						{
-							Listener.RegisterStimulus(TargetActor, FAIStimulus(*this, 1.f, OutSeenLocation, Listener.CachedLocation));
+							Listener.RegisterStimulus(TargetActor, FAIStimulus(*this, SightStrength, OutSeenLocation, Listener.CachedLocation));
 							SightQuery->bLastResult = true;
 						}
 						else
 						{
-//							UE_VLOG_LOCATION(Listener.Listener.Get()->GetOwner(), TargetLocation, 25.f, FColor::Red, TEXT(""));
+							SIGHT_LOG_LOCATION(Listener.Listener.Get()->GetOwner(), TargetLocation, 25.f, FColor::Red, TEXT(""));
 							Listener.RegisterStimulus(TargetActor, FAIStimulus(*this, 0.f, TargetLocation, Listener.CachedLocation, FAIStimulus::SensingFailed));
 							SightQuery->bLastResult = false;
 						}
@@ -159,9 +171,6 @@ float UAISense_Sight::Update()
 					else
 					{
 						// we need to do tests ourselves
-						/*const bool bHit = World->LineTraceTest(Listener.CachedLocation, TargetLocation
-							, FCollisionQueryParams(NAME_AILineOfSight, true, Listener.Listener->GetBodyActor())
-							, FCollisionObjectQueryParams(ECC_WorldStatic));*/
 						FHitResult HitResult;
 						const bool bHit = World->LineTraceSingleByObjectType(HitResult, Listener.CachedLocation, TargetLocation
 							, FCollisionObjectQueryParams(ECC_WorldStatic)
@@ -176,7 +185,7 @@ float UAISense_Sight::Update()
 						}
 						else
 						{
-//							UE_VLOG_LOCATION(Listener.Listener.Get()->GetOwner(), TargetLocation, 25.f, FColor::Red, TEXT(""));
+							SIGHT_LOG_LOCATION(Listener.Listener.Get()->GetOwner(), TargetLocation, 25.f, FColor::Red, TEXT(""));
 							Listener.RegisterStimulus(TargetActor, FAIStimulus(*this, 0.f, TargetLocation, Listener.CachedLocation, FAIStimulus::SensingFailed));
 							SightQuery->bLastResult = false;
 						}
@@ -184,7 +193,7 @@ float UAISense_Sight::Update()
 				}
 				else
 				{
-//					UE_VLOG_SEGMENT(Listener.Listener.Get()->GetOwner(), Listener.CachedLocation, TargetLocation, FColor::Red, TEXT("%s"), *(Target.TargetId.ToString()));
+					SIGHT_LOG_SEGMENT(Listener.Listener.Get()->GetOwner(), Listener.CachedLocation, TargetLocation, FColor::Red, TEXT("%s"), *(Target.TargetId.ToString()));
 					Listener.RegisterStimulus(TargetActor, FAIStimulus(*this, 0.f, TargetLocation, Listener.CachedLocation, FAIStimulus::SensingFailed));
 					SightQuery->bLastResult = false;
 				}
