Fixed BlackboardComponents not being registered for BB key synchronization #UE4

diff --git a/Engine/Source/Runtime/AIModule/Private/BehaviorTree/BlackboardComponent.cpp b/Engine/Source/Runtime/AIModule/Private/BehaviorTree/BlackboardComponent.cpp
--- a/Engine/Source/Runtime/AIModule/Private/BehaviorTree/BlackboardComponent.cpp
+++ b/Engine/Source/Runtime/AIModule/Private/BehaviorTree/BlackboardComponent.cpp
@@ -201,6 +201,7 @@ void UBlackboardComponent::PopulateSynchronizedKeys()
 	
 	UAISystem* AISystem = UAISystem::GetCurrentSafe(GetWorld());
 	check(AISystem);
+	AISystem->RegisterBlackboardComponent(*BlackboardAsset, *this);
 
 	for (auto Iter = AISystem->CreateBlackboardDataToComponentsIterator(*BlackboardAsset); Iter; ++Iter)
 	{
