Fixed Setting and Getting values for BB's NativeEnum keys not being compatible with Enum typed keys #UE4

diff --git a/Engine/Source/Runtime/AIModule/Classes/BehaviorTree/BlackboardComponent.h b/Engine/Source/Runtime/AIModule/Classes/BehaviorTree/BlackboardComponent.h
--- a/Engine/Source/Runtime/AIModule/Classes/BehaviorTree/BlackboardComponent.h
+++ b/Engine/Source/Runtime/AIModule/Classes/BehaviorTree/BlackboardComponent.h
@@ -14,6 +14,8 @@
 #include "Components/ActorComponent.h"
 #include "BehaviorTreeTypes.h"
 #include "BehaviorTree/Blackboard/BlackboardKeyType.h"
+#include "BehaviorTree/Blackboard/BlackboardKeyType_Enum.h"
+#include "BehaviorTree/Blackboard/BlackboardKeyType_NativeEnum.h"
 #include "BlackboardData.h"
 #include "AISystem.h"
 #include "BlackboardComponent.generated.h"
@@ -189,6 +191,9 @@ public:
 	template<class TDataClass>
 	typename TDataClass::FDataType GetValue(FBlackboard::FKey KeyID) const;
 
+	template<class TDataClass>
+	bool IsCompatibleKeyType(const UBlackboardKeyType& KeyType) const;
+
 	/** get pointer to raw data for given key */
 	FORCEINLINE uint8* GetKeyRawData(const FName& KeyName) { return GetKeyRawData(GetKeyID(KeyName)); }
 	FORCEINLINE uint8* GetKeyRawData(FBlackboard::FKey KeyID) { return ValueMemory.Num() && ValueOffsets.IsValidIndex(KeyID) ? (ValueMemory.GetData() + ValueOffsets[KeyID]) : NULL; }
@@ -321,6 +326,18 @@ FORCEINLINE bool UBlackboardComponent::HasValidAsset() const
 	return BlackboardAsset && BlackboardAsset->IsValid();
 }
 
+template<class TDataClass>
+FORCEINLINE bool UBlackboardComponent::IsCompatibleKeyType(const UBlackboardKeyType& KeyType) const
+{
+	return KeyType.IsA(TDataClass::StaticClass());
+}
+
+template<>
+FORCEINLINE bool UBlackboardComponent::IsCompatibleKeyType<UBlackboardKeyType_Enum>(const UBlackboardKeyType& KeyType) const
+{
+	return KeyType.GetClass() == UBlackboardKeyType_NativeEnum::StaticClass() || KeyType.IsA(UBlackboardKeyType_Enum::StaticClass());
+}
+
 template<class TDataClass>
 bool UBlackboardComponent::SetValue(const FName& KeyName, typename TDataClass::FDataType Value)
 {
@@ -332,7 +349,7 @@ template<class TDataClass>
 bool UBlackboardComponent::SetValue(FBlackboard::FKey KeyID, typename TDataClass::FDataType Value)
 {
 	const FBlackboardEntry* EntryInfo = BlackboardAsset ? BlackboardAsset->GetKey(KeyID) : nullptr;
-	if ((EntryInfo == nullptr) || (EntryInfo->KeyType == nullptr) || !EntryInfo->KeyType->IsA(TDataClass::StaticClass()))
+	if ((EntryInfo == nullptr) || (EntryInfo->KeyType == nullptr) || (IsCompatibleKeyType<TDataClass>(*EntryInfo->KeyType) == false))
 	{
 		return false;
 	}
@@ -381,7 +398,7 @@ template<class TDataClass>
 typename TDataClass::FDataType UBlackboardComponent::GetValue(FBlackboard::FKey KeyID) const
 {
 	const FBlackboardEntry* EntryInfo = BlackboardAsset ? BlackboardAsset->GetKey(KeyID) : nullptr;
-	if ((EntryInfo == nullptr) || (EntryInfo->KeyType == nullptr) || !EntryInfo->KeyType->IsA(TDataClass::StaticClass()))
+	if ((EntryInfo == nullptr) || (EntryInfo->KeyType == nullptr) || (IsCompatibleKeyType<TDataClass>(*EntryInfo->KeyType) == false))
 	{
 		return TDataClass::InvalidValue;
 	}
