Made BTTask_PlayAnimation work with any pawn owning a SkeletalMeshComponent, rather than just with Characters #UE4

diff --git a/Engine/Source/Runtime/AIModule/Private/BehaviorTree/Tasks/BTTask_PlayAnimation.cpp b/Engine/Source/Runtime/AIModule/Private/BehaviorTree/Tasks/BTTask_PlayAnimation.cpp
--- a/Engine/Source/Runtime/AIModule/Private/BehaviorTree/Tasks/BTTask_PlayAnimation.cpp
+++ b/Engine/Source/Runtime/AIModule/Private/BehaviorTree/Tasks/BTTask_PlayAnimation.cpp
@@ -31,12 +31,22 @@ EBTNodeResult::Type UBTTask_PlayAnimation::ExecuteTask(UBehaviorTreeComponent& O
 	TimerHandle.Invalidate();
 	MyOwnerComp = &OwnerComp;
 
-	if (AnimationToPlay && MyController)
+	if (AnimationToPlay && MyController && MyController->GetPawn())
 	{
+		USkeletalMeshComponent* SkelMesh = nullptr;
 		ACharacter* const MyCharacter = Cast<ACharacter>(MyController->GetPawn());
-		if (MyCharacter && MyCharacter->GetMesh())
+		if (MyCharacter)
 		{
-			MyCharacter->GetMesh()->PlayAnimation(AnimationToPlay, bLooping);
+			SkelMesh = MyCharacter->GetMesh();
+		}
+		else
+		{
+			SkelMesh = MyController->GetPawn()->FindComponentByClass<USkeletalMeshComponent>();
+		}
+
+		if (SkelMesh != nullptr)
+		{
+			SkelMesh->PlayAnimation(AnimationToPlay, bLooping);
 			const float FinishDelay = AnimationToPlay->GetMaxCurrentTime();
 
 			if (bNonBlocking == false && FinishDelay > 0)
