merged blackboard enum keys
#ue4

diff --git a/Engine/Source/Editor/BehaviorTreeEditor/Private/BehaviorTreeEditor.cpp b/Engine/Source/Editor/BehaviorTreeEditor/Private/BehaviorTreeEditor.cpp
--- a/Engine/Source/Editor/BehaviorTreeEditor/Private/BehaviorTreeEditor.cpp
+++ b/Engine/Source/Editor/BehaviorTreeEditor/Private/BehaviorTreeEditor.cpp
@@ -668,7 +668,13 @@ TSharedRef<SWidget> FBehaviorTreeEditor::SpawnBlackboardDetails()
 	FOnGetDetailCustomizationInstance LayoutVariableDetails = FOnGetDetailCustomizationInstance::CreateStatic(&FBlackboardDataDetails::MakeInstance, OnGetSelectedBlackboardItemIndex);
 	BlackboardDetailsView->RegisterInstancedCustomPropertyLayout(UBlackboardData::StaticClass(), LayoutVariableDetails);
 
-	BlackboardDetailsView->SetObject(GetBlackboardData());
+	UBlackboardData* BBData = GetBlackboardData();
+	if (BBData)
+	{
+		BBData->UpdateDeprecatedKeys();
+	}
+
+	BlackboardDetailsView->SetObject(BBData);
 	BlackboardDetailsView->SetEnabled(TAttribute<bool>::Create(TAttribute<bool>::FGetter::CreateSP(this, &FBehaviorTreeEditor::CanEditWithDebuggerActive)));
 
 	return BlackboardDetailsView.ToSharedRef();
diff --git a/Engine/Source/Editor/BehaviorTreeEditor/Private/SBehaviorTreeBlackboardEditor.cpp b/Engine/Source/Editor/BehaviorTreeEditor/Private/SBehaviorTreeBlackboardEditor.cpp
--- a/Engine/Source/Editor/BehaviorTreeEditor/Private/SBehaviorTreeBlackboardEditor.cpp
+++ b/Engine/Source/Editor/BehaviorTreeEditor/Private/SBehaviorTreeBlackboardEditor.cpp
@@ -173,7 +173,9 @@ public:
 	{
 		if(InClass != nullptr)
 		{
-			return !InClass->HasAnyClassFlags(CLASS_Abstract) && InClass->IsChildOf(UBlackboardKeyType::StaticClass());
+			return !InClass->HasAnyClassFlags(CLASS_Abstract) &&
+				InClass->HasAnyClassFlags(CLASS_EditInlineNew) &&
+				InClass->IsChildOf(UBlackboardKeyType::StaticClass());
 		}
 		return false;
 	}
diff --git a/Engine/Source/Runtime/AIModule/Classes/BehaviorTree/Blackboard/BlackboardKeyType.h b/Engine/Source/Runtime/AIModule/Classes/BehaviorTree/Blackboard/BlackboardKeyType.h
--- a/Engine/Source/Runtime/AIModule/Classes/BehaviorTree/Blackboard/BlackboardKeyType.h
+++ b/Engine/Source/Runtime/AIModule/Classes/BehaviorTree/Blackboard/BlackboardKeyType.h
@@ -117,6 +117,9 @@ class AIMODULE_API UBlackboardKeyType : public UObject
 	/** description of params for property view */
 	virtual FString DescribeSelf() const;
 
+	/** create replacement key for deprecated data */
+	virtual UBlackboardKeyType* UpdateDeprecatedKey();
+
 	/** @return key instance if bCreateKeyInstance was set */
 	const UBlackboardKeyType* GetKeyInstance(const UBlackboardComponent& OwnerComp, const uint8* MemoryBlock) const;
 	UBlackboardKeyType* GetKeyInstance(UBlackboardComponent& OwnerComp, const uint8* MemoryBlock) const;
diff --git a/Engine/Source/Runtime/AIModule/Classes/BehaviorTree/Blackboard/BlackboardKeyType_Enum.h b/Engine/Source/Runtime/AIModule/Classes/BehaviorTree/Blackboard/BlackboardKeyType_Enum.h
--- a/Engine/Source/Runtime/AIModule/Classes/BehaviorTree/Blackboard/BlackboardKeyType_Enum.h
+++ b/Engine/Source/Runtime/AIModule/Classes/BehaviorTree/Blackboard/BlackboardKeyType_Enum.h
@@ -14,6 +14,14 @@ class AIMODULE_API UBlackboardKeyType_Enum : public UBlackboardKeyType
 	UPROPERTY(Category=Blackboard, EditDefaultsOnly)
 	UEnum* EnumType;
 
+	/** name of enum defined in c++ code, will take priority over asset from EnumType property */
+	UPROPERTY(Category=Blackboard, EditDefaultsOnly)
+	FString EnumName;
+
+	/** set when EnumName override is valid and active */
+	UPROPERTY(Category = Blackboard, VisibleDefaultsOnly)
+	uint32 bIsEnumNameValid : 1;
+
 	static uint8 GetValue(const UBlackboardKeyType_Enum* KeyOb, const uint8* RawData);
 	static bool SetValue(UBlackboardKeyType_Enum* KeyOb, uint8* RawData, uint8 Value);
 
@@ -24,6 +32,10 @@ class AIMODULE_API UBlackboardKeyType_Enum : public UBlackboardKeyType
 	virtual FString DescribeArithmeticParam(int32 IntValue, float FloatValue) const override;
 	virtual bool IsAllowedByFilter(UBlackboardKeyType* FilterOb) const override;
 
+#if WITH_EDITOR
+	virtual void PostEditChangeProperty(struct FPropertyChangedEvent& PropertyChangedEvent) override;
+#endif
+
 protected:
 	virtual FString DescribeValue(const UBlackboardComponent& OwnerComp, const uint8* RawData) const override;
 	virtual bool TestArithmeticOperation(const UBlackboardComponent& OwnerComp, const uint8* MemoryBlock, EArithmeticKeyOperation::Type Op, int32 OtherIntValue, float OtherFloatValue) const override;
diff --git a/Engine/Source/Runtime/AIModule/Classes/BehaviorTree/Blackboard/BlackboardKeyType_NativeEnum.h b/Engine/Source/Runtime/AIModule/Classes/BehaviorTree/Blackboard/BlackboardKeyType_NativeEnum.h
--- a/Engine/Source/Runtime/AIModule/Classes/BehaviorTree/Blackboard/BlackboardKeyType_NativeEnum.h
+++ b/Engine/Source/Runtime/AIModule/Classes/BehaviorTree/Blackboard/BlackboardKeyType_NativeEnum.h
@@ -4,7 +4,9 @@
 #include "BehaviorTree/Blackboard/BlackboardKeyType.h"
 #include "BlackboardKeyType_NativeEnum.generated.h"
 
-UCLASS(EditInlineNew, meta=(DisplayName="Native Enum"))
+// DEPRECATED, please use UBlackboardKeyType_Enum instead
+
+UCLASS(NotEditInlineNew)
 class AIMODULE_API UBlackboardKeyType_NativeEnum : public UBlackboardKeyType
 {
 	GENERATED_UCLASS_BODY()
@@ -15,27 +17,11 @@ class AIMODULE_API UBlackboardKeyType_NativeEnum : public UBlackboardKeyType
 	UPROPERTY(Category=Blackboard, EditDefaultsOnly)
 	FString EnumName;
 
-	UPROPERTY(Category=Blackboard, VisibleDefaultsOnly)
-	bool bIsEnumNameValid;
-
 	UPROPERTY()
 	UEnum* EnumType;
 
+	virtual UBlackboardKeyType* UpdateDeprecatedKey() override;
+
 	static uint8 GetValue(const UBlackboardKeyType_NativeEnum* KeyOb, const uint8* RawData);
 	static bool SetValue(UBlackboardKeyType_NativeEnum* KeyOb, uint8* RawData, uint8 Value);
-
-	virtual EBlackboardCompare::Type CompareValues(const UBlackboardComponent& OwnerComp, const uint8* MemoryBlock,
-		const UBlackboardKeyType* OtherKeyOb, const uint8* OtherMemoryBlock) const override;
-
-	virtual bool IsAllowedByFilter(UBlackboardKeyType* FilterOb) const override;
-	virtual FString DescribeSelf() const override;
-	virtual FString DescribeArithmeticParam(int32 IntValue, float FloatValue) const override;
-
-#if WITH_EDITOR
-	virtual void PostEditChangeProperty(struct FPropertyChangedEvent& PropertyChangedEvent) override;
-#endif
-
-protected:
-	virtual FString DescribeValue(const UBlackboardComponent& OwnerComp, const uint8* RawData) const override;
-	virtual bool TestArithmeticOperation(const UBlackboardComponent& OwnerComp, const uint8* MemoryBlock, EArithmeticKeyOperation::Type Op, int32 OtherIntValue, float OtherFloatValue) const override;
 };
diff --git a/Engine/Source/Runtime/AIModule/Classes/BehaviorTree/BlackboardComponent.h b/Engine/Source/Runtime/AIModule/Classes/BehaviorTree/BlackboardComponent.h
--- a/Engine/Source/Runtime/AIModule/Classes/BehaviorTree/BlackboardComponent.h
+++ b/Engine/Source/Runtime/AIModule/Classes/BehaviorTree/BlackboardComponent.h
@@ -191,9 +191,6 @@ public:
 	template<class TDataClass>
 	typename TDataClass::FDataType GetValue(FBlackboard::FKey KeyID) const;
 
-	template<class TDataClass>
-	bool IsCompatibleKeyType(const UBlackboardKeyType& KeyType) const;
-
 	/** get pointer to raw data for given key */
 	FORCEINLINE uint8* GetKeyRawData(const FName& KeyName) { return GetKeyRawData(GetKeyID(KeyName)); }
 	FORCEINLINE uint8* GetKeyRawData(FBlackboard::FKey KeyID) { return ValueMemory.Num() && ValueOffsets.IsValidIndex(KeyID) ? (ValueMemory.GetData() + ValueOffsets[KeyID]) : NULL; }
@@ -326,18 +323,6 @@ FORCEINLINE bool UBlackboardComponent::HasValidAsset() const
 	return BlackboardAsset && BlackboardAsset->IsValid();
 }
 
-template<class TDataClass>
-FORCEINLINE bool UBlackboardComponent::IsCompatibleKeyType(const UBlackboardKeyType& KeyType) const
-{
-	return KeyType.IsA(TDataClass::StaticClass());
-}
-
-template<>
-FORCEINLINE bool UBlackboardComponent::IsCompatibleKeyType<UBlackboardKeyType_Enum>(const UBlackboardKeyType& KeyType) const
-{
-	return KeyType.GetClass() == UBlackboardKeyType_NativeEnum::StaticClass() || KeyType.IsA(UBlackboardKeyType_Enum::StaticClass());
-}
-
 template<class TDataClass>
 bool UBlackboardComponent::SetValue(const FName& KeyName, typename TDataClass::FDataType Value)
 {
@@ -349,7 +334,7 @@ template<class TDataClass>
 bool UBlackboardComponent::SetValue(FBlackboard::FKey KeyID, typename TDataClass::FDataType Value)
 {
 	const FBlackboardEntry* EntryInfo = BlackboardAsset ? BlackboardAsset->GetKey(KeyID) : nullptr;
-	if ((EntryInfo == nullptr) || (EntryInfo->KeyType == nullptr) || (IsCompatibleKeyType<TDataClass>(*EntryInfo->KeyType) == false))
+	if ((EntryInfo == nullptr) || (EntryInfo->KeyType == nullptr))
 	{
 		return false;
 	}
@@ -398,7 +383,7 @@ template<class TDataClass>
 typename TDataClass::FDataType UBlackboardComponent::GetValue(FBlackboard::FKey KeyID) const
 {
 	const FBlackboardEntry* EntryInfo = BlackboardAsset ? BlackboardAsset->GetKey(KeyID) : nullptr;
-	if ((EntryInfo == nullptr) || (EntryInfo->KeyType == nullptr) || (IsCompatibleKeyType<TDataClass>(*EntryInfo->KeyType) == false))
+	if ((EntryInfo == nullptr) || (EntryInfo->KeyType == nullptr))
 	{
 		return TDataClass::InvalidValue;
 	}
diff --git a/Engine/Source/Runtime/AIModule/Classes/BehaviorTree/BlackboardData.h b/Engine/Source/Runtime/AIModule/Classes/BehaviorTree/BlackboardData.h
--- a/Engine/Source/Runtime/AIModule/Classes/BehaviorTree/BlackboardData.h
+++ b/Engine/Source/Runtime/AIModule/Classes/BehaviorTree/BlackboardData.h
@@ -131,6 +131,9 @@ public:
 
 	void UpdateIfHasSynchronizedKeys();
 
+	/** fix entries with deprecated key types */
+	void UpdateDeprecatedKeys();
+
 protected:
 
 	enum EKeyLookupMode
diff --git a/Engine/Source/Runtime/AIModule/Private/BehaviorTree/Blackboard/BlackboardKeyType.cpp b/Engine/Source/Runtime/AIModule/Private/BehaviorTree/Blackboard/BlackboardKeyType.cpp
--- a/Engine/Source/Runtime/AIModule/Private/BehaviorTree/Blackboard/BlackboardKeyType.cpp
+++ b/Engine/Source/Runtime/AIModule/Private/BehaviorTree/Blackboard/BlackboardKeyType.cpp
@@ -297,3 +297,8 @@ FString UBlackboardKeyType::DescribeValue(const uint8* MemoryBlock) const
 {
 	return FString();
 }
+
+UBlackboardKeyType* UBlackboardKeyType::UpdateDeprecatedKey()
+{
+	return nullptr;
+}
diff --git a/Engine/Source/Runtime/AIModule/Private/BehaviorTree/Blackboard/BlackboardKeyType_Enum.cpp b/Engine/Source/Runtime/AIModule/Private/BehaviorTree/Blackboard/BlackboardKeyType_Enum.cpp
--- a/Engine/Source/Runtime/AIModule/Private/BehaviorTree/Blackboard/BlackboardKeyType_Enum.cpp
+++ b/Engine/Source/Runtime/AIModule/Private/BehaviorTree/Blackboard/BlackboardKeyType_Enum.cpp
@@ -69,3 +69,18 @@ FString UBlackboardKeyType_Enum::DescribeArithmeticParam(int32 IntValue, float F
 {
 	return EnumType ? EnumType->GetEnumText(IntValue).ToString() : FString("UNKNOWN!");
 }
+
+#if WITH_EDITOR
+void UBlackboardKeyType_Enum::PostEditChangeProperty(struct FPropertyChangedEvent& PropertyChangedEvent)
+{
+	Super::PostEditChangeProperty(PropertyChangedEvent);
+
+	if (PropertyChangedEvent.Property &&
+		PropertyChangedEvent.Property->GetFName() == GET_MEMBER_NAME_CHECKED(UBlackboardKeyType_Enum, EnumName))
+	{
+		EnumType = FindObject<UEnum>(ANY_PACKAGE, *EnumName, true);
+	}
+
+	bIsEnumNameValid = EnumType && !EnumName.IsEmpty();
+}
+#endif
diff --git a/Engine/Source/Runtime/AIModule/Private/BehaviorTree/Blackboard/BlackboardKeyType_NativeEnum.cpp b/Engine/Source/Runtime/AIModule/Private/BehaviorTree/Blackboard/BlackboardKeyType_NativeEnum.cpp
--- a/Engine/Source/Runtime/AIModule/Private/BehaviorTree/Blackboard/BlackboardKeyType_NativeEnum.cpp
+++ b/Engine/Source/Runtime/AIModule/Private/BehaviorTree/Blackboard/BlackboardKeyType_NativeEnum.cpp
@@ -1,85 +1,31 @@
 // Copyright 1998-2015 Epic Games, Inc. All Rights Reserved.
 
 #include "AIModulePrivate.h"
+#include "BehaviorTree/Blackboard/BlackboardKeyType.h"
 #include "BehaviorTree/Blackboard/BlackboardKeyType_NativeEnum.h"
 
 const UBlackboardKeyType_NativeEnum::FDataType UBlackboardKeyType_NativeEnum::InvalidValue = UBlackboardKeyType_NativeEnum::FDataType(0);
 
 UBlackboardKeyType_NativeEnum::UBlackboardKeyType_NativeEnum(const FObjectInitializer& ObjectInitializer) : Super(ObjectInitializer)
 {
-	ValueSize = sizeof(uint8);
-	SupportedOp = EBlackboardKeyOperation::Arithmetic;
 }
 
-uint8 UBlackboardKeyType_NativeEnum::GetValue(const UBlackboardKeyType_NativeEnum* KeyOb, const uint8* RawData)
-{
-	return GetValueFromMemory<uint8>(RawData);
-}
-
-bool UBlackboardKeyType_NativeEnum::SetValue(UBlackboardKeyType_NativeEnum* KeyOb, uint8* RawData, uint8 Value)
-{
-	return SetValueInMemory<uint8>(RawData, Value);
-}
-
-EBlackboardCompare::Type UBlackboardKeyType_NativeEnum::CompareValues(const UBlackboardComponent& OwnerComp, const uint8* MemoryBlock,
-	const UBlackboardKeyType* OtherKeyOb, const uint8* OtherMemoryBlock) const
-{
-	const uint8 MyValue = GetValue(this, MemoryBlock);
-	const uint8 OtherValue = GetValue((UBlackboardKeyType_NativeEnum*)OtherKeyOb, OtherMemoryBlock);
-
-	return (MyValue > OtherValue) ? EBlackboardCompare::Greater :
-		(MyValue < OtherValue) ? EBlackboardCompare::Less :
-		EBlackboardCompare::Equal;
-}
-
-FString UBlackboardKeyType_NativeEnum::DescribeValue(const UBlackboardComponent& OwnerComp, const uint8* RawData) const
+UBlackboardKeyType* UBlackboardKeyType_NativeEnum::UpdateDeprecatedKey()
 {
-	return EnumType ? EnumType->GetEnumName(GetValue(this, RawData)) : FString("UNKNOWN!");
-}
-
-FString UBlackboardKeyType_NativeEnum::DescribeSelf() const
-{
-	return *GetNameSafe(EnumType);
-}
-
-bool UBlackboardKeyType_NativeEnum::IsAllowedByFilter(UBlackboardKeyType* FilterOb) const
-{
-	UBlackboardKeyType_NativeEnum* FilterEnum = Cast<UBlackboardKeyType_NativeEnum>(FilterOb);
-	return (FilterEnum && FilterEnum->EnumName == EnumName);
-}
+	UBlackboardKeyType_Enum* KeyOb = NewObject<UBlackboardKeyType_Enum>(GetOuter());
+	KeyOb->EnumName = EnumName;
+	KeyOb->EnumType = EnumType;
+	KeyOb->bIsEnumNameValid = EnumType && !EnumName.IsEmpty();
 
-#if WITH_EDITOR
-void UBlackboardKeyType_NativeEnum::PostEditChangeProperty(FPropertyChangedEvent& PropertyChangedEvent)
-{
-	Super::PostEditChangeProperty(PropertyChangedEvent);
-
-	if (PropertyChangedEvent.Property &&
-		PropertyChangedEvent.Property->GetFName() == GET_MEMBER_NAME_CHECKED(UBlackboardKeyType_NativeEnum, EnumName))
-	{
-		EnumType = FindObject<UEnum>(ANY_PACKAGE, *EnumName, true);
-		bIsEnumNameValid =  (EnumType != NULL);
-	}
+	return KeyOb;
 }
-#endif
 
-bool UBlackboardKeyType_NativeEnum::TestArithmeticOperation(const UBlackboardComponent& OwnerComp, const uint8* MemoryBlock, EArithmeticKeyOperation::Type Op, int32 OtherIntValue, float OtherFloatValue) const
+uint8 UBlackboardKeyType_NativeEnum::GetValue(const UBlackboardKeyType_NativeEnum* KeyOb, const uint8* RawData)
 {
-	const uint8 Value = GetValue(this, MemoryBlock);
-	switch (Op)
-	{
-	case EArithmeticKeyOperation::Equal:			return (Value == OtherIntValue);
-	case EArithmeticKeyOperation::NotEqual:			return (Value != OtherIntValue);
-	case EArithmeticKeyOperation::Less:				return (Value < OtherIntValue);
-	case EArithmeticKeyOperation::LessOrEqual:		return (Value <= OtherIntValue);
-	case EArithmeticKeyOperation::Greater:			return (Value > OtherIntValue);
-	case EArithmeticKeyOperation::GreaterOrEqual:	return (Value >= OtherIntValue);
-	default: break;
-	}
-
-	return false;
+	return GetValueFromMemory<uint8>(RawData);
 }
 
-FString UBlackboardKeyType_NativeEnum::DescribeArithmeticParam(int32 IntValue, float FloatValue) const
+bool UBlackboardKeyType_NativeEnum::SetValue(UBlackboardKeyType_NativeEnum* KeyOb, uint8* RawData, uint8 Value)
 {
-	return EnumType ? EnumType->GetEnumName(IntValue) : FString("UNKNOWN!");
+	return SetValueInMemory<uint8>(RawData, Value);
 }
diff --git a/Engine/Source/Runtime/AIModule/Private/BehaviorTree/BlackboardComponent.cpp b/Engine/Source/Runtime/AIModule/Private/BehaviorTree/BlackboardComponent.cpp
--- a/Engine/Source/Runtime/AIModule/Private/BehaviorTree/BlackboardComponent.cpp
+++ b/Engine/Source/Runtime/AIModule/Private/BehaviorTree/BlackboardComponent.cpp
@@ -83,6 +83,7 @@ void UBlackboardComponent::InitializeParentChain(UBlackboardData* NewAsset)
 	if (NewAsset)
 	{
 		InitializeParentChain(NewAsset->Parent);
+		NewAsset->UpdateDeprecatedKeys();
 		NewAsset->UpdateKeyIDs();
 	}
 }
diff --git a/Engine/Source/Runtime/AIModule/Private/BehaviorTree/BlackboardData.cpp b/Engine/Source/Runtime/AIModule/Private/BehaviorTree/BlackboardData.cpp
--- a/Engine/Source/Runtime/AIModule/Private/BehaviorTree/BlackboardData.cpp
+++ b/Engine/Source/Runtime/AIModule/Private/BehaviorTree/BlackboardData.cpp
@@ -224,3 +224,19 @@ void UBlackboardData::UpdateKeyIDs()
 {
 	FirstKeyID = Parent ? Parent->GetNumKeys() : 0;
 }
+
+void UBlackboardData::UpdateDeprecatedKeys()
+{
+	for (int32 KeyIndex = 0; KeyIndex < Keys.Num(); KeyIndex++)
+	{
+		FBlackboardEntry& Entry = Keys[KeyIndex];
+		if (Entry.KeyType)
+		{
+			UBlackboardKeyType* UpdatedKey = Entry.KeyType->UpdateDeprecatedKey();
+			if (UpdatedKey)
+			{
+				Entry.KeyType = UpdatedKey;
+			}
+		}
+	}
+}
