Fixed PathFollowingComponent figuring out path's starting point wrong #UE4

diff --git a/Engine/Source/Runtime/AIModule/Private/Navigation/PathFollowingComponent.cpp b/Engine/Source/Runtime/AIModule/Private/Navigation/PathFollowingComponent.cpp
--- a/Engine/Source/Runtime/AIModule/Private/Navigation/PathFollowingComponent.cpp
+++ b/Engine/Source/Runtime/AIModule/Private/Navigation/PathFollowingComponent.cpp
@@ -544,8 +544,10 @@ int32 UPathFollowingComponent::DetermineStartingPathPoint(const FNavigationPath*
 			{
 				// check if is closer to first or second path point (don't assume AI's standing)
 				const FVector CurrentLocation = MovementComp->GetActorFeetLocation();
-				const float SqDistToFirstPoint = FVector::DistSquared(CurrentLocation, ConsideredPath->GetPathPoints()[0].Location);
-				const float SqDistToSecondPoint = FVector::DistSquared(CurrentLocation, ConsideredPath->GetPathPoints()[1].Location);
+				// making this test in 2d to avoid situation where agent's Z location not being in "navmesh plane"
+				// would influence the result
+				const float SqDistToFirstPoint = (CurrentLocation - ConsideredPath->GetPathPoints()[0].Location).SizeSquared2D();
+				const float SqDistToSecondPoint = (CurrentLocation - ConsideredPath->GetPathPoints()[1].Location).SizeSquared2D();
 				PickedPathPoint = (SqDistToFirstPoint < SqDistToSecondPoint) ? 0 : 1;
 			}
 			else
