Marked AISenseConfig_Blueprint as abstract so that it doesn't clutter senses dropdowns #UE4

diff --git a/Engine/Source/Runtime/AIModule/Classes/Perception/AISenseConfig_Blueprint.h b/Engine/Source/Runtime/AIModule/Classes/Perception/AISenseConfig_Blueprint.h
--- a/Engine/Source/Runtime/AIModule/Classes/Perception/AISenseConfig_Blueprint.h
+++ b/Engine/Source/Runtime/AIModule/Classes/Perception/AISenseConfig_Blueprint.h
@@ -7,7 +7,7 @@
 
 class UAISense_Blueprint;
 
-UCLASS(Blueprintable)
+UCLASS(Blueprintable, Abstract)
 class AIMODULE_API UAISenseConfig_Blueprint : public UAISenseConfig
 {
 	GENERATED_BODY()
