Made BT move tasks' bStopOnOverlap configurable via ini files #UE4

diff --git a/Engine/Source/Runtime/AIModule/Classes/BehaviorTree/Tasks/BTTask_MoveDirectlyToward.h b/Engine/Source/Runtime/AIModule/Classes/BehaviorTree/Tasks/BTTask_MoveDirectlyToward.h
--- a/Engine/Source/Runtime/AIModule/Classes/BehaviorTree/Tasks/BTTask_MoveDirectlyToward.h
+++ b/Engine/Source/Runtime/AIModule/Classes/BehaviorTree/Tasks/BTTask_MoveDirectlyToward.h
@@ -36,7 +36,7 @@ class AIMODULE_API UBTTask_MoveDirectlyToward : public UBTTask_BlackboardBase
 
 	/** if set to true agent's radius will be added to AcceptableRadius for purposes of checking
 	 *	if path's end point has been reached. Will result in AI stopping on contact with destination location*/
-	UPROPERTY(Category = Node, EditAnywhere)
+	UPROPERTY(Category = Node, EditAnywhere, config)
 	uint32 bStopOnOverlap : 1;
 
 	virtual EBTNodeResult::Type ExecuteTask(UBehaviorTreeComponent& OwnerComp, uint8* NodeMemory) override;
diff --git a/Engine/Source/Runtime/AIModule/Classes/BehaviorTree/Tasks/BTTask_MoveTo.h b/Engine/Source/Runtime/AIModule/Classes/BehaviorTree/Tasks/BTTask_MoveTo.h
--- a/Engine/Source/Runtime/AIModule/Classes/BehaviorTree/Tasks/BTTask_MoveTo.h
+++ b/Engine/Source/Runtime/AIModule/Classes/BehaviorTree/Tasks/BTTask_MoveTo.h
@@ -38,7 +38,7 @@ class AIMODULE_API UBTTask_MoveTo : public UBTTask_BlackboardBase
 
 	/** if set to true agent's radius will be added to AcceptableRadius for purposes of checking 
 	 *	if path's end point has been reached. Will result in AI stopping on contact with destination location*/
-	UPROPERTY(Category=Node, EditAnywhere)
+	UPROPERTY(Category=Node, EditAnywhere, config)
 	uint32 bStopOnOverlap : 1;
 	
 	virtual EBTNodeResult::Type ExecuteTask(UBehaviorTreeComponent& OwnerComp, uint8* NodeMemory) override;
