Exposed BlackboardComponent as an output param of UseBlackboard

This allows us to use convinence functions like "Set Value As..." as the component is not part of the valid BP context when dragging off UseBlackboard.

UE-3014 - EDITOR: BP: Set Value as... nodes do not show up in Use Blackboard's context sensitive menu

diff --git a/Engine/Source/Runtime/AIModule/Classes/AIController.h b/Engine/Source/Runtime/AIModule/Classes/AIController.h
--- a/Engine/Source/Runtime/AIModule/Classes/AIController.h
+++ b/Engine/Source/Runtime/AIModule/Classes/AIController.h
@@ -222,9 +222,14 @@ public:
 	UFUNCTION(BlueprintCallable, Category="AI")
 	virtual bool RunBehaviorTree(UBehaviorTree* BTAsset);
 
-	/** makes AI use specified BB asset */
+	/** 
+	 * Makes AI use the specified Blackboard asset & creates a Blackboard Component if one does not already exist.
+	 * @param	BlackboardAsset			The Blackboard asset to use.
+	 * @param	BlackboardComponent		The Blackboard component that was used or created to work with the passed-in Blackboard Asset.
+	 * @return true if we successfully linked the blackboard asset to the blackboard component.
+	 */
 	UFUNCTION(BlueprintCallable, Category = "AI")
-	bool UseBlackboard(UBlackboardData* BlackboardAsset);
+	bool UseBlackboard(UBlackboardData* BlackboardAsset, UBlackboardComponent*& BlackboardComponent);
 
 	/** does this AIController allow given UBlackboardComponent sync data with it */
 	virtual bool ShouldSyncBlackboardWith(const UBlackboardComponent& OtherBlackboardComponent) const { return true; }
diff --git a/Engine/Source/Runtime/AIModule/Private/AIController.cpp b/Engine/Source/Runtime/AIModule/Private/AIController.cpp
--- a/Engine/Source/Runtime/AIModule/Private/AIController.cpp
+++ b/Engine/Source/Runtime/AIModule/Private/AIController.cpp
@@ -779,8 +779,7 @@ bool AAIController::RunBehaviorTree(UBehaviorTree* BTAsset)
 	UBlackboardComponent* BlackboardComp = NULL;
 	if (BTAsset->BlackboardAsset)
 	{
-		bSuccess = UseBlackboard(BTAsset->BlackboardAsset);
-		BlackboardComp = FindComponentByClass<UBlackboardComponent>();
+		bSuccess = UseBlackboard(BTAsset->BlackboardAsset, BlackboardComp);
 	}
 
 	if (bSuccess)
@@ -812,7 +811,7 @@ bool AAIController::InitializeBlackboard(UBlackboardComponent& BlackboardComp, U
 	return false;
 }
 
-bool AAIController::UseBlackboard(UBlackboardData* BlackboardAsset)
+bool AAIController::UseBlackboard(UBlackboardData* BlackboardAsset, UBlackboardComponent*& BlackboardComponent)
 {
 	if (BlackboardAsset == NULL)
 	{
@@ -844,6 +843,8 @@ bool AAIController::UseBlackboard(UBlackboardData* BlackboardAsset)
 		InitializeBlackboard(*BlackboardComp, *BlackboardAsset);
 	}
 
+	BlackboardComponent = BlackboardComp;
+
 	return bSuccess;
 }
 
