added support for following a based navigation path
#ue4

diff --git a/Engine/Source/Runtime/AIModule/Classes/Navigation/PathFollowingComponent.h b/Engine/Source/Runtime/AIModule/Classes/Navigation/PathFollowingComponent.h
--- a/Engine/Source/Runtime/AIModule/Classes/Navigation/PathFollowingComponent.h
+++ b/Engine/Source/Runtime/AIModule/Classes/Navigation/PathFollowingComponent.h
@@ -212,10 +212,10 @@ class AIMODULE_API UPathFollowingComponent : public UActorComponent, public IAIR
 	FORCEINLINE FAIRequestID GetCurrentRequestId() const { return CurrentRequestId; }
 	FORCEINLINE uint32 GetCurrentPathIndex() const { return MoveSegmentStartIndex; }
 	FORCEINLINE uint32 GetNextPathIndex() const { return MoveSegmentEndIndex; }
+	FORCEINLINE UObject* GetCurrentCustomLinkOb() const { return CurrentCustomLinkOb.Get(); }
 	FORCEINLINE FVector GetCurrentTargetLocation() const { return *CurrentDestination; }
-	FORCEINLINE FVector GetCurrentDirection() const { return MoveSegmentDirection; }
 	FORCEINLINE FBasedPosition GetCurrentTargetLocationBased() const { return CurrentDestination; }
-	FORCEINLINE UObject* GetCurrentCustomLinkOb() const { return CurrentCustomLinkOb.Get(); }
+	FVector GetCurrentDirection() const;
 
 	/** will be deprecated soon, please use AIController.GetMoveStatus instead! */
 	UFUNCTION(BlueprintCallable, Category="AI|Components|PathFollowing")
diff --git a/Engine/Source/Runtime/AIModule/Private/Navigation/PathFollowingComponent.cpp b/Engine/Source/Runtime/AIModule/Private/Navigation/PathFollowingComponent.cpp
--- a/Engine/Source/Runtime/AIModule/Private/Navigation/PathFollowingComponent.cpp
+++ b/Engine/Source/Runtime/AIModule/Private/Navigation/PathFollowingComponent.cpp
@@ -61,7 +61,7 @@ void UPathFollowingComponent::LogPathHelper(const AActor* LogOwner, FNavigationP
 		FVisualLogEntry* Entry = Vlog.GetEntryToWrite(LogOwner, LogOwner->GetWorld()->TimeSeconds);
 		InLogPath->DescribeSelfToVisLog(Entry);
 
-		const FVector PathEnd = InLogPath->GetPathPoints().Last().Location;
+		const FVector PathEnd = *InLogPath->GetPathPointLocation(InLogPath->GetPathPoints().Num() - 1);
 		if (LogGoalActor)
 		{
 			const FVector GoalLoc = LogGoalActor->GetActorLocation();
@@ -483,7 +483,8 @@ FVector UPathFollowingComponent::GetMoveFocus(bool bAllowStrafe) const
 	}
 	else
 	{
-		MoveFocus = *CurrentDestination + MoveSegmentDirection * 20.0f;
+		const FVector CurrentMoveDirection = GetCurrentDirection();
+		MoveFocus = *CurrentDestination + (CurrentMoveDirection * 20.0f);
 	}
 
 	return MoveFocus;
@@ -544,10 +545,12 @@ int32 UPathFollowingComponent::DetermineStartingPathPoint(const FNavigationPath*
 			{
 				// check if is closer to first or second path point (don't assume AI's standing)
 				const FVector CurrentLocation = MovementComp->GetActorFeetLocation();
+				const FVector PathPt0 = *ConsideredPath->GetPathPointLocation(0);
+				const FVector PathPt1 = *ConsideredPath->GetPathPointLocation(1);
 				// making this test in 2d to avoid situation where agent's Z location not being in "navmesh plane"
 				// would influence the result
-				const float SqDistToFirstPoint = (CurrentLocation - ConsideredPath->GetPathPoints()[0].Location).SizeSquared2D();
-				const float SqDistToSecondPoint = (CurrentLocation - ConsideredPath->GetPathPoints()[1].Location).SizeSquared2D();
+				const float SqDistToFirstPoint = (CurrentLocation - PathPt0).SizeSquared2D();
+				const float SqDistToSecondPoint = (CurrentLocation - PathPt1).SizeSquared2D();
 				PickedPathPoint = (SqDistToFirstPoint < SqDistToSecondPoint) ? 0 : 1;
 			}
 			else
@@ -583,20 +586,22 @@ void UPathFollowingComponent::SetMoveSegment(int32 SegmentStartIndex)
 		MoveSegmentStartRef = PathPt0.NodeRef;
 		MoveSegmentEndRef = PathPt1.NodeRef;
 		
-		const FVector SegmentStart = PathPt0.Location;
-		const FVector SegmentEnd = PathPt1.Location;
+		CurrentDestination = Path->GetPathPointLocation(MoveSegmentEndIndex);
+		const FVector SegmentStart = *Path->GetPathPointLocation(MoveSegmentStartIndex);
+		FVector SegmentEnd = *CurrentDestination;
 
-		CurrentDestination.Set(Path->GetBaseActor(), SegmentEnd);
-		CurrentAcceptanceRadius = (Path->GetPathPoints().Num() == (EndSegmentIndex + 1)) ? AcceptanceRadius : 0.0f;
-
-		MoveSegmentDirection = (SegmentEnd - SegmentStart).GetSafeNormal();
-		
 		// make sure we have a non-zero direction if still following a valid path
-		if (MoveSegmentDirection.IsZero() && int32(MoveSegmentEndIndex + 1) < Path->GetPathPoints().Num())
+		if (SegmentStart.Equals(SegmentEnd) && Path->GetPathPoints().IsValidIndex(MoveSegmentEndIndex + 1))
 		{
-			MoveSegmentDirection = (Path->GetPathPoints()[MoveSegmentEndIndex + 1].Location - SegmentStart).GetSafeNormal();
+			MoveSegmentEndIndex++;
+
+			CurrentDestination = Path->GetPathPointLocation(MoveSegmentEndIndex);
+			SegmentEnd = *CurrentDestination;
 		}
 
+		CurrentAcceptanceRadius = (Path->GetPathPoints().Num() == (MoveSegmentEndIndex + 1)) ? AcceptanceRadius : 0.0f;
+		MoveSegmentDirection = (SegmentEnd - SegmentStart).GetSafeNormal();
+
 		// handle moving through custom nav links
 		if (PathPt0.CustomLinkId)
 		{
@@ -692,7 +697,8 @@ int32 UPathFollowingComponent::OptimizeSegmentVisibility(int32 StartIndex)
 		}
 
 		const FNavPathPoint& PathPt1 = Path->GetPathPoints()[Index];
-		const FVector AdjustedDestination = PathPt1.Location + FVector(0.0f, 0.0f, PawnHalfHeight);
+		const FVector PathPt1Location = *Path->GetPathPointLocation(Index);
+		const FVector AdjustedDestination = PathPt1Location + FVector(0.0f, 0.0f, PawnHalfHeight);
 
 		FVector HitLocation;
 #if WITH_RECAST
@@ -810,7 +816,8 @@ void UPathFollowingComponent::UpdatePathSegment()
 			if (Path->GetPathPoints().IsValidIndex(MoveSegmentEndIndex) && Path->GetPathPoints().IsValidIndex(MoveSegmentStartIndex))
 			{
 				LogBlockHelper(GetOwner(), MovementComp, MinAgentRadiusPct, MinAgentHalfHeightPct,
-					Path->GetPathPoints()[MoveSegmentStartIndex].Location, Path->GetPathPoints()[MoveSegmentEndIndex].Location);
+					*Path->GetPathPointLocation(MoveSegmentStartIndex),
+					*Path->GetPathPointLocation(MoveSegmentEndIndex));
 			}
 			else
 			{
@@ -883,7 +890,7 @@ bool UPathFollowingComponent::HasReached(const AActor& TestGoal, float InAccepta
 bool UPathFollowingComponent::HasReachedDestination(const FVector& CurrentLocation) const
 {
 	// get cylinder at goal location
-	FVector GoalLocation = Path->GetPathPoints()[Path->GetPathPoints().Num() - 1].Location;
+	FVector GoalLocation = *Path->GetPathPointLocation(Path->GetPathPoints().Num() - 1);
 	float GoalRadius = 0.0f;
 	float GoalHalfHeight = 0.0f;
 	
@@ -914,10 +921,11 @@ bool UPathFollowingComponent::HasReachedCurrentTarget(const FVector& CurrentLoca
 	}
 
 	const FVector CurrentTarget = GetCurrentTargetLocation();
+	const FVector CurrentDirection = GetCurrentDirection();
 
 	// check if moved too far
 	const FVector ToTarget = (CurrentTarget - MovementComp->GetActorFeetLocation());
-	const float SegmentDot = FVector::DotProduct(ToTarget, MoveSegmentDirection);
+	const float SegmentDot = FVector::DotProduct(ToTarget, CurrentDirection);
 	if (SegmentDot < 0.0)
 	{
 		return true;
@@ -990,7 +998,7 @@ void UPathFollowingComponent::DebugReachTest(float& CurrentDot, float& CurrentDi
 
 	if (bFollowingLastSegment)
 	{
-		GoalLocation = Path->GetPathPoints()[Path->GetPathPoints().Num() - 1].Location;
+		GoalLocation = *Path->GetPathPointLocation(Path->GetPathPoints().Num() - 1);
 		AgentRadiusPct = MinAgentRadiusPct;
 
 		// take goal's current location, unless path is partial
@@ -1011,7 +1019,8 @@ void UPathFollowingComponent::DebugReachTest(float& CurrentDot, float& CurrentDi
 	}
 
 	const FVector ToGoal = (GoalLocation - AgentLocation);
-	CurrentDot = FVector::DotProduct(ToGoal.GetSafeNormal(), MoveSegmentDirection);
+	const FVector CurrentDirection = GetCurrentDirection();
+	CurrentDot = FVector::DotProduct(ToGoal.GetSafeNormal(), CurrentDirection);
 	bDotFailed = (CurrentDot < 0.0f) ? 1 : 0;
 
 	// get cylinder of moving agent
@@ -1261,6 +1270,21 @@ FNavLocation UPathFollowingComponent::GetCurrentNavLocation() const
 	return CurrentNavLocation;
 }
 
+FVector UPathFollowingComponent::GetCurrentDirection() const
+{
+	if (CurrentDestination.Base)
+	{
+		// calculate direction to based destination
+		const FVector SegmentStartLocation = *Path->GetPathPointLocation(MoveSegmentStartIndex);
+		const FVector SegmentEndLocation = *CurrentDestination;
+
+		return (SegmentEndLocation - SegmentStartLocation).GetSafeNormal();
+	}
+
+	// use cached direction of current path segment
+	return MoveSegmentDirection;
+}
+
 // will be deprecated soon, please don't use it!
 EPathFollowingAction::Type UPathFollowingComponent::GetPathActionType() const
 {
diff --git a/Engine/Source/Runtime/Engine/Classes/AI/Navigation/NavigationData.h b/Engine/Source/Runtime/Engine/Classes/AI/Navigation/NavigationData.h
--- a/Engine/Source/Runtime/Engine/Classes/AI/Navigation/NavigationData.h
+++ b/Engine/Source/Runtime/Engine/Classes/AI/Navigation/NavigationData.h
@@ -239,6 +239,9 @@ struct ENGINE_API FNavigationPath : public TSharedFromThis<FNavigationPath, ESPM
 		return PathPoints;
 	}
 
+	/** get based position of path point */
+	FBasedPosition GetPathPointLocation(uint32 Index) const;
+
 	/** checks if given path, starting from StartingIndex, intersects with given AABB box */
 	virtual bool DoesIntersectBox(const FBox& Box, uint32 StartingIndex = 0, int32* IntersectingSegmentIndex = NULL, FVector* AgentExtent = NULL) const;
 	/** checks if given path, starting from StartingIndex, intersects with given AABB box. This version uses AgentLocation as beginning of the path
diff --git a/Engine/Source/Runtime/Engine/Private/AI/Navigation/NavigationPath.cpp b/Engine/Source/Runtime/Engine/Private/AI/Navigation/NavigationPath.cpp
--- a/Engine/Source/Runtime/Engine/Private/AI/Navigation/NavigationPath.cpp
+++ b/Engine/Source/Runtime/Engine/Private/AI/Navigation/NavigationPath.cpp
@@ -324,6 +324,18 @@ FVector FNavigationPath::GetSegmentDirection(uint32 SegmentEndIndex) const
 	return Result;
 }
 
+FBasedPosition FNavigationPath::GetPathPointLocation(uint32 Index) const
+{
+	FBasedPosition BasedPt;
+	if (PathPoints.IsValidIndex(Index))
+	{
+		BasedPt.Base = Base.Get();
+		BasedPt.Position = PathPoints[Index].Location;
+	}
+
+	return BasedPt;
+}
+
 #if ENABLE_VISUAL_LOG
 
 void FNavigationPath::DescribeSelfToVisLog(FVisualLogEntry* Snapshot) const 
