removed failsafe huge tick interval from TimeLimit decorator, it's not longer needed
#ue4 UE-6627

diff --git a/Engine/Source/Runtime/AIModule/Private/BehaviorTree/Decorators/BTDecorator_TimeLimit.cpp b/Engine/Source/Runtime/AIModule/Private/BehaviorTree/Decorators/BTDecorator_TimeLimit.cpp
--- a/Engine/Source/Runtime/AIModule/Private/BehaviorTree/Decorators/BTDecorator_TimeLimit.cpp
+++ b/Engine/Source/Runtime/AIModule/Private/BehaviorTree/Decorators/BTDecorator_TimeLimit.cpp
@@ -27,9 +27,6 @@ void UBTDecorator_TimeLimit::OnBecomeRelevant(UBehaviorTreeComponent& OwnerComp,
 void UBTDecorator_TimeLimit::TickNode(UBehaviorTreeComponent& OwnerComp, uint8* NodeMemory, float DeltaSeconds)
 {
 	OwnerComp.RequestExecution(this);
-
-	// make sure that next tick won't happen before execution request is processed
-	SetNextTickTime(NodeMemory, 1000000.0f);
 }
 
 FString UBTDecorator_TimeLimit::GetStaticDescription() const
