changed all overrides of agent radius multiplier in path following tests to parameter in test function
#ue4

diff --git a/Engine/Source/Runtime/AIModule/Classes/Navigation/PathFollowingComponent.h b/Engine/Source/Runtime/AIModule/Classes/Navigation/PathFollowingComponent.h
--- a/Engine/Source/Runtime/AIModule/Classes/Navigation/PathFollowingComponent.h
+++ b/Engine/Source/Runtime/AIModule/Classes/Navigation/PathFollowingComponent.h
@@ -417,7 +417,9 @@ protected:
 	bool HasReachedCurrentTarget(const FVector& CurrentLocation) const;
 
 	/** check if moving agent has reached goal defined by cylinder */
+	DEPRECATED(4.8, "Please use override with AgentRadiusMultiplier instead of this.")
 	bool HasReachedInternal(const FVector& GoalLocation, float GoalRadius, float GoalHalfHeight, const FVector& AgentLocation, float RadiusThreshold, bool bSuccessOnRadiusOverlap) const;
+	bool HasReachedInternal(const FVector& GoalLocation, float GoalRadius, float GoalHalfHeight, const FVector& AgentLocation, float RadiusThreshold, float AgentRadiusMultiplier) const;
 
 	/** check if agent is on path */
 	virtual bool IsOnPath() const;
diff --git a/Engine/Source/Runtime/AIModule/Private/Navigation/CrowdFollowingComponent.cpp b/Engine/Source/Runtime/AIModule/Private/Navigation/CrowdFollowingComponent.cpp
--- a/Engine/Source/Runtime/AIModule/Private/Navigation/CrowdFollowingComponent.cpp
+++ b/Engine/Source/Runtime/AIModule/Private/Navigation/CrowdFollowingComponent.cpp
@@ -664,7 +664,7 @@ void UCrowdFollowingComponent::UpdatePathSegment()
 
 			// can't use HasReachedDestination here, because it will use last path point
 			// which is not set correctly for partial paths without string pulling
-			if (bMovedTooFar || HasReachedInternal(GoalLocation, 0.0f, 0.0f, CurrentLocation, AcceptanceRadius, bStopOnOverlap))
+			if (bMovedTooFar || HasReachedInternal(GoalLocation, 0.0f, 0.0f, CurrentLocation, AcceptanceRadius, bStopOnOverlap ? MinAgentRadiusPct : 0.0f))
 			{
 				UE_VLOG(GetOwner(), LogCrowdFollowing, Log, TEXT("Last path segment finished due to \'%s\'"), bMovedTooFar ? TEXT("Missing Last Point") : TEXT("Reaching Destination"));
 				OnPathFinished(EPathFollowingResult::Success);
@@ -673,12 +673,8 @@ void UCrowdFollowingComponent::UpdatePathSegment()
 		else
 		{
 			// override radius multiplier and switch to next path part when closer than 4x agent radius
-			const float SavedAgentRadiusPct = MinAgentRadiusPct;
-			MinAgentRadiusPct = 4.0f;
-			
-			const bool bHasReached = HasReachedInternal(GoalLocation, 0.0f, 0.0f, CurrentLocation, 0.0f, false);
-			
-			MinAgentRadiusPct = SavedAgentRadiusPct;
+			const float NextPartMultiplier = 4.0f;
+			const bool bHasReached = HasReachedInternal(GoalLocation, 0.0f, 0.0f, CurrentLocation, 0.0f, NextPartMultiplier);
 
 			if (bHasReached)
 			{
diff --git a/Engine/Source/Runtime/AIModule/Private/Navigation/PathFollowingComponent.cpp b/Engine/Source/Runtime/AIModule/Private/Navigation/PathFollowingComponent.cpp
--- a/Engine/Source/Runtime/AIModule/Private/Navigation/PathFollowingComponent.cpp
+++ b/Engine/Source/Runtime/AIModule/Private/Navigation/PathFollowingComponent.cpp
@@ -860,7 +860,8 @@ bool UPathFollowingComponent::HasReached(const FVector& TestPoint, float InAccep
 	{
 		InAcceptanceRadius = MyDefaultAcceptanceRadius;
 	}
-	return HasReachedInternal(TestPoint, GoalRadius, GoalHalfHeight, CurrentLocation, InAcceptanceRadius, /*bSuccessOnRadiusOverlap=*/!bExactSpot);
+
+	return HasReachedInternal(TestPoint, GoalRadius, GoalHalfHeight, CurrentLocation, InAcceptanceRadius, bExactSpot ? 0.0f : MinAgentRadiusPct);
 }
 
 bool UPathFollowingComponent::HasReached(const AActor& TestGoal, float InAcceptanceRadius, bool bExactSpot) const
@@ -884,7 +885,7 @@ bool UPathFollowingComponent::HasReached(const AActor& TestGoal, float InAccepta
 	}
 
 	const FVector CurrentLocation = MovementComp ? MovementComp->GetActorFeetLocation() : FVector::ZeroVector;
-	return HasReachedInternal(TestPoint, GoalRadius, GoalHalfHeight, CurrentLocation, InAcceptanceRadius, /*bSuccessOnRadiusOverlap=*/!bExactSpot);
+	return HasReachedInternal(TestPoint, GoalRadius, GoalHalfHeight, CurrentLocation, InAcceptanceRadius, bExactSpot ? 0.0f : MinAgentRadiusPct);
 }
 
 bool UPathFollowingComponent::HasReachedDestination(const FVector& CurrentLocation) const
@@ -910,7 +911,7 @@ bool UPathFollowingComponent::HasReachedDestination(const FVector& CurrentLocati
 		}
 	}
 
-	return HasReachedInternal(GoalLocation, GoalRadius, GoalHalfHeight, CurrentLocation, AcceptanceRadius, /*bSuccessOnRadiusOverlap=*/bStopOnOverlap);
+	return HasReachedInternal(GoalLocation, GoalRadius, GoalHalfHeight, CurrentLocation, AcceptanceRadius, bStopOnOverlap ? MinAgentRadiusPct : 0.0f);
 }
 
 bool UPathFollowingComponent::HasReachedCurrentTarget(const FVector& CurrentLocation) const
@@ -935,16 +936,16 @@ bool UPathFollowingComponent::HasReachedCurrentTarget(const FVector& CurrentLoca
 	// don't use acceptance radius here, it has to be exact for moving near corners (2D test < 5% of agent radius)
 	const float GoalRadius = 0.0f;
 	const float GoalHalfHeight = 0.0f;
-	const float OrgAgentRadiusPct = MinAgentRadiusPct;
-	((UPathFollowingComponent*)this)->MinAgentRadiusPct = 0.05f;
 
-	const bool bReached = HasReachedInternal(CurrentTarget, GoalRadius, GoalHalfHeight, CurrentLocation, CurrentAcceptanceRadius, /*bSuccessOnRadiusOverlap=*/true);
+	return HasReachedInternal(CurrentTarget, GoalRadius, GoalHalfHeight, CurrentLocation, CurrentAcceptanceRadius, 0.05f);
+}
 
-	((UPathFollowingComponent*)this)->MinAgentRadiusPct = OrgAgentRadiusPct;
-	return bReached;
+bool UPathFollowingComponent::HasReachedInternal(const FVector& GoalLocation, float GoalRadius, float GoalHalfHeight, const FVector& AgentLocation, float RadiusThreshold, bool bSuccessOnRadiusOverlap) const
+{
+	return HasReachedInternal(GoalLocation, GoalRadius, GoalHalfHeight, AgentLocation, RadiusThreshold, bSuccessOnRadiusOverlap ? MinAgentRadiusPct : 0.0f);
 }
 
-bool UPathFollowingComponent::HasReachedInternal(const FVector& Goal, float GoalRadius, float GoalHalfHeight, const FVector& AgentLocation, float RadiusThreshold, bool bSuccessOnRadiusOverlap) const
+bool UPathFollowingComponent::HasReachedInternal(const FVector& GoalLocation, float GoalRadius, float GoalHalfHeight, const FVector& AgentLocation, float RadiusThreshold, float AgentRadiusMultiplier) const
 {
 	if (MovementComp == NULL)
 	{
@@ -958,10 +959,10 @@ bool UPathFollowingComponent::HasReachedInternal(const FVector& Goal, float Goal
 	MovingAgent->GetSimpleCollisionCylinder(AgentRadius, AgentHalfHeight);
 
 	// check if they overlap (with added AcceptanceRadius)
-	const FVector ToGoal = Goal - AgentLocation;
+	const FVector ToGoal = GoalLocation - AgentLocation;
 
 	const float Dist2DSq = ToGoal.SizeSquared2D();
-	const float UseRadius = FMath::Max(RadiusThreshold, GoalRadius + (bSuccessOnRadiusOverlap ? AgentRadius * MinAgentRadiusPct : 0.0f));
+	const float UseRadius = FMath::Max(RadiusThreshold, GoalRadius + (AgentRadius * AgentRadiusMultiplier));
 	if (Dist2DSq > FMath::Square(UseRadius))
 	{
 		return false;
