integrated EQS changes from dev branch:
CL# 2453024: added navigation filter to EQS pathfinding tests
CL# 2466155: added batch pathfinding test for EQS
CL# 2466237: EQS generators will now cache NavLocation of items if navmesh projection is enabled
CL# 2466356: added support for deprecated nodes in EQS editor
CL# 2466358: pathing grid EQS generator is now deprecated
CL# 2466439: added failsafe in EQS editor for recreating missing nodes
CL# 2467466: fixed recreating test nodes in EQS graph
CL# 2467476: disabled HPA* in pathfinding EQS test
CL# 2467508: removed hierarchical pathfinding option from EQS tests
CL# 2468574: EQS generators will now cache NavLocation of items if navmesh projection is enabled (crash fix)

diff --git a/Engine/Source/Editor/AIGraph/Classes/AIGraph.h b/Engine/Source/Editor/AIGraph/Classes/AIGraph.h
--- a/Engine/Source/Editor/AIGraph/Classes/AIGraph.h
+++ b/Engine/Source/Editor/AIGraph/Classes/AIGraph.h
@@ -11,6 +11,10 @@ class AIGRAPH_API UAIGraph : public UEdGraph
 	UPROPERTY()
 	int32 GraphVersion;
 
+	virtual void OnCreated();
+	virtual void OnLoaded();
+	virtual void Initialize();
+
 	virtual void UpdateAsset(int32 UpdateFlags = 0);
 	virtual void UpdateVersion();
 	virtual void MarkVersion();
@@ -18,6 +22,7 @@ class AIGRAPH_API UAIGraph : public UEdGraph
 	virtual void OnSubNodeDropped();
 
 	bool UpdateUnknownNodeClasses();
+	void UpdateDeprecatedClasses();
 	void RemoveOrphanedNodes();
 
 	bool IsLocked() const;
diff --git a/Engine/Source/Editor/AIGraph/Private/AIGraph.cpp b/Engine/Source/Editor/AIGraph/Private/AIGraph.cpp
--- a/Engine/Source/Editor/AIGraph/Private/AIGraph.cpp
+++ b/Engine/Source/Editor/AIGraph/Private/AIGraph.cpp
@@ -13,6 +13,22 @@ void UAIGraph::UpdateAsset(int32 UpdateFlags)
 	// empty in base class
 }
 
+void UAIGraph::OnCreated()
+{
+	MarkVersion();
+}
+
+void UAIGraph::OnLoaded()
+{
+	UpdateDeprecatedClasses();
+	UpdateUnknownNodeClasses();
+}
+
+void UAIGraph::Initialize()
+{
+	UpdateVersion();
+}
+
 void UAIGraph::UpdateVersion()
 {
 	if (GraphVersion == 1)
@@ -51,6 +67,29 @@ bool UAIGraph::UpdateUnknownNodeClasses()
 	return bUpdated;
 }
 
+void UAIGraph::UpdateDeprecatedClasses()
+{
+	for (int32 Idx = 0; Idx < Nodes.Num(); Idx++)
+	{
+		UAIGraphNode* Node = Cast<UAIGraphNode>(Nodes[Idx]);
+		if (Node)
+		{
+			if (Node->NodeInstance)
+			{
+				Node->ErrorMessage = FGraphNodeClassHelper::GetDeprecationMessage(Node->NodeInstance->GetClass());
+			}
+
+			for (int32 SubIdx = 0; SubIdx < Node->SubNodes.Num(); SubIdx++)
+			{
+				if (Node->SubNodes[SubIdx] && Node->SubNodes[SubIdx]->NodeInstance)
+				{
+					Node->SubNodes[SubIdx]->ErrorMessage = FGraphNodeClassHelper::GetDeprecationMessage(Node->SubNodes[SubIdx]->NodeInstance->GetClass());
+				}
+			}
+		}
+	}
+}
+
 void UAIGraph::CollectAllNodeInstances(TSet<UObject*>& NodeInstances)
 {
 	for (int32 Idx = 0; Idx < Nodes.Num(); Idx++)
diff --git a/Engine/Source/Editor/BehaviorTreeEditor/Classes/BehaviorTreeGraph.h b/Engine/Source/Editor/BehaviorTreeEditor/Classes/BehaviorTreeGraph.h
--- a/Engine/Source/Editor/BehaviorTreeEditor/Classes/BehaviorTreeGraph.h
+++ b/Engine/Source/Editor/BehaviorTreeEditor/Classes/BehaviorTreeGraph.h
@@ -24,6 +24,10 @@ class UBehaviorTreeGraph : public UAIGraph
 	UPROPERTY()
 	bool bIsUsingModCounter;
 
+	virtual void OnCreated() override;
+	virtual void OnLoaded() override;
+	virtual void Initialize() override;
+
 	virtual void UpdateVersion() override;
 	virtual void MarkVersion() override;
 	virtual void UpdateAsset(int32 UpdateFlags = 0) override;
diff --git a/Engine/Source/Editor/BehaviorTreeEditor/Private/BehaviorTreeEditor.cpp b/Engine/Source/Editor/BehaviorTreeEditor/Private/BehaviorTreeEditor.cpp
--- a/Engine/Source/Editor/BehaviorTreeEditor/Private/BehaviorTreeEditor.cpp
+++ b/Engine/Source/Editor/BehaviorTreeEditor/Private/BehaviorTreeEditor.cpp
@@ -263,24 +263,19 @@ void FBehaviorTreeEditor::RestoreBehaviorTree()
 	{
 		BehaviorTree->BTGraph = FBlueprintEditorUtils::CreateNewGraph(BehaviorTree, TEXT("Behavior Tree"), UBehaviorTreeGraph::StaticClass(), UEdGraphSchema_BehaviorTree::StaticClass());
 		MyGraph = Cast<UBehaviorTreeGraph>(BehaviorTree->BTGraph);
-		MyGraph->MarkVersion();
 
 		// Initialize the behavior tree graph
 		const UEdGraphSchema* Schema = MyGraph->GetSchema();
 		Schema->CreateDefaultNodesForGraph(*MyGraph);
 
-		MyGraph->SpawnMissingNodes();
+		MyGraph->OnCreated();
 	}
 	else
 	{
-		MyGraph->UpdatePinConnectionTypes();
-		MyGraph->UpdateDeprecatedNodes();
+		MyGraph->OnLoaded();
 	}
 
-	MyGraph->UpdateVersion();
-	MyGraph->UpdateBlackboardChange();
-	MyGraph->UpdateInjectedNodes();
-	MyGraph->UpdateUnknownNodeClasses();
+	MyGraph->Initialize();
 
 	TSharedRef<FTabPayload_UObject> Payload = FTabPayload_UObject::Make(MyGraph);
 	TSharedPtr<SDockTab> DocumentTab = DocumentManager->OpenDocument(Payload, bNewGraph ? FDocumentTracker::OpenNewDocument : FDocumentTracker::RestorePreviousDocument);
diff --git a/Engine/Source/Editor/BehaviorTreeEditor/Private/BehaviorTreeGraph.cpp b/Engine/Source/Editor/BehaviorTreeEditor/Private/BehaviorTreeGraph.cpp
--- a/Engine/Source/Editor/BehaviorTreeEditor/Private/BehaviorTreeGraph.cpp
+++ b/Engine/Source/Editor/BehaviorTreeEditor/Private/BehaviorTreeGraph.cpp
@@ -148,6 +148,26 @@ void UBehaviorTreeGraph::UpdateAsset(int32 UpdateFlags)
 	}
 }
 
+void UBehaviorTreeGraph::OnCreated()
+{
+	Super::OnCreated();
+	SpawnMissingNodes();
+}
+
+void UBehaviorTreeGraph::OnLoaded()
+{
+	Super::OnLoaded();
+	UpdatePinConnectionTypes();
+	UpdateDeprecatedNodes();
+}
+
+void UBehaviorTreeGraph::Initialize()
+{
+	Super::Initialize();
+	UpdateBlackboardChange();
+	UpdateInjectedNodes();
+}
+
 void UBehaviorTreeGraph::UpdatePinConnectionTypes()
 {
 	for (int32 Index = 0; Index < Nodes.Num(); ++Index)
@@ -207,28 +227,6 @@ void UBehaviorTreeGraph::UpdateDeprecatedNodes()
 
 				ReplaceNodeConnections(Node, NewNode);
 				Nodes[Index] = NewNode;
-				Node = NewNode;
-			}
-
-			if (Node->NodeInstance)
-			{
-				Node->ErrorMessage = FGraphNodeClassHelper::GetDeprecationMessage(Node->NodeInstance->GetClass());
-			}
-
-			for (int32 i = 0; i < Node->Decorators.Num(); i++)
-			{
-				if (Node->Decorators[i] && Node->Decorators[i]->NodeInstance)
-				{
-					Node->Decorators[i]->ErrorMessage = FGraphNodeClassHelper::GetDeprecationMessage(Node->Decorators[i]->NodeInstance->GetClass());
-				}
-			}
-
-			for (int32 i = 0; i < Node->Services.Num(); i++)
-			{
-				if (Node->Services[i] && Node->Services[i]->NodeInstance)
-				{
-					Node->Services[i]->ErrorMessage = FGraphNodeClassHelper::GetDeprecationMessage(Node->Services[i]->NodeInstance->GetClass());
-				}
 			}
 		}
 	}
diff --git a/Engine/Source/Editor/EnvironmentQueryEditor/Classes/EnvironmentQueryGraph.h b/Engine/Source/Editor/EnvironmentQueryEditor/Classes/EnvironmentQueryGraph.h
--- a/Engine/Source/Editor/EnvironmentQueryEditor/Classes/EnvironmentQueryGraph.h
+++ b/Engine/Source/Editor/EnvironmentQueryEditor/Classes/EnvironmentQueryGraph.h
@@ -9,10 +9,14 @@ class UEnvironmentQueryGraph : public UAIGraph
 {
 	GENERATED_UCLASS_BODY()
 
+	virtual void Initialize() override;
+	virtual void OnLoaded() override;
 	virtual void UpdateVersion() override;
 	virtual void MarkVersion() override;
 	virtual void UpdateAsset(int32 UpdateFlags = 0) override;
 
+	void UpdateDeprecatedGeneratorClasses();
+	void SpawnMissingNodes();
 	void CalculateAllWeights();
 	void CreateEnvQueryFromGraph(class UEnvironmentQueryGraphNode* RootEdNode);
 
diff --git a/Engine/Source/Editor/EnvironmentQueryEditor/Classes/EnvironmentQueryGraphNode_Root.h b/Engine/Source/Editor/EnvironmentQueryEditor/Classes/EnvironmentQueryGraphNode_Root.h
--- a/Engine/Source/Editor/EnvironmentQueryEditor/Classes/EnvironmentQueryGraphNode_Root.h
+++ b/Engine/Source/Editor/EnvironmentQueryEditor/Classes/EnvironmentQueryGraphNode_Root.h
@@ -9,4 +9,5 @@ class UEnvironmentQueryGraphNode_Root : public UEnvironmentQueryGraphNode
 
 	virtual void AllocateDefaultPins() override;
 	virtual FText GetNodeTitle(ENodeTitleType::Type TitleType) const override;
+	virtual bool HasErrors() const override { return false; }
 };
diff --git a/Engine/Source/Editor/EnvironmentQueryEditor/Private/EnvironmentQueryColors.h b/Engine/Source/Editor/EnvironmentQueryEditor/Private/EnvironmentQueryColors.h
--- a/Engine/Source/Editor/EnvironmentQueryEditor/Private/EnvironmentQueryColors.h
+++ b/Engine/Source/Editor/EnvironmentQueryEditor/Private/EnvironmentQueryColors.h
@@ -10,6 +10,7 @@ namespace EnvironmentQueryColors
 		const FLinearColor Generator(0.1f, 0.1f, 0.1f);
 		const FLinearColor Test(0.0f, 0.07f, 0.4f);
 		const FLinearColor TestInactive(0.1f, 0.1f, 0.1f);
+		const FLinearColor Error(1.0f, 0.0f, 0.0f);
 	}
 
 	namespace NodeBorder
diff --git a/Engine/Source/Editor/EnvironmentQueryEditor/Private/EnvironmentQueryEditor.cpp b/Engine/Source/Editor/EnvironmentQueryEditor/Private/EnvironmentQueryEditor.cpp
--- a/Engine/Source/Editor/EnvironmentQueryEditor/Private/EnvironmentQueryEditor.cpp
+++ b/Engine/Source/Editor/EnvironmentQueryEditor/Private/EnvironmentQueryEditor.cpp
@@ -49,6 +49,7 @@ void FEnvironmentQueryEditor::UnregisterTabSpawners(const TSharedRef<class FTabM
 
 void FEnvironmentQueryEditor::InitEnvironmentQueryEditor( const EToolkitMode::Type Mode, const TSharedPtr< class IToolkitHost >& InitToolkitHost, UEnvQuery* InScript )
 {
+	SelectedNodesCount = 0;
 	Query = InScript;
 	check(Query != NULL);
 
@@ -166,22 +167,24 @@ TSharedRef<SGraphEditor> FEnvironmentQueryEditor::CreateGraphEditorWidget(UEdGra
 TSharedRef<SDockTab> FEnvironmentQueryEditor::SpawnTab_UpdateGraph( const FSpawnTabArgs& Args )
 {
 	check( Args.GetTabId().TabType == EQSUpdateGraphTabId );
+	UEnvironmentQueryGraph* MyGraph = Cast<UEnvironmentQueryGraph>(Query->EdGraph);
 	if (Query->EdGraph == NULL)
 	{
-		UEnvironmentQueryGraph* MyGraph = NewObject<UEnvironmentQueryGraph>(Query, NAME_None, RF_Transactional);
-		MyGraph->MarkVersion();
-
+		MyGraph = NewObject<UEnvironmentQueryGraph>(Query, NAME_None, RF_Transactional);
 		Query->EdGraph = MyGraph;
 
 		// let's read data from BT script and generate nodes
 		const UEdGraphSchema* Schema = Query->EdGraph->GetSchema();
 		Schema->CreateDefaultNodesForGraph(*Query->EdGraph);
+
+		MyGraph->OnCreated();
+	}
+	else
+	{
+		MyGraph->OnLoaded();
 	}
 
-	// Conversion should be removed are a while
-	UEnvironmentQueryGraph* QueryGraph = Cast<UEnvironmentQueryGraph>(Query->EdGraph);
-	QueryGraph->UpdateVersion();
-	QueryGraph->CalculateAllWeights();
+	MyGraph->Initialize();
 
 	TSharedRef<SGraphEditor> UpdateGraphEditor = CreateGraphEditorWidget(Query->EdGraph);
 	UpdateGraphEdPtr = UpdateGraphEditor; // Keep pointer to editor
diff --git a/Engine/Source/Editor/EnvironmentQueryEditor/Private/EnvironmentQueryGraph.cpp b/Engine/Source/Editor/EnvironmentQueryEditor/Private/EnvironmentQueryGraph.cpp
--- a/Engine/Source/Editor/EnvironmentQueryEditor/Private/EnvironmentQueryGraph.cpp
+++ b/Engine/Source/Editor/EnvironmentQueryEditor/Private/EnvironmentQueryGraph.cpp
@@ -90,6 +90,19 @@ void UEnvironmentQueryGraph::UpdateAsset(int32 UpdateFlags)
 #endif
 }
 
+void UEnvironmentQueryGraph::Initialize()
+{
+	Super::Initialize();
+	SpawnMissingNodes();
+	CalculateAllWeights();
+}
+
+void UEnvironmentQueryGraph::OnLoaded()
+{
+	Super::OnLoaded();
+	UpdateDeprecatedGeneratorClasses();
+}
+
 void UEnvironmentQueryGraph::CalculateAllWeights()
 {
 	for (int32 Idx = 0; Idx < Nodes.Num(); Idx++)
@@ -255,3 +268,62 @@ void UEnvironmentQueryGraph::CollectAllNodeInstances(TSet<UObject*>& NodeInstanc
 		}
 	}
 }
+
+void UEnvironmentQueryGraph::UpdateDeprecatedGeneratorClasses()
+{
+	for (int32 Idx = 0; Idx < Nodes.Num(); Idx++)
+	{
+		UEnvironmentQueryGraphNode* MyNode = Cast<UEnvironmentQueryGraphNode>(Nodes[Idx]);
+		UEnvQueryOption* OptionInstance = MyNode ? Cast<UEnvQueryOption>(MyNode->NodeInstance) : nullptr;
+		if (OptionInstance && OptionInstance->Generator)
+		{
+			MyNode->ErrorMessage = FGraphNodeClassHelper::GetDeprecationMessage(OptionInstance->Generator->GetClass());
+		}
+	}
+}
+
+void UEnvironmentQueryGraph::SpawnMissingNodes()
+{
+	TSet<UEnvQueryOption*> ExistingNodes;
+
+	for (int32 Idx = 0; Idx < Nodes.Num(); Idx++)
+	{
+		UEnvironmentQueryGraphNode* MyNode = Cast<UEnvironmentQueryGraphNode>(Nodes[Idx]);
+		UEnvQueryOption* OptionInstance = MyNode ? Cast<UEnvQueryOption>(MyNode->NodeInstance) : nullptr;
+		if (OptionInstance && OptionInstance->Generator)
+		{
+			ExistingNodes.Add(OptionInstance);
+
+			TSet<UEnvQueryTest*> ExistingTests;
+			for (int32 SubIdx = 0; SubIdx < MyNode->SubNodes.Num(); SubIdx++)
+			{
+				UEnvironmentQueryGraphNode* MySubNode = Cast<UEnvironmentQueryGraphNode>(MyNode->SubNodes[SubIdx]);
+				UEnvQueryTest* TestInstance = MySubNode ? Cast<UEnvQueryTest>(MySubNode->NodeInstance) : nullptr;
+				if (TestInstance)
+				{
+					ExistingTests.Add(TestInstance);
+				}
+				else
+				{
+					MyNode->RemoveSubNode(MySubNode);
+					SubIdx--;
+				}
+			}
+
+			TArray<UEnvQueryTest*> TestsCopy = OptionInstance->Tests;
+			for (int32 SubIdx = 0; SubIdx < TestsCopy.Num(); SubIdx++)
+			{
+				if (ExistingTests.Contains(TestsCopy[SubIdx]) || (TestsCopy[SubIdx] == nullptr))
+				{
+					continue;
+				}
+
+				UEnvironmentQueryGraphNode_Test* TestNode = NewObject<UEnvironmentQueryGraphNode_Test>(this);
+				TestNode->NodeInstance = TestsCopy[SubIdx];
+				UpdateNodeClassData(TestNode, TestsCopy[SubIdx]->GetClass());
+
+				MyNode->AddSubNode(TestNode, this);
+			}
+		}
+	}
+}
diff --git a/Engine/Source/Editor/EnvironmentQueryEditor/Private/SGraphNode_EnvironmentQuery.cpp b/Engine/Source/Editor/EnvironmentQueryEditor/Private/SGraphNode_EnvironmentQuery.cpp
--- a/Engine/Source/Editor/EnvironmentQueryEditor/Private/SGraphNode_EnvironmentQuery.cpp
+++ b/Engine/Source/Editor/EnvironmentQueryEditor/Private/SGraphNode_EnvironmentQuery.cpp
@@ -84,6 +84,11 @@ FSlateColor SGraphNode_EnvironmentQuery::GetBackgroundColor() const
 		}
 	}
 
+	if (MyNode->HasErrors())
+	{
+		NodeColor = EnvironmentQueryColors::NodeBody::Error;
+	}
+
 	return NodeColor;
 }
 
diff --git a/Engine/Source/Runtime/AIModule/Classes/EnvironmentQuery/EnvQueryTypes.h b/Engine/Source/Runtime/AIModule/Classes/EnvironmentQuery/EnvQueryTypes.h
--- a/Engine/Source/Runtime/AIModule/Classes/EnvironmentQuery/EnvQueryTypes.h
+++ b/Engine/Source/Runtime/AIModule/Classes/EnvironmentQuery/EnvQueryTypes.h
@@ -523,7 +523,10 @@ struct AIMODULE_API FEnvQueryOptionInstance
 	TSubclassOf<UEnvQueryItemType> ItemType;
 
 	/** if set, items will be shuffled after tests */
-	bool bShuffleItems;
+	uint32 bShuffleItems : 1;
+
+	/** if set, items will be shuffled after tests */
+	uint32 bHasNavLocations : 1;
 
 	FORCEINLINE uint32 GetAllocatedSize() const { return sizeof(*this) + Tests.GetAllocatedSize(); }
 };
@@ -860,6 +863,11 @@ public:
 			bSkipped++;
 		}
 
+		void IgnoreTimeLimit()
+		{
+			Deadline = -1.0f;
+		}
+
 		operator bool() const
 		{
 			return CurrentItem < Instance->Items.Num() && !Instance->bFoundSingleResult && (Deadline < 0 || FPlatformTime::Seconds() < Deadline);
diff --git a/Engine/Source/Runtime/AIModule/Classes/EnvironmentQuery/Generators/EnvQueryGenerator_OnCircle.h b/Engine/Source/Runtime/AIModule/Classes/EnvironmentQuery/Generators/EnvQueryGenerator_OnCircle.h
--- a/Engine/Source/Runtime/AIModule/Classes/EnvironmentQuery/Generators/EnvQueryGenerator_OnCircle.h
+++ b/Engine/Source/Runtime/AIModule/Classes/EnvironmentQuery/Generators/EnvQueryGenerator_OnCircle.h
@@ -71,5 +71,5 @@ protected:
 		int32 StepsCount, float AngleStep, FEnvQueryInstance& OutQueryInstance) const;
 
 	virtual void AddItemDataForCircle(uint8* ContextRawData, UEnvQueryItemType* ContextItemType, 
-		const TArray<FVector>& Locations, FEnvQueryInstance& OutQueryInstance) const;
+		const TArray<FNavLocation>& Locations, FEnvQueryInstance& OutQueryInstance) const;
 };
diff --git a/Engine/Source/Runtime/AIModule/Classes/EnvironmentQuery/Generators/EnvQueryGenerator_PathingGrid.h b/Engine/Source/Runtime/AIModule/Classes/EnvironmentQuery/Generators/EnvQueryGenerator_PathingGrid.h
--- a/Engine/Source/Runtime/AIModule/Classes/EnvironmentQuery/Generators/EnvQueryGenerator_PathingGrid.h
+++ b/Engine/Source/Runtime/AIModule/Classes/EnvironmentQuery/Generators/EnvQueryGenerator_PathingGrid.h
@@ -1,69 +1,31 @@
 // Copyright 1998-2015 Epic Games, Inc. All Rights Reserved.
 
 #pragma once
-#include "EnvironmentQuery/EnvQueryGenerator.h"
+#include "EnvironmentQuery/Generators/EnvQueryGenerator_SimpleGrid.h"
 #include "EnvQueryGenerator_PathingGrid.generated.h"
 
-class UEnvQueryContext;
 class UNavigationQueryFilter;
-struct FNavLocation;
-#if WITH_RECAST
-class ARecastNavMesh;
-#endif // WITH_RECAST
 
 /**
  *  Navigation grid, generates points on navmesh
  *  with paths to/from context no further than given limit
  */
 
-UCLASS()
-class UEnvQueryGenerator_PathingGrid : public UEnvQueryGenerator
+UCLASS(meta = (DeprecatedNode, DeprecationMessage = "This class is now deprecated, please use SimpleGrid generator with Pathfinding Batch test"))
+class UEnvQueryGenerator_PathingGrid : public UEnvQueryGenerator_SimpleGrid
 {
 	GENERATED_UCLASS_BODY()
 
-	/** max distance of path between point and context */
-	UPROPERTY(EditDefaultsOnly, Category=Generator)
+	UPROPERTY()
 	FAIDataProviderFloatValue MaxDistance;
 
-	/** generation density */
-	UPROPERTY(EditDefaultsOnly, Category=Generator)
-	FAIDataProviderFloatValue SpaceBetween;
-
-	/** path direction switch: set = path from context to points, clear = path from points to context */
-	UPROPERTY(EditDefaultsOnly, Category=Generator)
+	/** DEPRECATED, please use pathfinding (batch) test instead! */
+	UPROPERTY(VisibleDefaultsOnly, Category=Generator)
 	FAIDataProviderBoolValue PathToItem;
 
-	/** context */
-	UPROPERTY(EditAnywhere, Category=Generator)
-	TSubclassOf<UEnvQueryContext> GenerateAround;
-	
-	/** navigation filter for tracing */
-	UPROPERTY(EditDefaultsOnly, Category=Generator)
+	/** DEPRECATED, please use pathfinding (batch) test instead! */
+	UPROPERTY(VisibleDefaultsOnly, Category = Generator)
 	TSubclassOf<UNavigationQueryFilter> NavigationFilter;
 
-	// BEGIN: deprecated properties
-	UPROPERTY()
-	FEnvFloatParam MaxPathDistance;
-
-	UPROPERTY()
-	FEnvFloatParam Density;
-
-	UPROPERTY()
-	FEnvBoolParam PathFromContext;
-	// END: deprecated properties
-
-	virtual void GenerateItems(FEnvQueryInstance& QueryInstance) const override;
 	virtual void PostLoad() override;
-
-	virtual FText GetDescriptionTitle() const override;
-	virtual FText GetDescriptionDetails() const override;
-
-protected:
-#if WITH_RECAST
-	/** find all navmesh node refs in pathing distance */
-	void FindNodeRefsInPathDistance(const ARecastNavMesh* NavMesh, const FVector& ContextLocation, float MaxPathDistance, bool bPathFromContext, TArray<NavNodeRef>& NodeRefs, FBox& NodeRefsBounds) const;
-
-	/** check if nav location is in allowed set */
-	bool IsNavLocationInPathDistance(const ARecastNavMesh* NavMesh, const FNavLocation& NavLocation, const TArray<NavNodeRef>& NodeRefs) const;
-#endif
 };
diff --git a/Engine/Source/Runtime/AIModule/Classes/EnvironmentQuery/Generators/EnvQueryGenerator_ProjectedPoints.h b/Engine/Source/Runtime/AIModule/Classes/EnvironmentQuery/Generators/EnvQueryGenerator_ProjectedPoints.h
--- a/Engine/Source/Runtime/AIModule/Classes/EnvironmentQuery/Generators/EnvQueryGenerator_ProjectedPoints.h
+++ b/Engine/Source/Runtime/AIModule/Classes/EnvironmentQuery/Generators/EnvQueryGenerator_ProjectedPoints.h
@@ -1,7 +1,7 @@
 // Copyright 1998-2015 Epic Games, Inc. All Rights Reserved.
 
 #pragma once
-#include "AI/Navigation/RecastNavMesh.h"
+#include "AI/Navigation/NavigationData.h"
 #include "EnvironmentQuery/EnvQueryGenerator.h"
 #include "EnvironmentQuery/EnvQueryTypes.h"
 #include "EnvQueryGenerator_ProjectedPoints.generated.h"
@@ -14,7 +14,7 @@ class UEnvQueryGenerator_ProjectedPoints : public UEnvQueryGenerator
 	GENERATED_UCLASS_BODY()
 
 	/** trace params */
-	UPROPERTY(EditDefaultsOnly, Category=Generator)
+	UPROPERTY(EditDefaultsOnly, Category = Generator)
 	FEnvTraceData ProjectionData;
 
 	struct FSortByHeight
@@ -30,5 +30,8 @@ class UEnvQueryGenerator_ProjectedPoints : public UEnvQueryGenerator
 	};
 
 	/** project all points in array and remove those outside navmesh */
-	void ProjectAndFilterNavPoints(TArray<FVector>& Points, const ANavigationData* NavData) const;
+	void ProjectAndFilterNavPoints(TArray<FNavLocation>& Points, FEnvQueryInstance& QueryInstance) const;
+
+	/** store points as generator's result */
+	void StoreNavPoints(const TArray<FNavLocation>& Points, FEnvQueryInstance& QueryInstance) const;
 };
diff --git a/Engine/Source/Runtime/AIModule/Classes/EnvironmentQuery/Items/EnvQueryItemType_Point.h b/Engine/Source/Runtime/AIModule/Classes/EnvironmentQuery/Items/EnvQueryItemType_Point.h
--- a/Engine/Source/Runtime/AIModule/Classes/EnvironmentQuery/Items/EnvQueryItemType_Point.h
+++ b/Engine/Source/Runtime/AIModule/Classes/EnvironmentQuery/Items/EnvQueryItemType_Point.h
@@ -13,9 +13,12 @@ class AIMODULE_API UEnvQueryItemType_Point : public UEnvQueryItemType_VectorBase
 
 	static FVector GetValue(const uint8* RawData);
 	static void SetValue(uint8* RawData, const FVector& Value);
+	static FNavLocation GetNavValue(const uint8* RawData);
+	static void SetNavValue(uint8* RawData, const FNavLocation& Value);
 
 	static void SetContextHelper(FEnvQueryContextData& ContextData, const FVector& SinglePoint);
 	static void SetContextHelper(FEnvQueryContextData& ContextData, const TArray<FVector>& MultiplePoints);
 
 	virtual FVector GetItemLocation(const uint8* RawData) const override;
+	virtual FNavLocation GetItemNavLocation(const uint8* RawData) const;
 };
diff --git a/Engine/Source/Runtime/AIModule/Classes/EnvironmentQuery/Tests/EnvQueryTest_Pathfinding.h b/Engine/Source/Runtime/AIModule/Classes/EnvironmentQuery/Tests/EnvQueryTest_Pathfinding.h
--- a/Engine/Source/Runtime/AIModule/Classes/EnvironmentQuery/Tests/EnvQueryTest_Pathfinding.h
+++ b/Engine/Source/Runtime/AIModule/Classes/EnvironmentQuery/Tests/EnvQueryTest_Pathfinding.h
@@ -38,9 +38,9 @@ class UEnvQueryTest_Pathfinding : public UEnvQueryTest
 	UPROPERTY(EditDefaultsOnly, Category=Pathfinding, AdvancedDisplay)
 	FAIDataProviderBoolValue SkipUnreachable;
 
-	/** if set, hierarchical (faster) pathfinding will be used */
-	UPROPERTY(EditDefaultsOnly, Category=Pathfinding, AdvancedDisplay)
-	FAIDataProviderBoolValue UseHierarchicalPathfinding;
+	/** navigation filter to use in pathfinding */
+	UPROPERTY(EditDefaultsOnly, Category=Pathfinding)
+	TSubclassOf<UNavigationQueryFilter> FilterClass;
 
 	// BEGIN: deprecated properties 
 	UPROPERTY()
diff --git a/Engine/Source/Runtime/AIModule/Classes/EnvironmentQuery/Tests/EnvQueryTest_PathfindingBatch.h b/Engine/Source/Runtime/AIModule/Classes/EnvironmentQuery/Tests/EnvQueryTest_PathfindingBatch.h
new file mode 100644
--- /dev/null
+++ b/Engine/Source/Runtime/AIModule/Classes/EnvironmentQuery/Tests/EnvQueryTest_PathfindingBatch.h
@@ -0,0 +1,20 @@
+// Copyright 1998-2015 Epic Games, Inc. All Rights Reserved.
+
+#pragma once
+#include "EnvironmentQuery/Tests/EnvQueryTest_Pathfinding.h"
+#include "EnvQueryTest_PathfindingBatch.generated.h"
+
+class UEnvQueryContext;
+
+UCLASS()
+class UEnvQueryTest_PathfindingBatch : public UEnvQueryTest_Pathfinding
+{
+	GENERATED_UCLASS_BODY()
+
+	/** multiplier for max distance between point and context */
+	UPROPERTY(EditDefaultsOnly, Category = Pathfinding, AdvancedDisplay)
+	FAIDataProviderFloatValue ScanRangeMultiplier;
+
+	virtual void RunTest(FEnvQueryInstance& QueryInstance) const override;
+	virtual FText GetDescriptionTitle() const override;
+};
diff --git a/Engine/Source/Runtime/AIModule/Private/EnvironmentQuery/EnvQueryInstance.cpp b/Engine/Source/Runtime/AIModule/Private/EnvironmentQuery/EnvQueryInstance.cpp
--- a/Engine/Source/Runtime/AIModule/Private/EnvironmentQuery/EnvQueryInstance.cpp
+++ b/Engine/Source/Runtime/AIModule/Private/EnvironmentQuery/EnvQueryInstance.cpp
@@ -190,6 +190,13 @@ void FEnvQueryInstance::ExecuteOneStep(double InTimeLimit)
 
 	check(IsFinished() == false);
 
+	if (!Options.IsValidIndex(OptionIndex))
+	{
+		NumValidItems = 0;
+		FinalizeQuery();
+		return;
+	}
+
 	FEnvQueryOptionInstance& OptionItem = Options[OptionIndex];
 	CONDITIONAL_SCOPE_CYCLE_COUNTER(STAT_AI_EQS_GeneratorTime, CurrentTest < 0);
 	CONDITIONAL_SCOPE_CYCLE_COUNTER(STAT_AI_EQS_TestTime, CurrentTest >= 0);
diff --git a/Engine/Source/Runtime/AIModule/Private/EnvironmentQuery/EnvQueryManager.cpp b/Engine/Source/Runtime/AIModule/Private/EnvironmentQuery/EnvQueryManager.cpp
--- a/Engine/Source/Runtime/AIModule/Private/EnvironmentQuery/EnvQueryManager.cpp
+++ b/Engine/Source/Runtime/AIModule/Private/EnvironmentQuery/EnvQueryManager.cpp
@@ -528,6 +528,7 @@ void UEnvQueryManager::CreateOptionInstance(UEnvQueryOption* OptionTemplate, con
 	OptionInstance.Generator = OptionTemplate->Generator;
 	OptionInstance.ItemType = OptionTemplate->Generator->ItemType;
 	OptionInstance.bShuffleItems = true;
+	OptionInstance.bHasNavLocations = false;
 
 	OptionInstance.Tests.AddZeroed(SortedTests.Num());
 	for (int32 TestIndex = 0; TestIndex < SortedTests.Num(); TestIndex++)
diff --git a/Engine/Source/Runtime/AIModule/Private/EnvironmentQuery/Generators/EnvQueryGenerator_Donut.cpp b/Engine/Source/Runtime/AIModule/Private/EnvironmentQuery/Generators/EnvQueryGenerator_Donut.cpp
--- a/Engine/Source/Runtime/AIModule/Private/EnvironmentQuery/Generators/EnvQueryGenerator_Donut.cpp
+++ b/Engine/Source/Runtime/AIModule/Private/EnvironmentQuery/Generators/EnvQueryGenerator_Donut.cpp
@@ -10,7 +10,6 @@
 UEnvQueryGenerator_Donut::UEnvQueryGenerator_Donut(const FObjectInitializer& ObjectInitializer) : Super(ObjectInitializer)
 {
 	Center = UEnvQueryContext_Querier::StaticClass();
-	ItemType = UEnvQueryItemType_Point::StaticClass();
 	InnerRadius.DefaultValue = 300.0f;
 	OuterRadius.DefaultValue = 1000.0f;
 	NumberOfRings.DefaultValue = 3;
@@ -108,6 +107,9 @@ void UEnvQueryGenerator_Donut::GenerateItems(FEnvQueryInstance& QueryInstance) c
 	const float AngleDelta = 2.0 * PI / NumPoints;
 	float SectionAngle = FMath::DegreesToRadians(ArcBisectDeg);
 
+	TArray<FNavLocation> Points;
+	Points.Reserve(NumPoints * NumRings);
+
 	for (int32 SectionIdx = 0; SectionIdx < NumPoints; SectionIdx++, SectionAngle += AngleDelta)
 	{
 		if (IsAngleAllowed(SectionAngle, ArcBisectDeg, ArcAngleDeg, bDefineArc))
@@ -122,11 +124,14 @@ void UEnvQueryGenerator_Donut::GenerateItems(FEnvQueryInstance& QueryInstance) c
 				for (int32 ContextIdx = 0; ContextIdx < CenterPoints.Num(); ContextIdx++)
 				{
 					const FVector PointPos = CenterPoints[ContextIdx] + RingPos;
-					QueryInstance.AddItemData<UEnvQueryItemType_Point>(PointPos);
+					Points.Add(PointPos);
 				}
 			}
 		}
 	}
+
+	ProjectAndFilterNavPoints(Points, QueryInstance);
+	StoreNavPoints(Points, QueryInstance);
 }
 
 FText UEnvQueryGenerator_Donut::GetDescriptionTitle() const
diff --git a/Engine/Source/Runtime/AIModule/Private/EnvironmentQuery/Generators/EnvQueryGenerator_OnCircle.cpp b/Engine/Source/Runtime/AIModule/Private/EnvironmentQuery/Generators/EnvQueryGenerator_OnCircle.cpp
--- a/Engine/Source/Runtime/AIModule/Private/EnvironmentQuery/Generators/EnvQueryGenerator_OnCircle.cpp
+++ b/Engine/Source/Runtime/AIModule/Private/EnvironmentQuery/Generators/EnvQueryGenerator_OnCircle.cpp
@@ -4,7 +4,6 @@
 #include "EnvironmentQuery/Contexts/EnvQueryContext_Querier.h"
 #include "EnvironmentQuery/Items/EnvQueryItemType_Point.h"
 #include "EnvironmentQuery/Generators/EnvQueryGenerator_OnCircle.h"
-#include "AI/Navigation/RecastNavMesh.h"
 
 #define LOCTEXT_NAMESPACE "EnvQueryGenerator"
 
@@ -50,45 +49,45 @@ struct FBatchTracingHelper
 	}
 
 	template<EEnvTraceShape::Type TraceType>
-	void DoSingleSourceMultiDestinations(const FVector& Source, TArray<FVector>& DestinationAndClampedResult)
+	void DoSingleSourceMultiDestinations(const FVector& Source, TArray<FNavLocation>& DestinationAndClampedResult)
 	{
 		UE_LOG(LogEQS, Error, TEXT("FBatchTracingHelper::DoSingleSourceMultiDestinations called with unhandled trace type: %d"), int32(TraceType));
 	}
 };
 
 template<>
-void FBatchTracingHelper::DoSingleSourceMultiDestinations<EEnvTraceShape::Line>(const FVector& Source, TArray<FVector>& DestinationAndClampedResult)
+void FBatchTracingHelper::DoSingleSourceMultiDestinations<EEnvTraceShape::Line>(const FVector& Source, TArray<FNavLocation>& DestinationAndClampedResult)
 {
 	for (auto& OutLocation : DestinationAndClampedResult)
 	{
-		OutLocation = RunLineTrace(Source, OutLocation);
+		OutLocation = RunLineTrace(Source, OutLocation.Location);
 	}
 }
 
 template<>
-void FBatchTracingHelper::DoSingleSourceMultiDestinations<EEnvTraceShape::Box>(const FVector& Source, TArray<FVector>& DestinationAndClampedResult)
+void FBatchTracingHelper::DoSingleSourceMultiDestinations<EEnvTraceShape::Box>(const FVector& Source, TArray<FNavLocation>& DestinationAndClampedResult)
 {
 	for (auto& OutLocation : DestinationAndClampedResult)
 	{
-		OutLocation = RunBoxTrace(Source, OutLocation);
+		OutLocation = RunBoxTrace(Source, OutLocation.Location);
 	}
 }
 
 template<>
-void FBatchTracingHelper::DoSingleSourceMultiDestinations<EEnvTraceShape::Sphere>(const FVector& Source, TArray<FVector>& DestinationAndClampedResult)
+void FBatchTracingHelper::DoSingleSourceMultiDestinations<EEnvTraceShape::Sphere>(const FVector& Source, TArray<FNavLocation>& DestinationAndClampedResult)
 {
 	for (auto& OutLocation : DestinationAndClampedResult)
 	{
-		OutLocation = RunSphereTrace(Source, OutLocation);
+		OutLocation = RunSphereTrace(Source, OutLocation.Location);
 	}
 }
 
 template<>
-void FBatchTracingHelper::DoSingleSourceMultiDestinations<EEnvTraceShape::Capsule>(const FVector& Source, TArray<FVector>& DestinationAndClampedResult)
+void FBatchTracingHelper::DoSingleSourceMultiDestinations<EEnvTraceShape::Capsule>(const FVector& Source, TArray<FNavLocation>& DestinationAndClampedResult)
 {
 	for (auto& OutLocation : DestinationAndClampedResult)
 	{
-		OutLocation = RunCapsuleTrace(Source, OutLocation);
+		OutLocation = RunCapsuleTrace(Source, OutLocation.Location);
 	}
 }
 
@@ -239,49 +238,42 @@ void UEnvQueryGenerator_OnCircle::GenerateItemsForCircle(uint8* ContextRawData,
 	const FVector& CenterLocation, const FVector& StartDirection, 
 	int32 StepsCount, float AngleStep, FEnvQueryInstance& OutQueryInstance) const
 {
-	TArray<FVector> ItemCandidates;
+	TArray<FNavLocation> ItemCandidates;
 	ItemCandidates.AddZeroed(StepsCount);
+
 	for (int32 Step = 0; Step < StepsCount; ++Step)
 	{
 		ItemCandidates[Step] = CenterLocation + StartDirection.RotateAngleAxis(AngleStep*Step, FVector::UpVector);
 	}
 
-#if WITH_RECAST
-	// @todo this needs to be optimized to batch raycasts
-	const ARecastNavMesh* NavMesh = 
-		(TraceData.TraceMode == EEnvQueryTrace::Navigation) || (ProjectionData.TraceMode == EEnvQueryTrace::Navigation) ?
-		FEQSHelpers::FindNavMeshForQuery(OutQueryInstance) : NULL;
-
-	if (NavMesh)
-	{
-		NavMesh->BeginBatchQuery();
-	}
-#endif
-
 	switch (TraceData.TraceMode)
 	{
 		case EEnvQueryTrace::Navigation:
+		{
+			ANavigationData* NavData = nullptr;
 #if WITH_RECAST
-			if (NavMesh != NULL)
+			NavData = (ANavigationData*)FEQSHelpers::FindNavMeshForQuery(OutQueryInstance);
+#endif
+			if (NavData)
 			{
-				TSharedPtr<const FNavigationQueryFilter> NavigationFilter = UNavigationQueryFilter::GetQueryFilter(NavMesh, TraceData.NavigationFilter);
+				TSharedPtr<const FNavigationQueryFilter> NavigationFilter = UNavigationQueryFilter::GetQueryFilter(NavData, TraceData.NavigationFilter);
 
 				TArray<FNavigationRaycastWork> RaycastWorkload;
 				RaycastWorkload.Reserve(ItemCandidates.Num());
 
 				for (const auto& ItemLocation : ItemCandidates)
 				{
-					RaycastWorkload.Add(FNavigationRaycastWork(CenterLocation, ItemLocation));
+					RaycastWorkload.Add(FNavigationRaycastWork(CenterLocation, ItemLocation.Location));
 				}
 
-				NavMesh->BatchRaycast(RaycastWorkload, NavigationFilter);
-			
+				NavData->BatchRaycast(RaycastWorkload, NavigationFilter);
+
 				for (int32 ItemIndex = 0; ItemIndex < ItemCandidates.Num(); ++ItemIndex)
 				{
-					ItemCandidates[ItemIndex] = RaycastWorkload[ItemIndex].HitLocation.Location;
+					ItemCandidates[ItemIndex] = RaycastWorkload[ItemIndex].HitLocation;
 				}
 			}
-#endif
+		}
 			break;
 
 		case EEnvQueryTrace::Geometry:
@@ -328,24 +320,14 @@ void UEnvQueryGenerator_OnCircle::GenerateItemsForCircle(uint8* ContextRawData,
 			break;
 	}
 
-#if WITH_RECAST
-	if (NavMesh)
-	{
-		ProjectAndFilterNavPoints(ItemCandidates, NavMesh);
-		NavMesh->FinishBatchQuery();
-	}
-#endif
-
+	ProjectAndFilterNavPoints(ItemCandidates, OutQueryInstance);
 	AddItemDataForCircle(ContextRawData, ContextItemType, ItemCandidates, OutQueryInstance);
 }
 
 void UEnvQueryGenerator_OnCircle::AddItemDataForCircle(uint8* ContextRawData, UEnvQueryItemType* ContextItemType,
-	const TArray<FVector>& Locations, FEnvQueryInstance& OutQueryInstance) const
+	const TArray<FNavLocation>& Locations, FEnvQueryInstance& OutQueryInstance) const
 {
-	for (int32 Idx = 0; Idx < Locations.Num(); Idx++)
-	{
-		OutQueryInstance.AddItemData<UEnvQueryItemType_Point>(Locations[Idx]);
-	}
+	StoreNavPoints(Locations, OutQueryInstance);
 }
 
 FText UEnvQueryGenerator_OnCircle::GetDescriptionTitle() const
@@ -418,4 +400,4 @@ void UEnvQueryGenerator_OnCircle::PostEditChangeProperty( FPropertyChangedEvent&
 
 #endif // WITH_EDITOR
 
-#undef LOCTEXT_NAMESPACE
\ No newline at end of file
+#undef LOCTEXT_NAMESPACE
diff --git a/Engine/Source/Runtime/AIModule/Private/EnvironmentQuery/Generators/EnvQueryGenerator_PathingGrid.cpp b/Engine/Source/Runtime/AIModule/Private/EnvironmentQuery/Generators/EnvQueryGenerator_PathingGrid.cpp
--- a/Engine/Source/Runtime/AIModule/Private/EnvironmentQuery/Generators/EnvQueryGenerator_PathingGrid.cpp
+++ b/Engine/Source/Runtime/AIModule/Private/EnvironmentQuery/Generators/EnvQueryGenerator_PathingGrid.cpp
@@ -1,177 +1,18 @@
 // Copyright 1998-2015 Epic Games, Inc. All Rights Reserved.
 
 #include "AIModulePrivate.h"
-#include "EnvironmentQuery/Items/EnvQueryItemType_Point.h"
-#include "EnvironmentQuery/Contexts/EnvQueryContext_Querier.h"
 #include "EnvironmentQuery/Generators/EnvQueryGenerator_PathingGrid.h"
-#include "AI/Navigation/NavFilters/RecastFilter_UseDefaultArea.h"
-#include "AI/Navigation/RecastNavMesh.h"
-
-#define LOCTEXT_NAMESPACE "EnvQueryGenerator"
 
 UEnvQueryGenerator_PathingGrid::UEnvQueryGenerator_PathingGrid(const FObjectInitializer& ObjectInitializer) : Super(ObjectInitializer)
 {
-	GenerateAround = UEnvQueryContext_Querier::StaticClass();
-	ItemType = UEnvQueryItemType_Point::StaticClass();
 	MaxDistance.DefaultValue = 100.0f;
 	SpaceBetween.DefaultValue = 10.0f;
 	PathToItem.DefaultValue = true;
-
-	// keep deprecated properties initialized
-	MaxPathDistance.Value = 100.0f;
-	Density.Value = 10.0f;
-	PathFromContext.Value = true;
 }
 
 void UEnvQueryGenerator_PathingGrid::PostLoad()
 {
-	if (VerNum < EnvQueryGeneratorVersion::DataProviders)
-	{
-		MaxPathDistance.Convert(this, MaxDistance);
-		Density.Convert(this, SpaceBetween);
-		PathFromContext.Convert(this, PathToItem);
-	}
-
 	Super::PostLoad();
-}
-
-void UEnvQueryGenerator_PathingGrid::GenerateItems(FEnvQueryInstance& QueryInstance) const
-{
-#if WITH_RECAST
-	const ARecastNavMesh* NavMesh = FEQSHelpers::FindNavMeshForQuery(QueryInstance);
-	if (NavMesh == NULL) 
-	{
-		return;
-	}
-
-	UObject* BindOwner = QueryInstance.Owner.Get();
-	MaxDistance.BindData(BindOwner, QueryInstance.QueryID);
-	SpaceBetween.BindData(BindOwner, QueryInstance.QueryID);
-	PathToItem.BindData(BindOwner, QueryInstance.QueryID);
-
-	float PathDistanceValue = MaxDistance.GetValue();
-	float DensityValue = SpaceBetween.GetValue();
-	bool bFromContextValue = PathToItem.GetValue();
-
-	const int32 ItemCount = FPlatformMath::TruncToInt((PathDistanceValue * 2.0f / DensityValue) + 1);
-	const int32 ItemCountHalf = ItemCount / 2;
 
-	TArray<FVector> ContextLocations;
-	QueryInstance.PrepareContext(GenerateAround, ContextLocations);
-	QueryInstance.ReserveItemData(ItemCountHalf * ItemCountHalf * ContextLocations.Num());
-
-	TArray<NavNodeRef> NavNodeRefs;
-	NavMesh->BeginBatchQuery();
-
-	int32 DataOffset = 0;
-	for (int32 ContextIndex = 0; ContextIndex < ContextLocations.Num(); ContextIndex++)
-	{
-		// find all node refs in pathing distance
-		FBox AllowedBounds;
-		NavNodeRefs.Reset();
-		FindNodeRefsInPathDistance(NavMesh, ContextLocations[ContextIndex], PathDistanceValue, bFromContextValue, NavNodeRefs, AllowedBounds);
-
-		// cast 2D grid on generated node refs
-		for (int32 IndexX = 0; IndexX < ItemCount; ++IndexX)
-		{
-			for (int32 IndexY = 0; IndexY < ItemCount; ++IndexY)
-			{
-				const FVector TestPoint = ContextLocations[ContextIndex] - FVector(DensityValue * (IndexX - ItemCountHalf), DensityValue * (IndexY - ItemCountHalf), 0);
-				if (!AllowedBounds.IsInsideXY(TestPoint))
-				{
-					continue;
-				}
-
-				// trace line on navmesh, and process all hits with collected node refs
-				TArray<FNavLocation> Hits;
-				NavMesh->ProjectPointMulti(TestPoint, Hits, FVector::ZeroVector, AllowedBounds.Min.Z, AllowedBounds.Max.Z);
-
-				for (int32 HitIndex = 0; HitIndex < Hits.Num(); HitIndex++)
-				{
-					if (IsNavLocationInPathDistance(NavMesh, Hits[HitIndex], NavNodeRefs))
-					{
-						// store generated point
-						QueryInstance.AddItemData<UEnvQueryItemType_Point>(Hits[HitIndex].Location);
-					}
-				}
-			}
-		}
-	}
-
-	NavMesh->FinishBatchQuery();
-#endif // WITH_RECAST
+	GridSize = MaxDistance;
 }
-
-FText UEnvQueryGenerator_PathingGrid::GetDescriptionTitle() const
-{
-	return FText::Format(LOCTEXT("DescriptionGenerateAroundContext", "{0}: generate around {1}"),
-		Super::GetDescriptionTitle(), UEnvQueryTypes::DescribeContext(GenerateAround));
-};
-
-FText UEnvQueryGenerator_PathingGrid::GetDescriptionDetails() const
-{
-	return FText::Format(LOCTEXT("DescriptionDetailsPathingGrid", "max distance: {0}, space between: {1}, path to item: {2}"),
-		FText::FromString(MaxDistance.ToString()), FText::FromString(SpaceBetween.ToString()), FText::FromString(PathToItem.ToString()));
-}
-
-#if WITH_RECAST
-#define ENVQUERY_CLUSTER_SEARCH 0
-
-void UEnvQueryGenerator_PathingGrid::FindNodeRefsInPathDistance(const ARecastNavMesh* NavMesh, const FVector& ContextLocation, float InMaxPathDistance, bool bPathFromContext, TArray<NavNodeRef>& NodeRefs, FBox& NodeRefsBounds) const
-{
-	FBox MyBounds(0);
-
-#if ENVQUERY_CLUSTER_SEARCH
-	const bool bUseBacktracking = !bPathFromContext;
-	NavMesh->GetClustersWithinPathingDistance(ContextLocation, InMaxPathDistance, NodeRefs, bUseBacktracking);
-
-	for (int32 RefIndex = 0; RefIndex < NodeRefs.Num(); RefIndex++)
-	{
-		FBox ClusterBounds;
-		
-		const bool bSuccess = NavMesh->GetClusterBounds(NodeRefs[RefIndex], ClusterBounds);
-		if (bSuccess)
-		{
-			MyBounds += ClusterBounds;
-		}
-	}
-#else
-	TSharedPtr<FNavigationQueryFilter> NavFilterInstance = NavigationFilter != NULL
-		? UNavigationQueryFilter::GetQueryFilter(NavMesh, NavigationFilter)->GetCopy()
-		: NavMesh->GetDefaultQueryFilter()->GetCopy();
-
-	NavFilterInstance->SetBacktrackingEnabled(!bPathFromContext);
-	NavMesh->GetPolysWithinPathingDistance(ContextLocation, InMaxPathDistance, NodeRefs, NavFilterInstance);
-
-	TArray<FVector> PolyVerts;
-	for (int32 RefIndex = 0; RefIndex < NodeRefs.Num(); RefIndex++)
-	{
-		PolyVerts.Reset();
-		
-		const bool bSuccess = NavMesh->GetPolyVerts(NodeRefs[RefIndex], PolyVerts);
-		if (bSuccess)
-		{
-			MyBounds += FBox(PolyVerts);
-		}
-	}
-#endif
-
-	NodeRefsBounds = MyBounds;
-}
-
-bool UEnvQueryGenerator_PathingGrid::IsNavLocationInPathDistance(const ARecastNavMesh* NavMesh,
-		const FNavLocation& NavLocation, const TArray<NavNodeRef>& NodeRefs) const
-{
-#if ENVQUERY_CLUSTER_SEARCH
-	const NavNodeRef ClusterRef = NavMesh->GetClusterRef(NavLocation.NodeRef);
-	return NodeRefs.Contains(ClusterRef);
-#else
-	return NodeRefs.Contains(NavLocation.NodeRef);
-#endif
-}
-
-#undef ENVQUERY_CLUSTER_SEARCH
-
-#endif // WITH_RECAST
-
-#undef LOCTEXT_NAMESPACE
\ No newline at end of file
diff --git a/Engine/Source/Runtime/AIModule/Private/EnvironmentQuery/Generators/EnvQueryGenerator_ProjectedPoints.cpp b/Engine/Source/Runtime/AIModule/Private/EnvironmentQuery/Generators/EnvQueryGenerator_ProjectedPoints.cpp
--- a/Engine/Source/Runtime/AIModule/Private/EnvironmentQuery/Generators/EnvQueryGenerator_ProjectedPoints.cpp
+++ b/Engine/Source/Runtime/AIModule/Private/EnvironmentQuery/Generators/EnvQueryGenerator_ProjectedPoints.cpp
@@ -1,6 +1,8 @@
 // Copyright 1998-2015 Epic Games, Inc. All Rights Reserved.
 
 #include "AIModulePrivate.h"
+#include "AI/Navigation/RecastNavMesh.h"
+#include "EnvironmentQuery/Items/EnvQueryItemType_Point.h"
 #include "EnvironmentQuery/Generators/EnvQueryGenerator_ProjectedPoints.h"
 
 UEnvQueryGenerator_ProjectedPoints::UEnvQueryGenerator_ProjectedPoints(const FObjectInitializer& ObjectInitializer) : Super(ObjectInitializer)
@@ -9,16 +11,24 @@ UEnvQueryGenerator_ProjectedPoints::UEnvQueryGenerator_ProjectedPoints(const FOb
 	ProjectionData.bCanProjectDown = true;
 	ProjectionData.bCanDisableTrace = true;
 	ProjectionData.ExtentX = 0.0f;
+
+	ItemType = UEnvQueryItemType_Point::StaticClass();
 }
 
-void UEnvQueryGenerator_ProjectedPoints::ProjectAndFilterNavPoints(TArray<FVector>& Points, const ANavigationData* NavData) const
+void UEnvQueryGenerator_ProjectedPoints::ProjectAndFilterNavPoints(TArray<FNavLocation>& Points, FEnvQueryInstance& QueryInstance) const
 {
-	const FVector VerticalOffset(0, 0, ProjectionData.PostProjectionVerticalOffset);
+	const bool bProjectToNavigation = (ProjectionData.TraceMode == EEnvQueryTrace::Navigation);
+	ANavigationData* NavData = nullptr;
+	if (bProjectToNavigation)
+	{
+#if WITH_RECAST
+		NavData = (ANavigationData*)FEQSHelpers::FindNavMeshForQuery(QueryInstance);
+#endif // WITH_RECAST
+	}
 
-	if (ProjectionData.TraceMode == EEnvQueryTrace::Navigation && NavData)
+	if (NavData)
 	{
 		TSharedPtr<const FNavigationQueryFilter> NavigationFilter = UNavigationQueryFilter::GetQueryFilter(NavData, ProjectionData.NavigationFilter);	
-		const ARecastNavMesh* NavMesh = Cast<const ARecastNavMesh>(NavData);
 		TArray<FNavigationProjectionWork> Workload;
 		Workload.Reserve(Points.Num());
 
@@ -26,32 +36,52 @@ void UEnvQueryGenerator_ProjectedPoints::ProjectAndFilterNavPoints(TArray<FVecto
 		{
 			for (const auto& Point : Points)
 			{
-				Workload.Add(FNavigationProjectionWork(Point));
+				Workload.Add(FNavigationProjectionWork(Point.Location));
 			}
 		}
 		else
 		{
-			const FVector ProjectionVerticalOffset = FVector(0, 0, (ProjectionData.ProjectUp - ProjectionData.ProjectDown) / 2);
+			const FVector VerticalOffset = FVector(0, 0, (ProjectionData.ProjectUp - ProjectionData.ProjectDown) / 2);
 			for (const auto& Point : Points)
 			{
-				Workload.Add(FNavigationProjectionWork(Point + ProjectionVerticalOffset));
+				Workload.Add(FNavigationProjectionWork(Point.Location + VerticalOffset));
 			}
 		}
 
 		const FVector ProjectionExtent(ProjectionData.ExtentX, ProjectionData.ExtentX, (ProjectionData.ProjectDown + ProjectionData.ProjectUp) / 2);
-
 		NavData->BatchProjectPoints(Workload, ProjectionExtent, NavigationFilter);
 
-		for (int32 PointIndex = Workload.Num() - 1; PointIndex >= 0; PointIndex--)
+		for (int32 Idx = Workload.Num() - 1; Idx >= 0; Idx--)
 		{
-			if (Workload[PointIndex].bResult == false)
+			if (Workload[Idx].bResult)
 			{
-				Points.RemoveAt(PointIndex);
+				Points[Idx] = Workload[Idx].OutLocation;
+				Points[Idx].Location.Z += ProjectionData.PostProjectionVerticalOffset;
 			}
 			else
 			{
-				Points[PointIndex] = Workload[PointIndex].OutLocation.Location + VerticalOffset;
+				Points.RemoveAt(Idx);
 			}
 		}
 	}
 }
+
+void UEnvQueryGenerator_ProjectedPoints::StoreNavPoints(const TArray<FNavLocation>& Points, FEnvQueryInstance& QueryInstance) const
+{
+	QueryInstance.ReserveItemData(QueryInstance.Items.Num() + Points.Num());
+	for (int32 Idx = 0; Idx < Points.Num(); Idx++)
+	{
+		// store using default function to handle creating new data entry 
+		QueryInstance.AddItemData<UEnvQueryItemType_Point>(Points[Idx].Location);
+	}
+
+	uint8* DataPtr = QueryInstance.RawData.GetData();
+	for (int32 Idx = 0; Idx < Points.Num(); Idx++)
+	{
+		// overwrite with more detailed info
+		UEnvQueryItemType_Point::SetNavValue(DataPtr + QueryInstance.Items[Idx].DataOffset, Points[Idx]);
+	}
+
+	FEnvQueryOptionInstance& OptionInstance = QueryInstance.Options[QueryInstance.OptionIndex];
+	OptionInstance.bHasNavLocations = (ProjectionData.TraceMode == EEnvQueryTrace::Navigation);
+}
diff --git a/Engine/Source/Runtime/AIModule/Private/EnvironmentQuery/Generators/EnvQueryGenerator_SimpleGrid.cpp b/Engine/Source/Runtime/AIModule/Private/EnvironmentQuery/Generators/EnvQueryGenerator_SimpleGrid.cpp
--- a/Engine/Source/Runtime/AIModule/Private/EnvironmentQuery/Generators/EnvQueryGenerator_SimpleGrid.cpp
+++ b/Engine/Source/Runtime/AIModule/Private/EnvironmentQuery/Generators/EnvQueryGenerator_SimpleGrid.cpp
@@ -2,16 +2,13 @@
 
 #include "AIModulePrivate.h"
 #include "EnvironmentQuery/Contexts/EnvQueryContext_Querier.h"
-#include "EnvironmentQuery/Items/EnvQueryItemType_Point.h"
 #include "EnvironmentQuery/Generators/EnvQueryGenerator_SimpleGrid.h"
-#include "AI/Navigation/RecastNavMesh.h"
 
 #define LOCTEXT_NAMESPACE "EnvQueryGenerator"
 
 UEnvQueryGenerator_SimpleGrid::UEnvQueryGenerator_SimpleGrid(const FObjectInitializer& ObjectInitializer) : Super(ObjectInitializer)
 {
 	GenerateAround = UEnvQueryContext_Querier::StaticClass();
-	ItemType = UEnvQueryItemType_Point::StaticClass();
 	GridSize.DefaultValue = 500.0f;
 	SpaceBetween.DefaultValue = 100.0f;
 
@@ -40,30 +37,13 @@ void UEnvQueryGenerator_SimpleGrid::GenerateItems(FEnvQueryInstance& QueryInstan
 	float RadiusValue = GridSize.GetValue();
 	float DensityValue = SpaceBetween.GetValue();
 
-	const bool bProjectToNavigation = (ProjectionData.TraceMode == EEnvQueryTrace::Navigation);
-#if WITH_RECAST
-	const ARecastNavMesh* NavMesh = bProjectToNavigation ? FEQSHelpers::FindNavMeshForQuery(QueryInstance) : NULL;
-	if (bProjectToNavigation && NavMesh == NULL) 
-	{
-		return;
-	}
-#endif // WITH_RECAST
-
 	const int32 ItemCount = FPlatformMath::TruncToInt((RadiusValue * 2.0f / DensityValue) + 1);
 	const int32 ItemCountHalf = ItemCount / 2;
 
 	TArray<FVector> ContextLocations;
 	QueryInstance.PrepareContext(GenerateAround, ContextLocations);
 
-#if WITH_RECAST
-	const FVector ProjectExtent(ProjectionData.ExtentX, ProjectionData.ExtentX, 0);
-	if (NavMesh)
-	{
-		NavMesh->BeginBatchQuery();
-	}
-#endif // WITH_RECAST
-
-	TArray<FVector> GridPoints;
+	TArray<FNavLocation> GridPoints;
 	GridPoints.Reserve(ItemCount * ItemCount * ContextLocations.Num());
 
 	for (int32 ContextIndex = 0; ContextIndex < ContextLocations.Num(); ContextIndex++)
@@ -78,19 +58,8 @@ void UEnvQueryGenerator_SimpleGrid::GenerateItems(FEnvQueryInstance& QueryInstan
 		}
 	}
 
-#if WITH_RECAST
-	if (NavMesh)
-	{
-		ProjectAndFilterNavPoints(GridPoints, NavMesh);
-		NavMesh->FinishBatchQuery();
-	}
-#endif // WITH_RECAST
-
-	QueryInstance.ReserveItemData(GridPoints.Num());
-	for (int32 PointIndex = 0; PointIndex < GridPoints.Num(); PointIndex++)
-	{
-		QueryInstance.AddItemData<UEnvQueryItemType_Point>(GridPoints[PointIndex]);
-	}
+	ProjectAndFilterNavPoints(GridPoints, QueryInstance);
+	StoreNavPoints(GridPoints, QueryInstance);
 }
 
 FText UEnvQueryGenerator_SimpleGrid::GetDescriptionTitle() const
@@ -116,4 +85,4 @@ FText UEnvQueryGenerator_SimpleGrid::GetDescriptionDetails() const
 	return Desc;
 }
 
-#undef LOCTEXT_NAMESPACE
\ No newline at end of file
+#undef LOCTEXT_NAMESPACE
diff --git a/Engine/Source/Runtime/AIModule/Private/EnvironmentQuery/Items/EnvQueryItemType_Point.cpp b/Engine/Source/Runtime/AIModule/Private/EnvironmentQuery/Items/EnvQueryItemType_Point.cpp
--- a/Engine/Source/Runtime/AIModule/Private/EnvironmentQuery/Items/EnvQueryItemType_Point.cpp
+++ b/Engine/Source/Runtime/AIModule/Private/EnvironmentQuery/Items/EnvQueryItemType_Point.cpp
@@ -5,17 +5,27 @@
 
 UEnvQueryItemType_Point::UEnvQueryItemType_Point(const FObjectInitializer& ObjectInitializer) : Super(ObjectInitializer)
 {
-	ValueSize = sizeof(FVector);
+	ValueSize = sizeof(FNavLocation);
 }
 
 FVector UEnvQueryItemType_Point::GetValue(const uint8* RawData)
 {
-	return GetValueFromMemory<FVector>(RawData);
+	return GetValueFromMemory<FNavLocation>(RawData).Location;
 }
 
 void UEnvQueryItemType_Point::SetValue(uint8* RawData, const FVector& Value)
 {
-	return SetValueInMemory<FVector>(RawData, Value);
+	return SetValueInMemory<FNavLocation>(RawData, FNavLocation(Value));
+}
+
+FNavLocation UEnvQueryItemType_Point::GetNavValue(const uint8* RawData)
+{
+	return GetValueFromMemory<FNavLocation>(RawData);
+}
+
+void UEnvQueryItemType_Point::SetNavValue(uint8* RawData, const FNavLocation& Value)
+{
+	return SetValueInMemory<FNavLocation>(RawData, Value);
 }
 
 FVector UEnvQueryItemType_Point::GetItemLocation(const uint8* RawData) const
@@ -23,11 +33,16 @@ FVector UEnvQueryItemType_Point::GetItemLocation(const uint8* RawData) const
 	return UEnvQueryItemType_Point::GetValue(RawData);
 }
 
+FNavLocation UEnvQueryItemType_Point::GetItemNavLocation(const uint8* RawData) const
+{
+	return UEnvQueryItemType_Point::GetNavValue(RawData);
+}
+
 void UEnvQueryItemType_Point::SetContextHelper(FEnvQueryContextData& ContextData, const FVector& SinglePoint)
 {
 	ContextData.ValueType = UEnvQueryItemType_Point::StaticClass();
 	ContextData.NumValues = 1;
-	ContextData.RawData.Init(sizeof(FVector));
+	ContextData.RawData.Init(sizeof(FNavLocation));
 
 	UEnvQueryItemType_Point::SetValue((uint8*)ContextData.RawData.GetData(), SinglePoint);
 }
@@ -36,12 +51,12 @@ void UEnvQueryItemType_Point::SetContextHelper(FEnvQueryContextData& ContextData
 {
 	ContextData.ValueType = UEnvQueryItemType_Point::StaticClass();
 	ContextData.NumValues = MultiplePoints.Num();
-	ContextData.RawData.Init(sizeof(FVector) * MultiplePoints.Num());
+	ContextData.RawData.Init(sizeof(FNavLocation)* MultiplePoints.Num());
 
 	uint8* RawData = (uint8*)ContextData.RawData.GetData();
 	for (int32 PointIndex = 0; PointIndex < MultiplePoints.Num(); PointIndex++)
 	{
 		UEnvQueryItemType_Point::SetValue(RawData, MultiplePoints[PointIndex]);
-		RawData += sizeof(FVector);
+		RawData += sizeof(FNavLocation);
 	}
 }
diff --git a/Engine/Source/Runtime/AIModule/Private/EnvironmentQuery/Tests/EnvQueryTest_Pathfinding.cpp b/Engine/Source/Runtime/AIModule/Private/EnvironmentQuery/Tests/EnvQueryTest_Pathfinding.cpp
--- a/Engine/Source/Runtime/AIModule/Private/EnvironmentQuery/Tests/EnvQueryTest_Pathfinding.cpp
+++ b/Engine/Source/Runtime/AIModule/Private/EnvironmentQuery/Tests/EnvQueryTest_Pathfinding.cpp
@@ -17,14 +17,12 @@ UEnvQueryTest_Pathfinding::UEnvQueryTest_Pathfinding(const FObjectInitializer& O
 	TestMode = EEnvTestPathfinding::PathExist;
 	PathFromContext.DefaultValue = true;
 	SkipUnreachable.DefaultValue = true;
-	UseHierarchicalPathfinding.DefaultValue = true;
 	FloatValueMin.DefaultValue = 1000.0f;
 	FloatValueMax.DefaultValue = 1000.0f;
 
 	// keep deprecated properties initialized
 	PathToItem.Value = true;
 	DiscardUnreachable.Value = true;
-	HierarchicalPathfinding.Value = true;
 	FloatFilterMin.Value = 1000.0f;
 	FloatFilterMax.Value = 1000.0f;
 }
@@ -35,13 +33,11 @@ void UEnvQueryTest_Pathfinding::RunTest(FEnvQueryInstance& QueryInstance) const
 	BoolValue.BindData(DataOwner, QueryInstance.QueryID);
 	PathFromContext.BindData(DataOwner, QueryInstance.QueryID);
 	SkipUnreachable.BindData(DataOwner, QueryInstance.QueryID);
-	UseHierarchicalPathfinding.BindData(DataOwner, QueryInstance.QueryID);
 	FloatValueMin.BindData(DataOwner, QueryInstance.QueryID);
 	FloatValueMax.BindData(DataOwner, QueryInstance.QueryID);
 
 	bool bWantsPath = BoolValue.GetValue();
 	bool bPathToItem = PathFromContext.GetValue();
-	bool bHierarchical = UseHierarchicalPathfinding.GetValue();
 	bool bDiscardFailed = SkipUnreachable.GetValue();
 	float MinThresholdValue = FloatValueMin.GetValue();
 	float MaxThresholdValue = FloatValueMax.GetValue();
@@ -59,7 +55,7 @@ void UEnvQueryTest_Pathfinding::RunTest(FEnvQueryInstance& QueryInstance) const
 		return;
 	}
 
-	EPathFindingMode::Type PFMode(bHierarchical ? EPathFindingMode::Hierarchical : EPathFindingMode::Regular);
+	EPathFindingMode::Type PFMode(EPathFindingMode::Regular);
 
 	if (GetWorkOnFloatValues())
 	{
@@ -129,17 +125,6 @@ FText UEnvQueryTest_Pathfinding::GetDescriptionTitle() const
 
 FText UEnvQueryTest_Pathfinding::GetDescriptionDetails() const
 {
-	FText HPathDesc = LOCTEXT("HierarchicalPathfinding", "hierarchical pathfinding");
-	FText Desc1;
-	if (UseHierarchicalPathfinding.IsDynamic())
-	{
-		Desc1 = FText::Format(FText::FromString("{0}: {1}"), HPathDesc, FText::FromString(UseHierarchicalPathfinding.ToString()));
-	}
-	else if (UseHierarchicalPathfinding.DefaultValue)
-	{
-		Desc1 = HPathDesc;
-	}
-
 	FText DiscardDesc = LOCTEXT("DiscardUnreachable", "discard unreachable");
 	FText Desc2;
 	if (SkipUnreachable.IsDynamic())
@@ -152,15 +137,7 @@ FText UEnvQueryTest_Pathfinding::GetDescriptionDetails() const
 	}
 
 	FText TestParamDesc = GetWorkOnFloatValues() ? DescribeFloatTestParams() : DescribeBoolTestParams("existing path");
-	if (!Desc1.IsEmpty() && !Desc2.IsEmpty())
-	{
-		return FText::Format(FText::FromString("{0}, {1}\n{2}"), Desc1, Desc2, TestParamDesc);
-	}
-	else if (!Desc1.IsEmpty())
-	{
-		return FText::Format(FText::FromString("{0}\n{1}"), Desc1, TestParamDesc);
-	}
-	else if (!Desc2.IsEmpty())
+	if (!Desc2.IsEmpty())
 	{
 		return FText::Format(FText::FromString("{0}\n{1}"), Desc2, TestParamDesc);
 	}
@@ -185,7 +162,6 @@ void UEnvQueryTest_Pathfinding::PostLoad()
 	{
 		PathToItem.Convert(this, PathFromContext);
 		DiscardUnreachable.Convert(this, SkipUnreachable);
-		HierarchicalPathfinding.Convert(this, UseHierarchicalPathfinding);
 	}
 
 	Super::PostLoad();
@@ -195,37 +171,55 @@ void UEnvQueryTest_Pathfinding::PostLoad()
 
 bool UEnvQueryTest_Pathfinding::TestPathFrom(const FVector& ItemPos, const FVector& ContextPos, EPathFindingMode::Type Mode, const ANavigationData* NavData, UNavigationSystem* NavSys, const UObject* PathOwner) const
 {
-	const bool bPathExists = NavSys->TestPathSync(FPathFindingQuery(PathOwner, NavData, ItemPos, ContextPos).SetAllowPartialPaths(false), Mode);
+	FPathFindingQuery Query(PathOwner, NavData, ItemPos, ContextPos, UNavigationQueryFilter::GetQueryFilter(NavData, FilterClass));
+	Query.SetAllowPartialPaths(false);
+
+	const bool bPathExists = NavSys->TestPathSync(Query, Mode);
 	return bPathExists;
 }
 
 bool UEnvQueryTest_Pathfinding::TestPathTo(const FVector& ItemPos, const FVector& ContextPos, EPathFindingMode::Type Mode, const ANavigationData* NavData, UNavigationSystem* NavSys, const UObject* PathOwner) const
 {
-	const bool bPathExists = NavSys->TestPathSync(FPathFindingQuery(PathOwner, NavData, ContextPos, ItemPos).SetAllowPartialPaths(false), Mode);
+	FPathFindingQuery Query(PathOwner, NavData, ContextPos, ItemPos, UNavigationQueryFilter::GetQueryFilter(NavData, FilterClass));
+	Query.SetAllowPartialPaths(false);
+
+	const bool bPathExists = NavSys->TestPathSync(Query, Mode);
 	return bPathExists;
 }
 
 float UEnvQueryTest_Pathfinding::FindPathCostFrom(const FVector& ItemPos, const FVector& ContextPos, EPathFindingMode::Type Mode, const ANavigationData* NavData, UNavigationSystem* NavSys, const UObject* PathOwner) const
 {
-	FPathFindingResult Result = NavSys->FindPathSync(FPathFindingQuery(PathOwner, NavData, ItemPos, ContextPos).SetAllowPartialPaths(false), Mode);
+	FPathFindingQuery Query(PathOwner, NavData, ItemPos, ContextPos, UNavigationQueryFilter::GetQueryFilter(NavData, FilterClass));
+	Query.SetAllowPartialPaths(false);
+
+	FPathFindingResult Result = NavSys->FindPathSync(Query, Mode);
 	return (Result.IsSuccessful()) ? Result.Path->GetCost() : BIG_NUMBER;
 }
 
 float UEnvQueryTest_Pathfinding::FindPathCostTo(const FVector& ItemPos, const FVector& ContextPos, EPathFindingMode::Type Mode, const ANavigationData* NavData, UNavigationSystem* NavSys, const UObject* PathOwner) const
 {
-	FPathFindingResult Result = NavSys->FindPathSync(FPathFindingQuery(PathOwner, NavData, ContextPos, ItemPos).SetAllowPartialPaths(false), Mode);
+	FPathFindingQuery Query(PathOwner, NavData, ContextPos, ItemPos, UNavigationQueryFilter::GetQueryFilter(NavData, FilterClass));
+	Query.SetAllowPartialPaths(false);
+
+	FPathFindingResult Result = NavSys->FindPathSync(Query, Mode);
 	return (Result.IsSuccessful()) ? Result.Path->GetCost() : BIG_NUMBER;
 }
 
 float UEnvQueryTest_Pathfinding::FindPathLengthFrom(const FVector& ItemPos, const FVector& ContextPos, EPathFindingMode::Type Mode, const ANavigationData* NavData, UNavigationSystem* NavSys, const UObject* PathOwner) const
 {
-	FPathFindingResult Result = NavSys->FindPathSync(FPathFindingQuery(PathOwner, NavData, ItemPos, ContextPos).SetAllowPartialPaths(false), Mode);
+	FPathFindingQuery Query(PathOwner, NavData, ItemPos, ContextPos, UNavigationQueryFilter::GetQueryFilter(NavData, FilterClass));
+	Query.SetAllowPartialPaths(false);
+
+	FPathFindingResult Result = NavSys->FindPathSync(Query, Mode);
 	return (Result.IsSuccessful()) ? Result.Path->GetLength() : BIG_NUMBER;
 }
 
 float UEnvQueryTest_Pathfinding::FindPathLengthTo(const FVector& ItemPos, const FVector& ContextPos, EPathFindingMode::Type Mode, const ANavigationData* NavData, UNavigationSystem* NavSys, const UObject* PathOwner) const
 {
-	FPathFindingResult Result = NavSys->FindPathSync(FPathFindingQuery(PathOwner, NavData, ContextPos, ItemPos).SetAllowPartialPaths(false), Mode);
+	FPathFindingQuery Query(PathOwner, NavData, ContextPos, ItemPos, UNavigationQueryFilter::GetQueryFilter(NavData, FilterClass));
+	Query.SetAllowPartialPaths(false);
+
+	FPathFindingResult Result = NavSys->FindPathSync(Query, Mode);
 	return (Result.IsSuccessful()) ? Result.Path->GetLength() : BIG_NUMBER;
 }
 
diff --git a/Engine/Source/Runtime/AIModule/Private/EnvironmentQuery/Tests/EnvQueryTest_PathfindingBatch.cpp b/Engine/Source/Runtime/AIModule/Private/EnvironmentQuery/Tests/EnvQueryTest_PathfindingBatch.cpp
new file mode 100644
--- /dev/null
+++ b/Engine/Source/Runtime/AIModule/Private/EnvironmentQuery/Tests/EnvQueryTest_PathfindingBatch.cpp
@@ -0,0 +1,212 @@
+// Copyright 1998-2015 Epic Games, Inc. All Rights Reserved.
+
+#include "AIModulePrivate.h"
+#include "AI/Navigation/NavigationSystem.h"
+#include "AI/Navigation/NavAgentInterface.h"
+#include "AI/Navigation/RecastNavMesh.h"
+#include "EnvironmentQuery/Contexts/EnvQueryContext_Querier.h"
+#include "EnvironmentQuery/Items/EnvQueryItemType_Point.h"
+#include "EnvironmentQuery/Tests/EnvQueryTest_PathfindingBatch.h"
+
+#if WITH_RECAST
+#include "AI/Navigation/PImplRecastNavMesh.h"
+#endif
+
+#define LOCTEXT_NAMESPACE "EnvQueryGenerator"
+
+UEnvQueryTest_PathfindingBatch::UEnvQueryTest_PathfindingBatch(const FObjectInitializer& ObjectInitializer) : Super(ObjectInitializer)
+{
+	ScanRangeMultiplier.DefaultValue = 1.5f;
+}
+
+#if WITH_RECAST
+namespace NodePoolHelpers
+{
+	static bool HasPath(const FRecastDebugPathfindingData& NodePool, const FNavigationProjectionWork& TestPt)
+	{
+		FRecastDebugPathfindingNode SearchKey(TestPt.OutLocation.NodeRef);
+		const FRecastDebugPathfindingNode* MyNode = NodePool.Nodes.Find(SearchKey);
+		return MyNode != nullptr;
+	}
+
+	static float GetPathLength(const FRecastDebugPathfindingData& NodePool, const FNavigationProjectionWork& TestPt, const dtQueryFilter* Filter)
+	{
+		FRecastDebugPathfindingNode SearchKey(TestPt.OutLocation.NodeRef);
+		const FRecastDebugPathfindingNode* MyNode = NodePool.Nodes.Find(SearchKey);
+		if (MyNode)
+		{
+			float LastSegmentLength = FVector::Dist(MyNode->NodePos, TestPt.OutLocation.Location);
+			return MyNode->Length + LastSegmentLength;
+		}
+
+		return BIG_NUMBER;
+	}
+
+	static float GetPathCost(const FRecastDebugPathfindingData& NodePool, const FNavigationProjectionWork& TestPt, const dtQueryFilter* Filter)
+	{
+		FRecastDebugPathfindingNode SearchKey(TestPt.OutLocation.NodeRef);
+		const FRecastDebugPathfindingNode* MyNode = NodePool.Nodes.Find(SearchKey);
+		if (MyNode)
+		{
+			return MyNode->TotalCost;
+		}
+
+		return BIG_NUMBER;
+	}
+
+	typedef float(*PathParamFunc)(const FRecastDebugPathfindingData&, const FNavigationProjectionWork&, const dtQueryFilter*);
+}
+
+DECLARE_CYCLE_STAT(TEXT("Debug test total"), STAT_AI_EQS_Debug1, STATGROUP_AI_EQS);
+
+void UEnvQueryTest_PathfindingBatch::RunTest(FEnvQueryInstance& QueryInstance) const
+{
+	SCOPE_CYCLE_COUNTER(STAT_AI_EQS_Debug1);
+
+	UObject* DataOwner = QueryInstance.Owner.Get();
+	BoolValue.BindData(DataOwner, QueryInstance.QueryID);
+	PathFromContext.BindData(DataOwner, QueryInstance.QueryID);
+	SkipUnreachable.BindData(DataOwner, QueryInstance.QueryID);
+	FloatValueMin.BindData(DataOwner, QueryInstance.QueryID);
+	FloatValueMax.BindData(DataOwner, QueryInstance.QueryID);
+	ScanRangeMultiplier.BindData(DataOwner, QueryInstance.QueryID);
+
+	bool bWantsPath = BoolValue.GetValue();
+	bool bPathToItem = PathFromContext.GetValue();
+	bool bDiscardFailed = SkipUnreachable.GetValue();
+	float MinThresholdValue = FloatValueMin.GetValue();
+	float MaxThresholdValue = FloatValueMax.GetValue();
+	float RangeMultiplierValue = ScanRangeMultiplier.GetValue();
+
+	UNavigationSystem* NavSys = QueryInstance.World->GetNavigationSystem();
+	ANavigationData* NavData = FindNavigationData(NavSys, QueryInstance.Owner.Get());
+	ARecastNavMesh* NavMeshData = Cast<ARecastNavMesh>(NavData);
+	if (!NavMeshData)
+	{
+		return;
+	}
+
+	TArray<FVector> ContextLocations;
+	if (!QueryInstance.PrepareContext(Context, ContextLocations))
+	{
+		return;
+	}
+
+	TArray<FNavigationProjectionWork> TestPoints;
+	TArray<float> CollectDistanceSq;
+	CollectDistanceSq.Init(0.0f, ContextLocations.Num());
+
+	TSharedPtr<FNavigationQueryFilter> NavigationFilter = FilterClass != nullptr
+		? UNavigationQueryFilter::GetQueryFilter(NavMeshData, FilterClass)->GetCopy()
+		: NavMeshData->GetDefaultQueryFilter()->GetCopy();
+	NavigationFilter->SetBacktrackingEnabled(!bPathToItem);
+	const dtQueryFilter* NavQueryFilter = ((const FRecastQueryFilter*)NavigationFilter->GetImplementation())->GetAsDetourQueryFilter();
+
+	UEnvQueryItemType_Point* PointCDO = QueryInstance.ItemType->GetDefaultObject<UEnvQueryItemType_Point>();
+	if (QueryInstance.Options[QueryInstance.OptionIndex].bHasNavLocations && PointCDO)
+	{
+		for (int32 ItemIdx = 0; ItemIdx < QueryInstance.Items.Num(); ItemIdx++)
+		{
+			if (QueryInstance.Items[ItemIdx].IsValid())
+			{
+				const FNavLocation ItemNavLocation = PointCDO->GetItemNavLocation(QueryInstance.RawData.GetData() + QueryInstance.Items[ItemIdx].DataOffset);
+
+				FNavigationProjectionWork ProjData(ItemNavLocation.Location);
+				ProjData.OutLocation = ItemNavLocation;
+				ProjData.bResult = true;
+				TestPoints.Add(ProjData);
+
+				for (int32 ContextIdx = 0; ContextIdx < ContextLocations.Num(); ContextIdx++)
+				{
+					const float TestDistanceSq = FVector::DistSquared(ItemNavLocation.Location, ContextLocations[ContextIdx]);
+					CollectDistanceSq[ContextIdx] = FMath::Max(CollectDistanceSq[ContextIdx], TestDistanceSq);
+				}
+			}
+		}
+	}
+	else
+	{
+		// can't use FEnvQueryInstance::ItemIterator yet, since it has built in scoring functionality
+		for (int32 ItemIdx = 0; ItemIdx < QueryInstance.Items.Num(); ItemIdx++)
+		{
+			if (QueryInstance.Items[ItemIdx].IsValid())
+			{
+				const FVector ItemLocation = GetItemLocation(QueryInstance, ItemIdx);
+				TestPoints.Add(FNavigationProjectionWork(ItemLocation));
+
+				for (int32 ContextIdx = 0; ContextIdx < ContextLocations.Num(); ContextIdx++)
+				{
+					const float TestDistanceSq = FVector::DistSquared(ItemLocation, ContextLocations[ContextIdx]);
+					CollectDistanceSq[ContextIdx] = FMath::Max(CollectDistanceSq[ContextIdx], TestDistanceSq);
+				}
+			}
+		}
+
+		NavMeshData->BatchProjectPoints(TestPoints, NavMeshData->GetDefaultQueryExtent(), NavigationFilter);
+	}
+
+	TArray<FRecastDebugPathfindingData> NodePoolData;
+	NodePoolData.SetNum(ContextLocations.Num());
+
+	{
+		// scope for perf timer
+
+		TArray<NavNodeRef> Polys;
+		for (int32 ContextIdx = 0; ContextIdx < ContextLocations.Num(); ContextIdx++)
+		{
+			const float MaxPathDistance = FMath::Sqrt(CollectDistanceSq[ContextIdx]) * RangeMultiplierValue;
+
+			Polys.Reset();
+			NodePoolData[ContextIdx].Flags = ERecastDebugPathfindingFlags::PathLength;
+
+			NavMeshData->GetPolysWithinPathingDistance(ContextLocations[ContextIdx], MaxPathDistance, Polys, NavigationFilter, nullptr, &NodePoolData[ContextIdx]);
+		}
+	}
+
+	if (GetWorkOnFloatValues())
+	{
+		NodePoolHelpers::PathParamFunc Func[] = { nullptr, NodePoolHelpers::GetPathCost, NodePoolHelpers::GetPathLength };
+		FEnvQueryInstance::ItemIterator It(this, QueryInstance);
+		for (It.IgnoreTimeLimit(); It; ++It)
+		{
+			for (int32 ContextIndex = 0; ContextIndex < ContextLocations.Num(); ContextIndex++)
+			{
+				const float PathValue = Func[TestMode](NodePoolData[ContextIndex], TestPoints[*It], NavQueryFilter);
+				It.SetScore(TestPurpose, FilterType, PathValue, MinThresholdValue, MaxThresholdValue);
+
+				if (bDiscardFailed && PathValue >= BIG_NUMBER)
+				{
+					It.DiscardItem();
+				}
+			}
+		}
+	}
+	else
+	{
+		FEnvQueryInstance::ItemIterator It(this, QueryInstance);
+		for (It.IgnoreTimeLimit(); It; ++It)
+		{
+			for (int32 ContextIndex = 0; ContextIndex < ContextLocations.Num(); ContextIndex++)
+			{
+				const bool bFoundPath = NodePoolHelpers::HasPath(NodePoolData[ContextIndex], TestPoints[*It]);
+				It.SetScore(TestPurpose, FilterType, bFoundPath, bWantsPath);
+			}
+		}
+	}
+}
+
+#else
+
+void UEnvQueryTest_PathfindingBatch::RunTest(FEnvQueryInstance& QueryInstance) const
+{
+	// can't do anything without navmesh
+}
+
+#endif
+
+FText UEnvQueryTest_PathfindingBatch::GetDescriptionTitle() const
+{
+	return FText::Format(LOCTEXT("PathfindingBatch","{0} [batch]"), Super::GetDescriptionTitle());
+}
+
+#undef LOCTEXT_NAMESPACE
diff --git a/Engine/Source/Runtime/Engine/Classes/AI/Navigation/NavigationTestingActor.h b/Engine/Source/Runtime/Engine/Classes/AI/Navigation/NavigationTestingActor.h
--- a/Engine/Source/Runtime/Engine/Classes/AI/Navigation/NavigationTestingActor.h
+++ b/Engine/Source/Runtime/Engine/Classes/AI/Navigation/NavigationTestingActor.h
@@ -138,7 +138,7 @@ public:
 
 #if WITH_RECAST && WITH_EDITORONLY_DATA
 	/** detail data gathered from each step of regular A* algorithm */
-	TArray<struct FRecastDebugPathfindingStep> DebugSteps;
+	TArray<struct FRecastDebugPathfindingData> DebugSteps;
 
 	FNavTestTickHelper* TickHelper;
 #endif
diff --git a/Engine/Source/Runtime/Engine/Classes/AI/Navigation/RecastNavMesh.h b/Engine/Source/Runtime/Engine/Classes/AI/Navigation/RecastNavMesh.h
--- a/Engine/Source/Runtime/Engine/Classes/AI/Navigation/RecastNavMesh.h
+++ b/Engine/Source/Runtime/Engine/Classes/AI/Navigation/RecastNavMesh.h
@@ -219,6 +219,7 @@ struct FRecastDebugPathfindingNode
 	NavNodeRef ParentRef;
 	float Cost;
 	float TotalCost;
+	float Length;
 	uint32 bOpenSet : 1;
 	uint32 bOffMeshLink : 1;
 	uint32 bModified : 1;
@@ -226,16 +227,35 @@ struct FRecastDebugPathfindingNode
 	FVector NodePos;
 	TArray<FVector> Verts;
 
+	FRecastDebugPathfindingNode() : PolyRef(0), ParentRef(0) {}
+	FRecastDebugPathfindingNode(NavNodeRef InPolyRef) : PolyRef(InPolyRef), ParentRef(0) {}
+
+	FORCEINLINE bool operator==(const NavNodeRef& OtherPolyRef) const { return PolyRef == OtherPolyRef; }
 	FORCEINLINE bool operator==(const FRecastDebugPathfindingNode& Other) const { return PolyRef == Other.PolyRef; }
 	FORCEINLINE friend uint32 GetTypeHash(const FRecastDebugPathfindingNode& Other) { return Other.PolyRef; }
 
 	FORCEINLINE float GetHeuristicCost() const { return TotalCost - Cost; }
 };
 
-struct FRecastDebugPathfindingStep
+namespace ERecastDebugPathfindingFlags
+{
+	enum Type
+	{
+		Basic = 0x0,
+		BestNode = 0x1,
+		Vertices = 0x2,
+		PathLength = 0x4
+	};
+}
+
+struct FRecastDebugPathfindingData
 {
 	TSet<FRecastDebugPathfindingNode> Nodes;
 	FSetElementId BestNode;
+	uint8 Flags;
+
+	FRecastDebugPathfindingData() : Flags(ERecastDebugPathfindingFlags::Basic) {}
+	FRecastDebugPathfindingData(ERecastDebugPathfindingFlags::Type InFlags) : Flags(InFlags) {}
 };
 
 struct FRecastDebugGeometry
@@ -845,7 +865,8 @@ public:
 	 *	@NOTE query is not using string-pulled path distance (for performance reasons),
 	 *		it measured distance between middles of portal edges, do you might want to 
 	 *		add an extra margin to PathingDistance */
-	bool GetPolysWithinPathingDistance(FVector const& StartLoc, const float PathingDistance, TArray<NavNodeRef>& FoundPolys, TSharedPtr<const FNavigationQueryFilter> Filter = NULL, const UObject* Querier = NULL) const;
+	bool GetPolysWithinPathingDistance(FVector const& StartLoc, const float PathingDistance, TArray<NavNodeRef>& FoundPolys,
+		TSharedPtr<const FNavigationQueryFilter> Filter = nullptr, const UObject* Querier = nullptr, FRecastDebugPathfindingData* DebugData = nullptr) const;
 
 	/** Filters nav polys in PolyRefs with Filter */
 	bool FilterPolys(TArray<NavNodeRef>& PolyRefs, const FRecastQueryFilter* Filter, const UObject* Querier = NULL) const;
@@ -881,7 +902,7 @@ public:
 	bool FindStraightPath(const FVector& StartLoc, const FVector& EndLoc, const TArray<NavNodeRef>& PathCorridor, TArray<FNavPathPoint>& PathPoints, TArray<uint32>* CustomLinks = NULL) const;
 
 	/** Runs A* pathfinding on navmesh and collect data for every step */
-	int32 DebugPathfinding(const FPathFindingQuery& Query, TArray<FRecastDebugPathfindingStep>& Steps);
+	int32 DebugPathfinding(const FPathFindingQuery& Query, TArray<FRecastDebugPathfindingData>& Steps);
 
 	static const FRecastQueryFilter* GetNamedFilter(ERecastNamedFilter::Type FilterType);
 	FORCEINLINE static FNavPolyFlags GetNavLinkFlag() { return NavLinkFlag; }
diff --git a/Engine/Source/Runtime/Engine/Private/AI/Navigation/NavTestRenderingComponent.cpp b/Engine/Source/Runtime/Engine/Private/AI/Navigation/NavTestRenderingComponent.cpp
--- a/Engine/Source/Runtime/Engine/Private/AI/Navigation/NavTestRenderingComponent.cpp
+++ b/Engine/Source/Runtime/Engine/Private/AI/Navigation/NavTestRenderingComponent.cpp
@@ -211,7 +211,7 @@ public:
 		if (NavTestActor && NavTestActor->DebugSteps.Num() && NavTestActor->ShowStepIndex >= 0)
 		{
 			const int32 ShowIdx = FMath::Min(NavTestActor->ShowStepIndex, NavTestActor->DebugSteps.Num() - 1);
-			const FRecastDebugPathfindingStep& DebugStep = NavTestActor->DebugSteps[ShowIdx];
+			const FRecastDebugPathfindingData& DebugStep = NavTestActor->DebugSteps[ShowIdx];
 			int32 BaseOpen = 0;
 			int32 BaseClosed = 0;
 
@@ -451,7 +451,7 @@ FBoxSphereBounds UNavTestRenderingComponent::CalcBounds(const FTransform& LocalT
 		if (TestActor->DebugSteps.Num() && TestActor->ShowStepIndex >= 0)
 		{
 			const int32 ShowIdx = FMath::Min(TestActor->ShowStepIndex, TestActor->DebugSteps.Num() - 1);
-			const FRecastDebugPathfindingStep& DebugStep = TestActor->DebugSteps[ShowIdx];
+			const FRecastDebugPathfindingData& DebugStep = TestActor->DebugSteps[ShowIdx];
 			for (TSet<FRecastDebugPathfindingNode>::TConstIterator It(DebugStep.Nodes); It; ++It)
 			{
 				const FRecastDebugPathfindingNode& DebugNode = *It;
diff --git a/Engine/Source/Runtime/Engine/Private/AI/Navigation/PImplRecastNavMesh.cpp b/Engine/Source/Runtime/Engine/Private/AI/Navigation/PImplRecastNavMesh.cpp
--- a/Engine/Source/Runtime/Engine/Private/AI/Navigation/PImplRecastNavMesh.cpp
+++ b/Engine/Source/Runtime/Engine/Private/AI/Navigation/PImplRecastNavMesh.cpp
@@ -1133,7 +1133,7 @@ bool FPImplRecastNavMesh::FindStraightPath(const FVector& StartLoc, const FVecto
 	return bResult;
 }
 
-static bool IsDebugNodeModified(const FRecastDebugPathfindingNode& NodeData, const FRecastDebugPathfindingStep& PreviousStep)
+static bool IsDebugNodeModified(const FRecastDebugPathfindingNode& NodeData, const FRecastDebugPathfindingData& PreviousStep)
 {
 	const FRecastDebugPathfindingNode* PrevNodeData = PreviousStep.Nodes.Find(NodeData);
 	if (PrevNodeData)
@@ -1150,11 +1150,27 @@ static bool IsDebugNodeModified(const FRecastDebugPathfindingNode& NodeData, con
 	return true;
 }
 
-static void StorePathfindingDebugStep(const dtNavMeshQuery& NavQuery, const dtNavMesh* NavMesh, TArray<FRecastDebugPathfindingStep>& Steps)
+static void StorePathfindingDebugLength(FRecastDebugPathfindingNode& Node, FRecastDebugPathfindingData& Data)
+{
+	if (Node.Length >= 0.0f)
+	{
+		return;
+	}
+
+	FRecastDebugPathfindingNode* ParentNode = Data.Nodes.Find(FRecastDebugPathfindingNode(Node.ParentRef));
+	if (ParentNode)
+	{
+		StorePathfindingDebugLength(*ParentNode, Data);
+		Node.Length = ParentNode->Length + FVector::Dist(Node.NodePos, ParentNode->NodePos);
+	}
+	else
+	{
+		Node.Length = 0.0f;
+	}
+}
+
+static void StorePathfindingDebugData(const dtNavMeshQuery& NavQuery, const dtNavMesh* NavMesh, FRecastDebugPathfindingData& Data)
 {
-	const int StepIdx = Steps.AddZeroed(1);
-	FRecastDebugPathfindingStep& StepInfo = Steps[StepIdx];
-	
 	dtNode* BestNode = 0;
 	float BestNodeCost = 0.0f;
 	NavQuery.getCurrentBestResult(BestNode, BestNodeCost);
@@ -1163,12 +1179,13 @@ static void StorePathfindingDebugStep(const dtNavMeshQuery& NavQuery, const dtNa
 	for (int32 i = 0; i < NodePool->getNodeCount(); i++)
 	{
 		const dtNode* Node = NodePool->getNodeAtIdx(i + 1);
-		
+
 		FRecastDebugPathfindingNode NodeInfo;
 		NodeInfo.PolyRef = Node->id;
 		NodeInfo.ParentRef = Node->pidx ? NodePool->getNodeAtIdx(Node->pidx)->id : 0;
 		NodeInfo.Cost = Node->cost;
 		NodeInfo.TotalCost = Node->total;
+		NodeInfo.Length = -1.0f;
 		NodeInfo.bOpenSet = !NavQuery.isInClosedList(Node->id);
 		NodeInfo.bModified = true;
 		NodeInfo.NodePos = Recast2UnrealPoint(&Node->pos[0]);
@@ -1178,21 +1195,42 @@ static void StorePathfindingDebugStep(const dtNavMeshQuery& NavQuery, const dtNa
 		NavMesh->getTileAndPolyByRef(Node->id, &NavTile, &NavPoly);
 
 		NodeInfo.bOffMeshLink = NavPoly ? (NavPoly->getType() != DT_POLYTYPE_GROUND) : false;
-		for (int32 iv = 0; iv < NavPoly->vertCount; iv++)
+		if (Data.Flags & ERecastDebugPathfindingFlags::Vertices)
 		{
-			NodeInfo.Verts.Add(Recast2UnrealPoint(&NavTile->verts[NavPoly->verts[iv] * 3]));
+			for (int32 iv = 0; iv < NavPoly->vertCount; iv++)
+			{
+				NodeInfo.Verts.Add(Recast2UnrealPoint(&NavTile->verts[NavPoly->verts[iv] * 3]));
+			}
+		}
+
+		FSetElementId SetId = Data.Nodes.Add(NodeInfo);
+		if (Node == BestNode && (Data.Flags & ERecastDebugPathfindingFlags::BestNode))
+		{
+			Data.BestNode = SetId;
 		}
+	}
 
-		FSetElementId NodeId = StepInfo.Nodes.Add(NodeInfo);
-		if (Node == BestNode)
+	if (Data.Flags & ERecastDebugPathfindingFlags::PathLength)
+	{
+		for (TSet<FRecastDebugPathfindingNode>::TIterator It(Data.Nodes); It; ++It)
 		{
-			StepInfo.BestNode = NodeId;
+			FRecastDebugPathfindingNode& Node = *It;
+			StorePathfindingDebugLength(Node, Data);
 		}
 	}
+}
+
+static void StorePathfindingDebugStep(const dtNavMeshQuery& NavQuery, const dtNavMesh* NavMesh, TArray<FRecastDebugPathfindingData>& Steps)
+{
+	const int StepIdx = Steps.AddZeroed(1);
+	FRecastDebugPathfindingData& StepInfo = Steps[StepIdx];
+	StepInfo.Flags = ERecastDebugPathfindingFlags::BestNode | ERecastDebugPathfindingFlags::Vertices;
+	
+	StorePathfindingDebugData(NavQuery, NavMesh, StepInfo);
 
 	if (Steps.Num() > 1)
 	{
-		FRecastDebugPathfindingStep& PrevStepInfo = Steps[StepIdx - 1];
+		FRecastDebugPathfindingData& PrevStepInfo = Steps[StepIdx - 1];
 		for (TSet<FRecastDebugPathfindingNode>::TIterator It(StepInfo.Nodes); It; ++It)
 		{
 			FRecastDebugPathfindingNode& NodeData = *It;
@@ -1201,7 +1239,7 @@ static void StorePathfindingDebugStep(const dtNavMeshQuery& NavQuery, const dtNa
 	}
 }
 
-int32 FPImplRecastNavMesh::DebugPathfinding(const FVector& StartLoc, const FVector& EndLoc, const FNavigationQueryFilter& Filter, const UObject* Owner, TArray<FRecastDebugPathfindingStep>& Steps)
+int32 FPImplRecastNavMesh::DebugPathfinding(const FVector& StartLoc, const FVector& EndLoc, const FNavigationQueryFilter& Filter, const UObject* Owner, TArray<FRecastDebugPathfindingData>& Steps)
 {
 	int32 NumSteps = 0;
 
@@ -1475,7 +1513,9 @@ NavNodeRef FPImplRecastNavMesh::FindNearestPoly(FVector const& Loc, FVector cons
 	return INVALID_NAVNODEREF;
 }
 
-bool FPImplRecastNavMesh::GetPolysWithinPathingDistance(FVector const& StartLoc, const float PathingDistance, const FNavigationQueryFilter& Filter, const UObject* Owner, TArray<NavNodeRef>& FoundPolys) const
+bool FPImplRecastNavMesh::GetPolysWithinPathingDistance(FVector const& StartLoc, const float PathingDistance, 
+	const FNavigationQueryFilter& Filter, const UObject* Owner,
+	TArray<NavNodeRef>& FoundPolys, FRecastDebugPathfindingData* DebugData) const
 {
 	ensure(PathingDistance > 0.0f && "PathingDistance <= 0 doesn't make sense");
 
@@ -1513,6 +1553,11 @@ bool FPImplRecastNavMesh::GetPolysWithinPathingDistance(FVector const& StartLoc,
 
 	FoundPolys.RemoveAt(NumPolys, FoundPolys.Num() - NumPolys);
 
+	if (DebugData)
+	{
+		StorePathfindingDebugData(NavQuery, DetourNavMesh, *DebugData);
+	}
+
 	return FoundPolys.Num() > 0;
 }
 
diff --git a/Engine/Source/Runtime/Engine/Private/AI/Navigation/RecastNavMesh.cpp b/Engine/Source/Runtime/Engine/Private/AI/Navigation/RecastNavMesh.cpp
--- a/Engine/Source/Runtime/Engine/Private/AI/Navigation/RecastNavMesh.cpp
+++ b/Engine/Source/Runtime/Engine/Private/AI/Navigation/RecastNavMesh.cpp
@@ -1254,12 +1254,13 @@ bool ARecastNavMesh::GetClusterBounds(NavNodeRef ClusterRef, FBox& OutBounds) co
 	return bSuccess;
 }
 
-bool ARecastNavMesh::GetPolysWithinPathingDistance(FVector const& StartLoc, const float PathingDistance, TArray<NavNodeRef>& FoundPolys, TSharedPtr<const FNavigationQueryFilter> Filter, const UObject* QueryOwner) const
+bool ARecastNavMesh::GetPolysWithinPathingDistance(FVector const& StartLoc, const float PathingDistance, TArray<NavNodeRef>& FoundPolys,
+	TSharedPtr<const FNavigationQueryFilter> Filter, const UObject* QueryOwner, FRecastDebugPathfindingData* DebugData) const
 {
 	bool bSuccess = false;
 	if (RecastNavMeshImpl)
 	{
-		bSuccess = RecastNavMeshImpl->GetPolysWithinPathingDistance(StartLoc, PathingDistance, GetRightFilterRef(Filter), QueryOwner, FoundPolys);
+		bSuccess = RecastNavMeshImpl->GetPolysWithinPathingDistance(StartLoc, PathingDistance, GetRightFilterRef(Filter), QueryOwner, FoundPolys, DebugData);
 	}
 
 	return bSuccess;
@@ -1812,7 +1813,7 @@ bool ARecastNavMesh::FindStraightPath(const FVector& StartLoc, const FVector& En
 	return bResult;
 }
 
-int32 ARecastNavMesh::DebugPathfinding(const FPathFindingQuery& Query, TArray<FRecastDebugPathfindingStep>& Steps)
+int32 ARecastNavMesh::DebugPathfinding(const FPathFindingQuery& Query, TArray<FRecastDebugPathfindingData>& Steps)
 {
 	int32 NumSteps = 0;
 
diff --git a/Engine/Source/Runtime/Engine/Public/AI/Navigation/PImplRecastNavMesh.h b/Engine/Source/Runtime/Engine/Public/AI/Navigation/PImplRecastNavMesh.h
--- a/Engine/Source/Runtime/Engine/Public/AI/Navigation/PImplRecastNavMesh.h
+++ b/Engine/Source/Runtime/Engine/Public/AI/Navigation/PImplRecastNavMesh.h
@@ -116,7 +116,7 @@ public:
 	void Raycast2D(NavNodeRef StartNode, const FVector& StartLoc, const FVector& EndLoc, const FNavigationQueryFilter& InQueryFilter, const UObject* Owner, ARecastNavMesh::FRaycastResult& RaycastResult) const;
 
 	/** Generates path from given query and collect data for every step of A* algorithm */
-	int32 DebugPathfinding(const FVector& StartLoc, const FVector& EndLoc, const FNavigationQueryFilter& Filter, const UObject* Owner, TArray<FRecastDebugPathfindingStep>& Steps);
+	int32 DebugPathfinding(const FVector& StartLoc, const FVector& EndLoc, const FNavigationQueryFilter& Filter, const UObject* Owner, TArray<FRecastDebugPathfindingData>& Steps);
 
 	/** Returns a random location on the navmesh. */
 	FNavLocation GetRandomPoint(const FNavigationQueryFilter& Filter, const UObject* Owner) const;
@@ -141,7 +141,9 @@ public:
 	 *	@NOTE query is not using string-pulled path distance (for performance reasons),
 	 *		it measured distance between middles of portal edges, do you might want to 
 	 *		add an extra margin to PathingDistance */
-	bool GetPolysWithinPathingDistance(FVector const& StartLoc, const float PathingDistance, const FNavigationQueryFilter& Filter, const UObject* Owner, TArray<NavNodeRef>& FoundPolys) const;
+	bool GetPolysWithinPathingDistance(FVector const& StartLoc, const float PathingDistance,
+		const FNavigationQueryFilter& Filter, const UObject* Owner,
+		TArray<NavNodeRef>& FoundPolys, FRecastDebugPathfindingData* DebugData) const;
 
 	//@todo document
 	void GetEdgesForPathCorridor(const TArray<NavNodeRef>* PathCorridor, TArray<FNavigationPortalEdge>* PathCorridorEdges) const;
