removed support for EQS named params in all engine-level classes, marked structs as deprecated
#ue4

diff --git a/Engine/Source/Editor/EnvironmentQueryEditor/Private/DetailCustomizations/EnvQueryParamSetupCustomization.cpp b/Engine/Source/Editor/EnvironmentQueryEditor/Private/DetailCustomizations/EnvQueryParamSetupCustomization.cpp
deleted file mode 100644
diff --git a/Engine/Source/Editor/EnvironmentQueryEditor/Private/DetailCustomizations/EnvQueryParamSetupCustomization.h b/Engine/Source/Editor/EnvironmentQueryEditor/Private/DetailCustomizations/EnvQueryParamSetupCustomization.h
deleted file mode 100644
diff --git a/Engine/Source/Editor/EnvironmentQueryEditor/Private/EnvironmentQueryEditorModule.cpp b/Engine/Source/Editor/EnvironmentQueryEditor/Private/EnvironmentQueryEditorModule.cpp
--- a/Engine/Source/Editor/EnvironmentQueryEditor/Private/EnvironmentQueryEditorModule.cpp
+++ b/Engine/Source/Editor/EnvironmentQueryEditor/Private/EnvironmentQueryEditorModule.cpp
@@ -4,7 +4,6 @@
 
 #include "DetailCustomizations/EnvDirectionCustomization.h"
 #include "DetailCustomizations/EnvTraceDataCustomization.h"
-#include "DetailCustomizations/EnvQueryParamSetupCustomization.h"
 #include "DetailCustomizations/EnvQueryTestDetails.h"
 
 #include "ModuleManager.h"
@@ -48,9 +47,6 @@ void FEnvironmentQueryEditorModule::StartupModule()
 
 	// Register the details customizer
 	FPropertyEditorModule& PropertyModule = FModuleManager::LoadModuleChecked<FPropertyEditorModule>("PropertyEditor");
-	PropertyModule.RegisterCustomPropertyTypeLayout( "EnvFloatParam", FOnGetPropertyTypeCustomizationInstance::CreateStatic( &FEnvQueryParamSetupCustomization::MakeInstance ) );
-	PropertyModule.RegisterCustomPropertyTypeLayout( "EnvIntParam", FOnGetPropertyTypeCustomizationInstance::CreateStatic( &FEnvQueryParamSetupCustomization::MakeInstance ) );
-	PropertyModule.RegisterCustomPropertyTypeLayout( "EnvBoolParam", FOnGetPropertyTypeCustomizationInstance::CreateStatic( &FEnvQueryParamSetupCustomization::MakeInstance ) );
 	PropertyModule.RegisterCustomPropertyTypeLayout( "EnvDirection", FOnGetPropertyTypeCustomizationInstance::CreateStatic( &FEnvDirectionCustomization::MakeInstance ) );
 	PropertyModule.RegisterCustomPropertyTypeLayout( "EnvTraceData", FOnGetPropertyTypeCustomizationInstance::CreateStatic( &FEnvTraceDataCustomization::MakeInstance ) );
 	PropertyModule.RegisterCustomClassLayout( "EnvQueryTest", FOnGetDetailCustomizationInstance::CreateStatic( &FEnvQueryTestDetails::MakeInstance ) );
@@ -89,9 +85,6 @@ void FEnvironmentQueryEditorModule::ShutdownModule()
 	if (FModuleManager::Get().IsModuleLoaded("PropertyEditor"))
 	{
 		FPropertyEditorModule& PropertyModule = FModuleManager::LoadModuleChecked<FPropertyEditorModule>("PropertyEditor");
-		PropertyModule.UnregisterCustomPropertyTypeLayout( "EnvFloatParam" );
-		PropertyModule.UnregisterCustomPropertyTypeLayout( "EnvIntParam" );
-		PropertyModule.UnregisterCustomPropertyTypeLayout( "EnvBoolParam" );
 		PropertyModule.UnregisterCustomPropertyTypeLayout( "EnvDirection" );
 		PropertyModule.UnregisterCustomPropertyTypeLayout( "EnvTraceData" );
 
diff --git a/Engine/Source/Editor/EnvironmentQueryEditor/Private/EnvironmentQueryGraphNode_Option.cpp b/Engine/Source/Editor/EnvironmentQueryEditor/Private/EnvironmentQueryGraphNode_Option.cpp
--- a/Engine/Source/Editor/EnvironmentQueryEditor/Private/EnvironmentQueryGraphNode_Option.cpp
+++ b/Engine/Source/Editor/EnvironmentQueryEditor/Private/EnvironmentQueryGraphNode_Option.cpp
@@ -108,9 +108,9 @@ void UEnvironmentQueryGraphNode_Option::CalculateWeights()
 		}
 
 		UEnvQueryTest* TestInstance = Cast<UEnvQueryTest>(TestNode->NodeInstance);
-		if (TestInstance && !TestInstance->Weight.IsNamedParam())
+		if (TestInstance && !TestInstance->ScoringFactor.IsDynamic())
 		{
-			MaxWeight = FMath::Max(MaxWeight, FMath::Abs(TestInstance->Weight.Value));
+			MaxWeight = FMath::Max(MaxWeight, FMath::Abs(TestInstance->ScoringFactor.DefaultValue));
 		}
 	}
 
@@ -128,12 +128,12 @@ void UEnvironmentQueryGraphNode_Option::CalculateWeights()
 		}
 		
 		UEnvQueryTest* TestInstance = Cast<UEnvQueryTest>(TestNode->NodeInstance);
-		const bool bHasNamed = TestInstance && TestNode->bTestEnabled && TestInstance->Weight.IsNamedParam();
+		const bool bHasDynamic = TestInstance && TestNode->bTestEnabled && TestInstance->ScoringFactor.IsDynamic();
 		const float NewWeight = (TestInstance && TestNode->bTestEnabled) ?
-			(TestInstance->Weight.IsNamedParam() ? 1.0f : FMath::Clamp(FMath::Abs(TestInstance->Weight.Value) / MaxWeight, 0.0f, 1.0f)) : 
+			(TestInstance->ScoringFactor.IsDynamic() ? 1.0f : FMath::Clamp(FMath::Abs(TestInstance->ScoringFactor.DefaultValue) / MaxWeight, 0.0f, 1.0f)) :
 			-1.0f;
 
-		TestNode->SetDisplayedWeight(NewWeight, bHasNamed);
+		TestNode->SetDisplayedWeight(NewWeight, bHasDynamic);
 	}
 }
 
diff --git a/Engine/Source/Runtime/AIModule/Classes/EnvironmentQuery/EnvQueryTest.h b/Engine/Source/Runtime/AIModule/Classes/EnvironmentQuery/EnvQueryTest.h
--- a/Engine/Source/Runtime/AIModule/Classes/EnvironmentQuery/EnvQueryTest.h
+++ b/Engine/Source/Runtime/AIModule/Classes/EnvironmentQuery/EnvQueryTest.h
@@ -91,26 +91,6 @@ class AIMODULE_API UEnvQueryTest : public UEnvQueryNode
 					|| (ScoringEquation == EEnvTestScoreEquation::Constant));	// We are giving a constant score value for passing.
 	}
 
-	// BEGIN: deprecated properties, do not use them
-	UPROPERTY()
-	FEnvBoolParam BoolFilter;
-
-	UPROPERTY()
-	FEnvFloatParam FloatFilterMin;
-
-	UPROPERTY()
-	FEnvFloatParam FloatFilterMax;
-
-	UPROPERTY()
-	FEnvFloatParam ScoreClampingMin;
-
-	UPROPERTY()
-	FEnvFloatParam ScoreClampingMax;
-
-	UPROPERTY()
-	FEnvFloatParam Weight;
-	// END: deprecated properties, do not use them
-
 private:
 	/** When set, test operates on float values (e.g. distance, with AtLeast, UpTo conditions),
 	 *  otherwise it will accept bool values (e.g. visibility, with Equals condition) */
diff --git a/Engine/Source/Runtime/AIModule/Classes/EnvironmentQuery/EnvQueryTypes.h b/Engine/Source/Runtime/AIModule/Classes/EnvironmentQuery/EnvQueryTypes.h
--- a/Engine/Source/Runtime/AIModule/Classes/EnvironmentQuery/EnvQueryTypes.h
+++ b/Engine/Source/Runtime/AIModule/Classes/EnvironmentQuery/EnvQueryTypes.h
@@ -184,60 +184,6 @@
 	};
 }
 
-// DEPRECATED, will be removed soon - use AI Data Providers instead 
-USTRUCT()
-struct AIMODULE_API FEnvFloatParam
-{
-	GENERATED_USTRUCT_BODY();
-
-	/** default value */
-	UPROPERTY(EditDefaultsOnly, Category=Param)
-	float Value;
-
-	/** name of parameter */
-	UPROPERTY(EditDefaultsOnly, Category=Param)
-	FName ParamName;
-
-	bool IsNamedParam() const { return ParamName != NAME_None; }
-	void Convert(UObject* Owner, FAIDataProviderFloatValue& ValueProvider);
-};
-
-// DEPRECATED, will be removed soon - use AI Data Providers instead 
-USTRUCT()
-struct AIMODULE_API FEnvIntParam
-{
-	GENERATED_USTRUCT_BODY();
-
-	/** default value */
-	UPROPERTY(EditDefaultsOnly, Category=Param)
-	int32 Value;
-
-	/** name of parameter */
-	UPROPERTY(EditDefaultsOnly, Category=Param)
-	FName ParamName;
-
-	bool IsNamedParam() const { return ParamName != NAME_None; }
-	void Convert(UObject* Owner, FAIDataProviderIntValue& ValueProvider);
-};
-
-// DEPRECATED, will be removed soon - use AI Data Providers instead 
-USTRUCT()
-struct AIMODULE_API FEnvBoolParam
-{
-	GENERATED_USTRUCT_BODY();
-
-	/** default value */
-	UPROPERTY(EditDefaultsOnly, Category=Param)
-	bool Value;
-
-	/** name of parameter */
-	UPROPERTY(EditDefaultsOnly, Category=Param)
-	FName ParamName;
-
-	bool IsNamedParam() const { return ParamName != NAME_None; }
-	void Convert(UObject* Owner, FAIDataProviderBoolValue& ValueProvider);
-};
-
 USTRUCT(BlueprintType)
 struct AIMODULE_API FEnvNamedValue
 {
@@ -935,6 +881,83 @@
 #endif // WITH_RECAST
 }
 
+// BEGIN DEPRECATED SUPPORT
+
+USTRUCT()
+struct FEnvFloatParam_DEPRECATED
+{
+	GENERATED_USTRUCT_BODY();
+
+	/** default value */
+	UPROPERTY(EditDefaultsOnly, Category = Param)
+	float Value;
+
+	/** name of parameter */
+	UPROPERTY(EditDefaultsOnly, Category = Param)
+	FName ParamName;
+
+	bool IsNamedParam() const { return ParamName != NAME_None; }
+	void Convert(UObject* Owner, FAIDataProviderFloatValue& ValueProvider);
+};
+
+USTRUCT()
+struct FEnvIntParam_DEPRECATED
+{
+	GENERATED_USTRUCT_BODY();
+
+	/** default value */
+	UPROPERTY(EditDefaultsOnly, Category = Param)
+	int32 Value;
+
+	/** name of parameter */
+	UPROPERTY(EditDefaultsOnly, Category = Param)
+	FName ParamName;
+
+	bool IsNamedParam() const { return ParamName != NAME_None; }
+	void Convert(UObject* Owner, FAIDataProviderIntValue& ValueProvider);
+};
+
+USTRUCT()
+struct FEnvBoolParam_DEPRECATED
+{
+	GENERATED_USTRUCT_BODY();
+
+	/** default value */
+	UPROPERTY(EditDefaultsOnly, Category = Param)
+	bool Value;
+
+	/** name of parameter */
+	UPROPERTY(EditDefaultsOnly, Category = Param)
+	FName ParamName;
+
+	bool IsNamedParam() const { return ParamName != NAME_None; }
+	void Convert(UObject* Owner, FAIDataProviderBoolValue& ValueProvider);
+};
+
+USTRUCT()
+struct AIMODULE_API FEnvFloatParam : public FEnvFloatParam_DEPRECATED
+{
+	GENERATED_USTRUCT_BODY();
+};
+
+USTRUCT()
+struct AIMODULE_API FEnvIntParam : public FEnvIntParam_DEPRECATED
+{
+	GENERATED_USTRUCT_BODY();
+};
+
+USTRUCT()
+struct AIMODULE_API FEnvBoolParam : public FEnvBoolParam_DEPRECATED
+{
+	GENERATED_USTRUCT_BODY();
+};
+
+struct DEPRECATED(4.8, "FEnvFloatParam is deprecated in 4.8 and was replaced with FAIDataProviderFloatValue. Please use that type instead.") FEnvFloatParam;
+struct DEPRECATED(4.8, "FEnvIntParam is deprecated in 4.8 and was replaced with FAIDataProviderIntValue. Please use that type instead.") FEnvIntParam;
+struct DEPRECATED(4.8, "FEnvBoolParam is deprecated in 4.8 and was replaced with FAIDataProviderBoolValue. Please use that type instead.") FEnvBoolParam;
+
+// END DEPRECATED SUPPORT
+
 UCLASS(Abstract)
 class AIMODULE_API UEnvQueryTypes : public UObject
 {
diff --git a/Engine/Source/Runtime/AIModule/Classes/EnvironmentQuery/Generators/EnvQueryGenerator_ActorsOfClass.h b/Engine/Source/Runtime/AIModule/Classes/EnvironmentQuery/Generators/EnvQueryGenerator_ActorsOfClass.h
--- a/Engine/Source/Runtime/AIModule/Classes/EnvironmentQuery/Generators/EnvQueryGenerator_ActorsOfClass.h
+++ b/Engine/Source/Runtime/AIModule/Classes/EnvironmentQuery/Generators/EnvQueryGenerator_ActorsOfClass.h
@@ -23,12 +23,7 @@ class UEnvQueryGenerator_ActorsOfClass : public UEnvQueryGenerator
 	UPROPERTY(EditAnywhere, Category=Generator)
 	TSubclassOf<UEnvQueryContext> SearchCenter;
 
-	// DEPRECATED
-	UPROPERTY()
-	FEnvFloatParam Radius;
-
 	virtual void GenerateItems(FEnvQueryInstance& QueryInstance) const override;
-	virtual void PostLoad() override;
 
 	virtual FText GetDescriptionTitle() const override;
 	virtual FText GetDescriptionDetails() const override;
diff --git a/Engine/Source/Runtime/AIModule/Classes/EnvironmentQuery/Generators/EnvQueryGenerator_OnCircle.h b/Engine/Source/Runtime/AIModule/Classes/EnvironmentQuery/Generators/EnvQueryGenerator_OnCircle.h
--- a/Engine/Source/Runtime/AIModule/Classes/EnvironmentQuery/Generators/EnvQueryGenerator_OnCircle.h
+++ b/Engine/Source/Runtime/AIModule/Classes/EnvironmentQuery/Generators/EnvQueryGenerator_OnCircle.h
@@ -41,17 +41,6 @@ class AIMODULE_API UEnvQueryGenerator_OnCircle : public UEnvQueryGenerator_Proje
 	UPROPERTY()
 	uint32 bDefineArc:1;
 
-	// BEGIN: deprecated properties
-	UPROPERTY()
-	FEnvFloatParam Radius;
-
-	UPROPERTY()
-	FEnvFloatParam ItemSpacing;
-
-	UPROPERTY()
-	FEnvFloatParam Angle;
-	// END: deprecated properties
-
 	virtual void PostLoad() override;
 
 	virtual void GenerateItems(FEnvQueryInstance& QueryInstance) const override;
diff --git a/Engine/Source/Runtime/AIModule/Classes/EnvironmentQuery/Generators/EnvQueryGenerator_SimpleGrid.h b/Engine/Source/Runtime/AIModule/Classes/EnvironmentQuery/Generators/EnvQueryGenerator_SimpleGrid.h
--- a/Engine/Source/Runtime/AIModule/Classes/EnvironmentQuery/Generators/EnvQueryGenerator_SimpleGrid.h
+++ b/Engine/Source/Runtime/AIModule/Classes/EnvironmentQuery/Generators/EnvQueryGenerator_SimpleGrid.h
@@ -27,16 +27,7 @@ class UEnvQueryGenerator_SimpleGrid : public UEnvQueryGenerator_ProjectedPoints
 	UPROPERTY(EditDefaultsOnly, Category=Generator)
 	TSubclassOf<UEnvQueryContext> GenerateAround;
 
-	// BEGIN: deprecated properties
-	UPROPERTY()
-	FEnvFloatParam Radius;
-
-	UPROPERTY()
-	FEnvFloatParam Density;
-	// END: deprecated properties
-
 	virtual void GenerateItems(FEnvQueryInstance& QueryInstance) const override;
-	virtual void PostLoad() override;
 
 	virtual FText GetDescriptionTitle() const override;
 	virtual FText GetDescriptionDetails() const override;
diff --git a/Engine/Source/Runtime/AIModule/Classes/EnvironmentQuery/Tests/EnvQueryTest_Pathfinding.h b/Engine/Source/Runtime/AIModule/Classes/EnvironmentQuery/Tests/EnvQueryTest_Pathfinding.h
--- a/Engine/Source/Runtime/AIModule/Classes/EnvironmentQuery/Tests/EnvQueryTest_Pathfinding.h
+++ b/Engine/Source/Runtime/AIModule/Classes/EnvironmentQuery/Tests/EnvQueryTest_Pathfinding.h
@@ -42,17 +42,6 @@ class UEnvQueryTest_Pathfinding : public UEnvQueryTest
 	UPROPERTY(EditDefaultsOnly, Category=Pathfinding)
 	TSubclassOf<UNavigationQueryFilter> FilterClass;
 
-	// BEGIN: deprecated properties 
-	UPROPERTY()
-	FEnvBoolParam PathToItem;
-
-	UPROPERTY()
-	FEnvBoolParam DiscardUnreachable;
-
-	UPROPERTY()
-	FEnvBoolParam HierarchicalPathfinding;
-	// END: deprecated properties
-
 	virtual void RunTest(FEnvQueryInstance& QueryInstance) const override;
 
 	virtual FText GetDescriptionTitle() const override;
diff --git a/Engine/Source/Runtime/AIModule/Classes/EnvironmentQuery/Tests/EnvQueryTest_Trace.h b/Engine/Source/Runtime/AIModule/Classes/EnvironmentQuery/Tests/EnvQueryTest_Trace.h
--- a/Engine/Source/Runtime/AIModule/Classes/EnvironmentQuery/Tests/EnvQueryTest_Trace.h
+++ b/Engine/Source/Runtime/AIModule/Classes/EnvironmentQuery/Tests/EnvQueryTest_Trace.h
@@ -31,19 +31,7 @@ class UEnvQueryTest_Trace : public UEnvQueryTest
 	UPROPERTY(EditDefaultsOnly, Category=Trace)
 	TSubclassOf<UEnvQueryContext> Context;
 
-	// BEGIN: deprecated properties 
-	UPROPERTY()
-	FEnvBoolParam TraceToItem;
-
-	UPROPERTY()
-	FEnvFloatParam ItemOffsetZ;
-
-	UPROPERTY()
-	FEnvFloatParam ContextOffsetZ;
-	// END: deprecated properties
-
 	virtual void RunTest(FEnvQueryInstance& QueryInstance) const override;
-	virtual void PostLoad() override;
 
 	virtual FText GetDescriptionTitle() const override;
 	virtual FText GetDescriptionDetails() const override;
diff --git a/Engine/Source/Runtime/AIModule/Private/EnvironmentQuery/EnvQuery.cpp b/Engine/Source/Runtime/AIModule/Private/EnvironmentQuery/EnvQuery.cpp
--- a/Engine/Source/Runtime/AIModule/Private/EnvironmentQuery/EnvQuery.cpp
+++ b/Engine/Source/Runtime/AIModule/Private/EnvironmentQuery/EnvQuery.cpp
@@ -1,6 +1,7 @@
 // Copyright 1998-2015 Epic Games, Inc. All Rights Reserved.
 
 #include "AIModulePrivate.h"
+#include "DataProviders/AIDataProvider_QueryParams.h"
 #include "EnvironmentQuery/EnvQuery.h"
 #include "EnvironmentQuery/EnvQueryGenerator.h"
 #include "EnvironmentQuery/EnvQueryOption.h"
@@ -60,20 +61,32 @@ static void AddNamedValuesFromObject(const UObject* Ob, TArray<FEnvNamedValue>&
 		}
 
 		FString TypeDesc = TestStruct->GetCPPType(NULL, CPPF_None);
-		if (TypeDesc.Contains(GET_STRUCT_NAME_CHECKED(FEnvIntParam)))
+		if (TypeDesc.Contains(GET_STRUCT_NAME_CHECKED(FAIDataProviderIntValue)))
 		{
-			const FEnvIntParam* PropertyValue = TestStruct->ContainerPtrToValuePtr<FEnvIntParam>(Ob);
-			AddNamedValue(PropertyValue->ParamName, EEnvQueryParam::Int, *((float*)&PropertyValue->Value), NamedValues, RequiredParams);
+			const FAIDataProviderIntValue* PropertyValue = TestStruct->ContainerPtrToValuePtr<FAIDataProviderIntValue>(Ob);
+			const UAIDataProvider_QueryParams* QueryParamProvider = PropertyValue ? Cast<const UAIDataProvider_QueryParams>(PropertyValue->DataBinding) : nullptr;
+			if (QueryParamProvider && !QueryParamProvider->ParamName.IsNone())
+			{
+				AddNamedValue(QueryParamProvider->ParamName, EEnvQueryParam::Int, *((float*)&PropertyValue->DefaultValue), NamedValues, RequiredParams);
+			}
 		}
-		else if (TypeDesc.Contains(GET_STRUCT_NAME_CHECKED(FEnvFloatParam)))
+		else if (TypeDesc.Contains(GET_STRUCT_NAME_CHECKED(FAIDataProviderFloatValue)))
 		{
-			const FEnvFloatParam* PropertyValue = TestStruct->ContainerPtrToValuePtr<FEnvFloatParam>(Ob);
-			AddNamedValue(PropertyValue->ParamName, EEnvQueryParam::Float, PropertyValue->Value, NamedValues, RequiredParams);
+			const FAIDataProviderFloatValue* PropertyValue = TestStruct->ContainerPtrToValuePtr<FAIDataProviderFloatValue>(Ob);
+			const UAIDataProvider_QueryParams* QueryParamProvider = PropertyValue ? Cast<const UAIDataProvider_QueryParams>(PropertyValue->DataBinding) : nullptr;
+			if (QueryParamProvider && !QueryParamProvider->ParamName.IsNone())
+			{
+				AddNamedValue(QueryParamProvider->ParamName, EEnvQueryParam::Float, PropertyValue->DefaultValue, NamedValues, RequiredParams);
+			}
 		}
-		else if (TypeDesc.Contains(GET_STRUCT_NAME_CHECKED(FEnvBoolParam)))
+		else if (TypeDesc.Contains(GET_STRUCT_NAME_CHECKED(FAIDataProviderBoolValue)))
 		{
-			const FEnvBoolParam* PropertyValue = TestStruct->ContainerPtrToValuePtr<FEnvBoolParam>(Ob);
-			AddNamedValue(PropertyValue->ParamName, EEnvQueryParam::Bool, PropertyValue->Value ? 1.0f : -1.0f, NamedValues, RequiredParams);
+			const FAIDataProviderBoolValue* PropertyValue = TestStruct->ContainerPtrToValuePtr<FAIDataProviderBoolValue>(Ob);
+			const UAIDataProvider_QueryParams* QueryParamProvider = PropertyValue ? Cast<const UAIDataProvider_QueryParams>(PropertyValue->DataBinding) : nullptr;
+			if (QueryParamProvider && !QueryParamProvider->ParamName.IsNone())
+			{
+				AddNamedValue(QueryParamProvider->ParamName, EEnvQueryParam::Bool, PropertyValue->DefaultValue ? 1.0f : -1.0f, NamedValues, RequiredParams);
+			}
 		}
 	}
 }
diff --git a/Engine/Source/Runtime/AIModule/Private/EnvironmentQuery/EnvQueryTest.cpp b/Engine/Source/Runtime/AIModule/Private/EnvironmentQuery/EnvQueryTest.cpp
--- a/Engine/Source/Runtime/AIModule/Private/EnvironmentQuery/EnvQueryTest.cpp
+++ b/Engine/Source/Runtime/AIModule/Private/EnvironmentQuery/EnvQueryTest.cpp
@@ -23,10 +23,6 @@ UEnvQueryTest::UEnvQueryTest(const FObjectInitializer& ObjectInitializer) : Supe
 	ScoringFactor.DefaultValue = 1.0f;
 
 	bWorkOnFloatValues = true;
-
-	// keep deprecated properties initialized
-	BoolFilter.Value = true;
-	Weight.Value = 1.0f;
 }
 
 void UEnvQueryTest::NormalizeItemScores(FEnvQueryInstance& QueryInstance)
@@ -185,17 +181,6 @@ AActor* UEnvQueryTest::GetItemActor(FEnvQueryInstance& QueryInstance, int32 Item
 void UEnvQueryTest::PostLoad()
 {
 	Super::PostLoad();
-
-	if (VerNum < EnvQueryTestVersion::DataProviders)
-	{
-		BoolFilter.Convert(this, BoolValue);
-		FloatFilterMin.Convert(this, FloatValueMin);
-		FloatFilterMax.Convert(this, FloatValueMax);
-		ScoreClampingMin.Convert(this, ScoreClampMin);
-		ScoreClampingMax.Convert(this, ScoreClampMax);
-		Weight.Convert(this, ScoringFactor);
-	}
-
 	UpdateNodeVersion();
 }
 
diff --git a/Engine/Source/Runtime/AIModule/Private/EnvironmentQuery/EnvQueryTypes.cpp b/Engine/Source/Runtime/AIModule/Private/EnvironmentQuery/EnvQueryTypes.cpp
--- a/Engine/Source/Runtime/AIModule/Private/EnvironmentQuery/EnvQueryTypes.cpp
+++ b/Engine/Source/Runtime/AIModule/Private/EnvironmentQuery/EnvQueryTypes.cpp
@@ -196,7 +196,7 @@ void FEnvTraceData::SetNavmeshOnly()
 	bCanDisableTrace = false;
 }
 
-void FEnvBoolParam::Convert(UObject* Owner, FAIDataProviderBoolValue& ValueProvider)
+void FEnvBoolParam_DEPRECATED::Convert(UObject* Owner, FAIDataProviderBoolValue& ValueProvider)
 {
 	ValueProvider.DefaultValue = Value;
 	if (IsNamedParam())
@@ -208,7 +208,7 @@ void FEnvBoolParam::Convert(UObject* Owner, FAIDataProviderBoolValue& ValueProvi
 	}
 }
 
-void FEnvIntParam::Convert(UObject* Owner, FAIDataProviderIntValue& ValueProvider)
+void FEnvIntParam_DEPRECATED::Convert(UObject* Owner, FAIDataProviderIntValue& ValueProvider)
 {
 	ValueProvider.DefaultValue = Value;
 	if (IsNamedParam())
@@ -220,7 +220,7 @@ void FEnvIntParam::Convert(UObject* Owner, FAIDataProviderIntValue& ValueProvide
 	}
 }
 
-void FEnvFloatParam::Convert(UObject* Owner, FAIDataProviderFloatValue& ValueProvider)
+void FEnvFloatParam_DEPRECATED::Convert(UObject* Owner, FAIDataProviderFloatValue& ValueProvider)
 {
 	ValueProvider.DefaultValue = Value;
 	if (IsNamedParam())
diff --git a/Engine/Source/Runtime/AIModule/Private/EnvironmentQuery/Generators/EnvQueryGenerator_ActorsOfClass.cpp b/Engine/Source/Runtime/AIModule/Private/EnvironmentQuery/Generators/EnvQueryGenerator_ActorsOfClass.cpp
--- a/Engine/Source/Runtime/AIModule/Private/EnvironmentQuery/Generators/EnvQueryGenerator_ActorsOfClass.cpp
+++ b/Engine/Source/Runtime/AIModule/Private/EnvironmentQuery/Generators/EnvQueryGenerator_ActorsOfClass.cpp
@@ -15,9 +15,6 @@ UEnvQueryGenerator_ActorsOfClass::UEnvQueryGenerator_ActorsOfClass(const FObject
 	ItemType = UEnvQueryItemType_Actor::StaticClass();
 	SearchRadius.DefaultValue = 500.0f;
 	SearchedActorClass = AActor::StaticClass();
-
-	// keep deprecated properties initialized
-	Radius.Value = 500.0f;
 }
 
 void UEnvQueryGenerator_ActorsOfClass::GenerateItems(FEnvQueryInstance& QueryInstance) const
@@ -69,14 +66,4 @@ FText UEnvQueryGenerator_ActorsOfClass::GetDescriptionDetails() const
 	return Desc;
 }
 
-void UEnvQueryGenerator_ActorsOfClass::PostLoad()
-{
-	if (VerNum < EnvQueryGeneratorVersion::DataProviders)
-	{
-		Radius.Convert(this, SearchRadius);
-	}
-
-	Super::PostLoad();
-}
-
 #undef LOCTEXT_NAMESPACE
\ No newline at end of file
diff --git a/Engine/Source/Runtime/AIModule/Private/EnvironmentQuery/Generators/EnvQueryGenerator_OnCircle.cpp b/Engine/Source/Runtime/AIModule/Private/EnvironmentQuery/Generators/EnvQueryGenerator_OnCircle.cpp
--- a/Engine/Source/Runtime/AIModule/Private/EnvironmentQuery/Generators/EnvQueryGenerator_OnCircle.cpp
+++ b/Engine/Source/Runtime/AIModule/Private/EnvironmentQuery/Generators/EnvQueryGenerator_OnCircle.cpp
@@ -112,22 +112,10 @@ UEnvQueryGenerator_OnCircle::UEnvQueryGenerator_OnCircle(const FObjectInitialize
 	TraceData.TraceMode = EEnvQueryTrace::Navigation;
 
 	ProjectionData.TraceMode = EEnvQueryTrace::None;
-
-	// keep deprecated properties initialized
-	Radius.Value = 1000.0f;
-	ItemSpacing.Value = 50.0f;
-	Angle.Value = 360.0f;
 }
 
 void UEnvQueryGenerator_OnCircle::PostLoad()
 {
-	if (VerNum < EnvQueryGeneratorVersion::DataProviders)
-	{
-		Radius.Convert(this, CircleRadius);
-		ItemSpacing.Convert(this, SpaceBetween);
-		Angle.Convert(this, ArcAngle);
-	}
-
 	Super::PostLoad();
 
 	if (ArcAngle.DefaultValue <= 0)
@@ -206,7 +194,7 @@ void UEnvQueryGenerator_OnCircle::GenerateItems(FEnvQueryInstance& QueryInstance
 
 	// first generate points on a circle
 	const float CircumferenceLength = 2.f * PI * RadiusValue;
-	const float ArcAnglePercentage = Angle.Value / 360.f;
+	const float ArcAnglePercentage = AngleDegree / 360.f;
 	const float ArcLength = CircumferenceLength * ArcAnglePercentage;
 	const int32 StepsCount = FMath::CeilToInt(ArcLength / ItemSpace) + 1;
 	const float AngleStep = AngleDegree / (StepsCount - 1);
@@ -352,7 +340,7 @@ FText UEnvQueryGenerator_OnCircle::GetDescriptionDetails() const
 	{
 		FFormatNamedArguments ArcArgs;
 		ArcArgs.Add(TEXT("Description"), Desc);
-		ArcArgs.Add(TEXT("AngleValue"), Angle.Value);
+		ArcArgs.Add(TEXT("AngleValue"), FText::FromString(ArcAngle.ToString()));
 		ArcArgs.Add(TEXT("ArcDirection"), ArcDirection.ToText());
 		Desc = FText::Format(LOCTEXT("DescriptionWithArc", "{Description}\nLimit to {AngleValue} angle both sides on {ArcDirection}"), ArcArgs);
 	}
diff --git a/Engine/Source/Runtime/AIModule/Private/EnvironmentQuery/Generators/EnvQueryGenerator_SimpleGrid.cpp b/Engine/Source/Runtime/AIModule/Private/EnvironmentQuery/Generators/EnvQueryGenerator_SimpleGrid.cpp
--- a/Engine/Source/Runtime/AIModule/Private/EnvironmentQuery/Generators/EnvQueryGenerator_SimpleGrid.cpp
+++ b/Engine/Source/Runtime/AIModule/Private/EnvironmentQuery/Generators/EnvQueryGenerator_SimpleGrid.cpp
@@ -11,21 +11,6 @@ UEnvQueryGenerator_SimpleGrid::UEnvQueryGenerator_SimpleGrid(const FObjectInitia
 	GenerateAround = UEnvQueryContext_Querier::StaticClass();
 	GridSize.DefaultValue = 500.0f;
 	SpaceBetween.DefaultValue = 100.0f;
-
-	// keep deprecated properties initialized
-	Radius.Value = 500.0f;
-	Density.Value = 500.0f;
-}
-
-void UEnvQueryGenerator_SimpleGrid::PostLoad()
-{
-	if (VerNum < EnvQueryGeneratorVersion::DataProviders)
-	{
-		Radius.Convert(this, GridSize);
-		Density.Convert(this, SpaceBetween);
-	}
-
-	Super::PostLoad();
 }
 
 void UEnvQueryGenerator_SimpleGrid::GenerateItems(FEnvQueryInstance& QueryInstance) const
diff --git a/Engine/Source/Runtime/AIModule/Private/EnvironmentQuery/Tests/EnvQueryTest_Pathfinding.cpp b/Engine/Source/Runtime/AIModule/Private/EnvironmentQuery/Tests/EnvQueryTest_Pathfinding.cpp
--- a/Engine/Source/Runtime/AIModule/Private/EnvironmentQuery/Tests/EnvQueryTest_Pathfinding.cpp
+++ b/Engine/Source/Runtime/AIModule/Private/EnvironmentQuery/Tests/EnvQueryTest_Pathfinding.cpp
@@ -19,12 +19,6 @@ UEnvQueryTest_Pathfinding::UEnvQueryTest_Pathfinding(const FObjectInitializer& O
 	SkipUnreachable.DefaultValue = true;
 	FloatValueMin.DefaultValue = 1000.0f;
 	FloatValueMax.DefaultValue = 1000.0f;
-
-	// keep deprecated properties initialized
-	PathToItem.Value = true;
-	DiscardUnreachable.Value = true;
-	FloatFilterMin.Value = 1000.0f;
-	FloatFilterMax.Value = 1000.0f;
 }
 
 void UEnvQueryTest_Pathfinding::RunTest(FEnvQueryInstance& QueryInstance) const
@@ -158,12 +152,6 @@ void UEnvQueryTest_Pathfinding::PostEditChangeProperty(FPropertyChangedEvent& Pr
 
 void UEnvQueryTest_Pathfinding::PostLoad()
 {
-	if (VerNum < EnvQueryTestVersion::DataProviders)
-	{
-		PathToItem.Convert(this, PathFromContext);
-		DiscardUnreachable.Convert(this, SkipUnreachable);
-	}
-
 	Super::PostLoad();
 	
 	SetWorkOnFloatValues(TestMode != EEnvTestPathfinding::PathExist);
diff --git a/Engine/Source/Runtime/AIModule/Private/EnvironmentQuery/Tests/EnvQueryTest_Trace.cpp b/Engine/Source/Runtime/AIModule/Private/EnvironmentQuery/Tests/EnvQueryTest_Trace.cpp
--- a/Engine/Source/Runtime/AIModule/Private/EnvironmentQuery/Tests/EnvQueryTest_Trace.cpp
+++ b/Engine/Source/Runtime/AIModule/Private/EnvironmentQuery/Tests/EnvQueryTest_Trace.cpp
@@ -179,16 +179,4 @@ bool UEnvQueryTest_Trace::RunCapsuleTraceFrom(const FVector& ItemPos, const FVec
 	return bHit;
 }
 
-void UEnvQueryTest_Trace::PostLoad()
-{
-	if (VerNum < EnvQueryTestVersion::DataProviders)
-	{
-		TraceToItem.Convert(this, TraceFromContext);
-		ItemOffsetZ.Convert(this, ItemHeightOffset);
-		ContextOffsetZ.Convert(this, ContextHeightOffset);
-	}
-
-	Super::PostLoad();
-}
-
 #undef LOCTEXT_NAMESPACE
