compilation fix

diff --git a/Engine/Source/Runtime/AIModule/Classes/EnvironmentQuery/EnvQueryTypes.h b/Engine/Source/Runtime/AIModule/Classes/EnvironmentQuery/EnvQueryTypes.h
--- a/Engine/Source/Runtime/AIModule/Classes/EnvironmentQuery/EnvQueryTypes.h
+++ b/Engine/Source/Runtime/AIModule/Classes/EnvironmentQuery/EnvQueryTypes.h
@@ -884,7 +884,7 @@ namespace FEQSHelpers
 // BEGIN DEPRECATED SUPPORT
 
 USTRUCT()
-struct FEnvFloatParam_DEPRECATED
+struct AIMODULE_API FEnvFloatParam_DEPRECATED
 {
 	GENERATED_USTRUCT_BODY();
 
@@ -901,7 +901,7 @@ struct FEnvFloatParam_DEPRECATED
 };
 
 USTRUCT()
-struct FEnvIntParam_DEPRECATED
+struct AIMODULE_API FEnvIntParam_DEPRECATED
 {
 	GENERATED_USTRUCT_BODY();
 
@@ -918,7 +918,7 @@ struct FEnvIntParam_DEPRECATED
 };
 
 USTRUCT()
-struct FEnvBoolParam_DEPRECATED
+struct AIMODULE_API FEnvBoolParam_DEPRECATED
 {
 	GENERATED_USTRUCT_BODY();
 
@@ -934,6 +934,10 @@ struct FEnvBoolParam_DEPRECATED
 	void Convert(UObject* Owner, FAIDataProviderBoolValue& ValueProvider);
 };
 
+struct DEPRECATED(4.8, "FEnvFloatParam is deprecated in 4.8 and was replaced with FAIDataProviderFloatValue. Please use that type instead.") FEnvFloatParam;
+struct DEPRECATED(4.8, "FEnvIntParam is deprecated in 4.8 and was replaced with FAIDataProviderIntValue. Please use that type instead.") FEnvIntParam;
+struct DEPRECATED(4.8, "FEnvBoolParam is deprecated in 4.8 and was replaced with FAIDataProviderBoolValue. Please use that type instead.") FEnvBoolParam;
+
 USTRUCT()
 struct AIMODULE_API FEnvFloatParam : public FEnvFloatParam_DEPRECATED
 {
@@ -952,10 +956,6 @@ struct AIMODULE_API FEnvBoolParam : public FEnvBoolParam_DEPRECATED
 	GENERATED_USTRUCT_BODY();
 };
 
-struct DEPRECATED(4.8, "FEnvFloatParam is deprecated in 4.8 and was replaced with FAIDataProviderFloatValue. Please use that type instead.") FEnvFloatParam;
-struct DEPRECATED(4.8, "FEnvIntParam is deprecated in 4.8 and was replaced with FAIDataProviderIntValue. Please use that type instead.") FEnvIntParam;
-struct DEPRECATED(4.8, "FEnvBoolParam is deprecated in 4.8 and was replaced with FAIDataProviderBoolValue. Please use that type instead.") FEnvBoolParam;
-
 // END DEPRECATED SUPPORT
 
 UCLASS(Abstract)
