Fixed PathFollowingComponent's handling of last path-to-actor segment to care about goal actor being on navmesh #UE4

diff --git a/Engine/Source/Runtime/AIModule/Private/Navigation/PathFollowingComponent.cpp b/Engine/Source/Runtime/AIModule/Private/Navigation/PathFollowingComponent.cpp
--- a/Engine/Source/Runtime/AIModule/Private/Navigation/PathFollowingComponent.cpp
+++ b/Engine/Source/Runtime/AIModule/Private/Navigation/PathFollowingComponent.cpp
@@ -120,6 +120,9 @@ FString GetPathDescHelper(FNavPathSharedPtr Path)
 
 void UPathFollowingComponent::OnPathEvent(FNavigationPath* InvalidatedPath, ENavPathEvent::Type Event)
 {
+	const static UEnum* NavPathEventEnum = FindObject<UEnum>(ANY_PACKAGE, TEXT("ENavPathEvent"));
+	UE_VLOG(GetOwner(), LogPathFollowing, Log, TEXT("OnPathEvent: %s"), *NavPathEventEnum->GetEnumName(Event));
+
 	if (InvalidatedPath == nullptr || Path.Get() != InvalidatedPath)
 	{
 		return;
@@ -783,27 +786,31 @@ void UPathFollowingComponent::UpdatePathSegment()
 		{
 			// use goal actor for end of last path segment
 			// UNLESS it's partial path (can't reach goal)
-			if (DestinationActor.IsValid() && !Path->IsPartial())
+			if (DestinationActor.IsValid() && Path->IsPartial() == false)
 			{
 				const FVector AgentLocation = DestinationAgent ? DestinationAgent->GetNavAgentLocation() : DestinationActor->GetActorLocation();
+				// note that the condition below requires GoalLocation to be in world space.
 				const FVector GoalLocation = FRotationTranslationMatrix(DestinationActor->GetActorRotation(), AgentLocation).TransformPosition(MoveOffset);
-				CurrentDestination.Set(NULL, GoalLocation);
+				FVector HitLocation;
+
+				if (MyNavData == nullptr //|| MyNavData->DoesNodeContainLocation(Path->GetPathPoints().Last().NodeRef, GoalLocation))
+					|| (FVector::DistSquared(GoalLocation, *CurrentDestination) > SMALL_NUMBER &&  MyNavData->Raycast(CurrentLocation, GoalLocation, HitLocation, nullptr) == false))
+				{
+					CurrentDestination.Set(NULL, GoalLocation);
 
-				UE_VLOG(this, LogPathFollowing, Log, TEXT("Moving directly to move goal rather than following last path segment"));
-				UE_VLOG_LOCATION(this, LogPathFollowing, VeryVerbose, GoalLocation, 30, FColor::Green, TEXT("Last-segment-to-actor"));
-				UE_VLOG_SEGMENT(this, LogPathFollowing, VeryVerbose, CurrentLocation, GoalLocation, FColor::Green, TEXT_EMPTY);
+					UE_VLOG(this, LogPathFollowing, Log, TEXT("Moving directly to move goal rather than following last path segment"));
+					UE_VLOG_LOCATION(this, LogPathFollowing, VeryVerbose, GoalLocation, 30, FColor::Green, TEXT("Last-segment-to-actor"));
+					UE_VLOG_SEGMENT(this, LogPathFollowing, VeryVerbose, CurrentLocation, GoalLocation, FColor::Green, TEXT_EMPTY);
+				}
 			}
 
 			UpdateMoveFocus();
 		}
-		else
+		// check if current move segment is finished
+		else if (HasReachedCurrentTarget(CurrentLocation))
 		{
-			// check if current move segment is finished
-			if (HasReachedCurrentTarget(CurrentLocation))
-			{
-				OnSegmentFinished();
-				SetNextMoveSegment();
-			}
+			OnSegmentFinished();
+			SetNextMoveSegment();
 		}
 	}
 
