composite EQS generator is now hidden in editor, as it's supposed to be "code only" class

diff --git a/Engine/Source/Runtime/AIModule/Classes/EnvironmentQuery/Generators/EnvQueryGenerator_Composite.h b/Engine/Source/Runtime/AIModule/Classes/EnvironmentQuery/Generators/EnvQueryGenerator_Composite.h
--- a/Engine/Source/Runtime/AIModule/Classes/EnvironmentQuery/Generators/EnvQueryGenerator_Composite.h
+++ b/Engine/Source/Runtime/AIModule/Classes/EnvironmentQuery/Generators/EnvQueryGenerator_Composite.h
@@ -12,7 +12,7 @@ struct FEnvQueryInstance;
  * Currently it's available only from code
  */
 
-UCLASS(MinimalAPI)
+UCLASS(MinimalAPI, meta = (HiddenNode))
 class UEnvQueryGenerator_Composite : public UEnvQueryGenerator
 {
 	GENERATED_UCLASS_BODY()
