Fixed Clang compile error

diff --git a/Engine/Source/Runtime/AIModule/Private/EnvironmentQuery/EnvQueryTest.cpp b/Engine/Source/Runtime/AIModule/Private/EnvironmentQuery/EnvQueryTest.cpp
--- a/Engine/Source/Runtime/AIModule/Private/EnvironmentQuery/EnvQueryTest.cpp
+++ b/Engine/Source/Runtime/AIModule/Private/EnvironmentQuery/EnvQueryTest.cpp
@@ -356,7 +356,7 @@ void UEnvQueryTest::UpdatePreviewData()
 
 	const float* AllSamples[] = { SamplesLinear, SamplesSquare, SamplesLinear, SamplesConstant };
 
-	int32 EquationType = (ScoringEquation >= ARRAY_COUNT(AllSamples)) ? EEnvTestScoreEquation::Constant : ScoringEquation;
+	int32 EquationType = (ScoringEquation >= ARRAY_COUNT(AllSamples)) ? EEnvTestScoreEquation::Constant : (EEnvTestScoreEquation::Type)ScoringEquation;
 	if (TestPurpose == EEnvTestPurpose::Filter)
 	{
 		// pure filtering won't apply any scoring, draw flat line
