disabled deprecation warnings for EQS structs until macro gets fixed
#ue4

diff --git a/Engine/Source/Runtime/AIModule/Classes/EnvironmentQuery/EnvQueryTypes.h b/Engine/Source/Runtime/AIModule/Classes/EnvironmentQuery/EnvQueryTypes.h
--- a/Engine/Source/Runtime/AIModule/Classes/EnvironmentQuery/EnvQueryTypes.h
+++ b/Engine/Source/Runtime/AIModule/Classes/EnvironmentQuery/EnvQueryTypes.h
@@ -934,9 +934,11 @@ struct AIMODULE_API FEnvBoolParam_DEPRECATED
 	void Convert(UObject* Owner, FAIDataProviderBoolValue& ValueProvider);
 };
 
+/* DISABLED FOR NOW, DEPRECATED MACRO DOESN'T WORK WITH USTRUCTS
 struct DEPRECATED(4.8, "FEnvFloatParam is deprecated in 4.8 and was replaced with FAIDataProviderFloatValue. Please use that type instead.") FEnvFloatParam;
 struct DEPRECATED(4.8, "FEnvIntParam is deprecated in 4.8 and was replaced with FAIDataProviderIntValue. Please use that type instead.") FEnvIntParam;
 struct DEPRECATED(4.8, "FEnvBoolParam is deprecated in 4.8 and was replaced with FAIDataProviderBoolValue. Please use that type instead.") FEnvBoolParam;
+*/
 
 USTRUCT()
 struct AIMODULE_API FEnvFloatParam : public FEnvFloatParam_DEPRECATED
