behavior tree service nodes will continue from remaining tick time unless it's first activation
#ue4 UE-7397

diff --git a/Engine/Source/Runtime/AIModule/Classes/BehaviorTree/BTAuxiliaryNode.h b/Engine/Source/Runtime/AIModule/Classes/BehaviorTree/BTAuxiliaryNode.h
--- a/Engine/Source/Runtime/AIModule/Classes/BehaviorTree/BTAuxiliaryNode.h
+++ b/Engine/Source/Runtime/AIModule/Classes/BehaviorTree/BTAuxiliaryNode.h
@@ -73,6 +73,9 @@ protected:
 	/** sets next tick time */
 	void SetNextTickTime(uint8* NodeMemory, float RemainingTime) const;
 
+	/** gets remaining time for next tick */
+	float GetNextTickRemainingTime(uint8* NodeMemory) const;
+
 	//----------------------------------------------------------------------//
 	// DEPRECATED
 	//----------------------------------------------------------------------//
diff --git a/Engine/Source/Runtime/AIModule/Private/BehaviorTree/BTAuxiliaryNode.cpp b/Engine/Source/Runtime/AIModule/Private/BehaviorTree/BTAuxiliaryNode.cpp
--- a/Engine/Source/Runtime/AIModule/Private/BehaviorTree/BTAuxiliaryNode.cpp
+++ b/Engine/Source/Runtime/AIModule/Private/BehaviorTree/BTAuxiliaryNode.cpp
@@ -76,6 +76,17 @@ void UBTAuxiliaryNode::SetNextTickTime(uint8* NodeMemory, float RemainingTime) c
 	}
 }
 
+float UBTAuxiliaryNode::GetNextTickRemainingTime(uint8* NodeMemory) const
+{
+	if (bTickIntervals)
+	{
+		FBTAuxiliaryMemory* AuxMemory = GetSpecialNodeMemory<FBTAuxiliaryMemory>(NodeMemory);
+		return FMath::Max(0.0f, AuxMemory->NextTickRemainingTime);
+	}
+
+	return 0.0f;
+}
+
 void UBTAuxiliaryNode::DescribeRuntimeValues(const UBehaviorTreeComponent& OwnerComp, uint8* NodeMemory, EBTDescriptionVerbosity::Type Verbosity, TArray<FString>& Values) const
 {
 	Super::DescribeRuntimeValues(OwnerComp, NodeMemory, Verbosity, Values);
diff --git a/Engine/Source/Runtime/AIModule/Private/BehaviorTree/BTService.cpp b/Engine/Source/Runtime/AIModule/Private/BehaviorTree/BTService.cpp
--- a/Engine/Source/Runtime/AIModule/Private/BehaviorTree/BTService.cpp
+++ b/Engine/Source/Runtime/AIModule/Private/BehaviorTree/BTService.cpp
@@ -36,7 +36,11 @@ void UBTService::NotifyParentActivation(FBehaviorTreeSearchData& SearchData)
 
 			if (bNotifyTick)
 			{
-				ServiceNodeOb->ScheduleNextTick(NodeMemory);
+				const float RemainingTime = GetNextTickRemainingTime(NodeMemory);
+				if (RemainingTime <= 0.0f)
+				{
+					ServiceNodeOb->ScheduleNextTick(NodeMemory);
+				}
 			}
 
 			if (bNotifyOnSearch)
