added option to restart service's tick interval on each activation

diff --git a/Engine/Source/Runtime/AIModule/Classes/BehaviorTree/BTService.h b/Engine/Source/Runtime/AIModule/Classes/BehaviorTree/BTService.h
--- a/Engine/Source/Runtime/AIModule/Classes/BehaviorTree/BTService.h
+++ b/Engine/Source/Runtime/AIModule/Classes/BehaviorTree/BTService.h
@@ -59,6 +59,10 @@ protected:
 	UPROPERTY(Category = Service, EditAnywhere, AdvancedDisplay)
 	uint32 bCallTickOnSearchStart : 1;
 
+	/** if set, next tick time will be always reset to service's interval when node is activated */
+	UPROPERTY(Category = Service, EditAnywhere, AdvancedDisplay)
+	uint32 bRestartTimerOnEachActivation : 1;
+
 	/** if set, service will be notified about search entering underlying branch */
 	uint32 bNotifyOnSearch : 1;
 
diff --git a/Engine/Source/Runtime/AIModule/Private/BehaviorTree/BTService.cpp b/Engine/Source/Runtime/AIModule/Private/BehaviorTree/BTService.cpp
--- a/Engine/Source/Runtime/AIModule/Private/BehaviorTree/BTService.cpp
+++ b/Engine/Source/Runtime/AIModule/Private/BehaviorTree/BTService.cpp
@@ -9,6 +9,7 @@ UBTService::UBTService(const FObjectInitializer& ObjectInitializer) : Super(Obje
 	bNotifyOnSearch = true;
 	bTickIntervals = true;
 	bCallTickOnSearchStart = false;
+	bRestartTimerOnEachActivation = false;
 
 	Interval = 0.5f;
 	RandomDeviation = 0.1f;
@@ -36,7 +37,7 @@ void UBTService::NotifyParentActivation(FBehaviorTreeSearchData& SearchData)
 
 			if (bNotifyTick)
 			{
-				const float RemainingTime = GetNextTickRemainingTime(NodeMemory);
+				const float RemainingTime = bRestartTimerOnEachActivation ? 0.0f : GetNextTickRemainingTime(NodeMemory);
 				if (RemainingTime <= 0.0f)
 				{
 					ServiceNodeOb->ScheduleNextTick(NodeMemory);
