Fixed a silly bug making UAIPerceptionComponent::GetActorsPerception useless #UE4

diff --git a/Engine/Source/Runtime/AIModule/Private/Perception/AIPerceptionComponent.cpp b/Engine/Source/Runtime/AIModule/Private/Perception/AIPerceptionComponent.cpp
--- a/Engine/Source/Runtime/AIModule/Private/Perception/AIPerceptionComponent.cpp
+++ b/Engine/Source/Runtime/AIModule/Private/Perception/AIPerceptionComponent.cpp
@@ -586,7 +586,7 @@ void UAIPerceptionComponent::GetPerceivedActors(TSubclassOf<UAISense> SenseToUse
 bool UAIPerceptionComponent::GetActorsPerception(AActor* Actor, FActorPerceptionBlueprintInfo& Info)
 {
 	bool bInfoFound = false;
-	if (Actor && Actor->IsPendingKillPending())
+	if (Actor != nullptr && Actor->IsPendingKillPending() == false)
 	{
 		const FActorPerceptionInfo* PerceivedInfo = GetActorInfo(*Actor);
 		if (PerceivedInfo)
