restored EQS pathing grid generator

diff --git a/Engine/Source/Runtime/AIModule/Classes/EnvironmentQuery/Generators/EnvQueryGenerator_PathingGrid.h b/Engine/Source/Runtime/AIModule/Classes/EnvironmentQuery/Generators/EnvQueryGenerator_PathingGrid.h
--- a/Engine/Source/Runtime/AIModule/Classes/EnvironmentQuery/Generators/EnvQueryGenerator_PathingGrid.h
+++ b/Engine/Source/Runtime/AIModule/Classes/EnvironmentQuery/Generators/EnvQueryGenerator_PathingGrid.h
@@ -11,21 +11,22 @@ class UNavigationQueryFilter;
  *  with paths to/from context no further than given limit
  */
 
-UCLASS(meta = (DeprecatedNode, DeprecationMessage = "This class is now deprecated, please use SimpleGrid generator with Pathfinding Batch test"))
+UCLASS()
 class UEnvQueryGenerator_PathingGrid : public UEnvQueryGenerator_SimpleGrid
 {
 	GENERATED_UCLASS_BODY()
 
-	UPROPERTY()
-	FAIDataProviderFloatValue MaxDistance;
-
-	/** DEPRECATED, please use pathfinding (batch) test instead! */
-	UPROPERTY(VisibleDefaultsOnly, Category=Generator)
+	/** pathfinding direction */
+	UPROPERTY(EditDefaultsOnly, Category = Generator)
 	FAIDataProviderBoolValue PathToItem;
 
-	/** DEPRECATED, please use pathfinding (batch) test instead! */
-	UPROPERTY(VisibleDefaultsOnly, Category = Generator)
+	/** navigation filter to use in pathfinding */
+	UPROPERTY(EditDefaultsOnly, Category = Generator)
 	TSubclassOf<UNavigationQueryFilter> NavigationFilter;
 
-	virtual void PostLoad() override;
+	/** multiplier for max distance between point and context */
+	UPROPERTY(EditDefaultsOnly, Category = Pathfinding, AdvancedDisplay)
+	FAIDataProviderFloatValue ScanRangeMultiplier;
+
+	virtual void ProjectAndFilterNavPoints(TArray<FNavLocation>& Points, FEnvQueryInstance& QueryInstance) const override;
 };
diff --git a/Engine/Source/Runtime/AIModule/Classes/EnvironmentQuery/Generators/EnvQueryGenerator_ProjectedPoints.h b/Engine/Source/Runtime/AIModule/Classes/EnvironmentQuery/Generators/EnvQueryGenerator_ProjectedPoints.h
--- a/Engine/Source/Runtime/AIModule/Classes/EnvironmentQuery/Generators/EnvQueryGenerator_ProjectedPoints.h
+++ b/Engine/Source/Runtime/AIModule/Classes/EnvironmentQuery/Generators/EnvQueryGenerator_ProjectedPoints.h
@@ -9,7 +9,7 @@
 class ANavigationData;
 
 UCLASS(Abstract)
-class UEnvQueryGenerator_ProjectedPoints : public UEnvQueryGenerator
+class AIMODULE_API UEnvQueryGenerator_ProjectedPoints : public UEnvQueryGenerator
 {
 	GENERATED_UCLASS_BODY()
 
@@ -30,8 +30,8 @@ class UEnvQueryGenerator_ProjectedPoints : public UEnvQueryGenerator
 	};
 
 	/** project all points in array and remove those outside navmesh */
-	void ProjectAndFilterNavPoints(TArray<FNavLocation>& Points, FEnvQueryInstance& QueryInstance) const;
+	virtual void ProjectAndFilterNavPoints(TArray<FNavLocation>& Points, FEnvQueryInstance& QueryInstance) const;
 
 	/** store points as generator's result */
-	void StoreNavPoints(const TArray<FNavLocation>& Points, FEnvQueryInstance& QueryInstance) const;
+	virtual void StoreNavPoints(const TArray<FNavLocation>& Points, FEnvQueryInstance& QueryInstance) const;
 };
diff --git a/Engine/Source/Runtime/AIModule/Private/EnvironmentQuery/Generators/EnvQueryGenerator_PathingGrid.cpp b/Engine/Source/Runtime/AIModule/Private/EnvironmentQuery/Generators/EnvQueryGenerator_PathingGrid.cpp
--- a/Engine/Source/Runtime/AIModule/Private/EnvironmentQuery/Generators/EnvQueryGenerator_PathingGrid.cpp
+++ b/Engine/Source/Runtime/AIModule/Private/EnvironmentQuery/Generators/EnvQueryGenerator_PathingGrid.cpp
@@ -1,18 +1,85 @@
 // Copyright 1998-2015 Epic Games, Inc. All Rights Reserved.
 
 #include "AIModulePrivate.h"
+#include "AI/Navigation/RecastNavMesh.h"
 #include "EnvironmentQuery/Generators/EnvQueryGenerator_PathingGrid.h"
 
+#define LOCTEXT_NAMESPACE "EnvQueryGenerator"
+
 UEnvQueryGenerator_PathingGrid::UEnvQueryGenerator_PathingGrid(const FObjectInitializer& ObjectInitializer) : Super(ObjectInitializer)
 {
-	MaxDistance.DefaultValue = 100.0f;
-	SpaceBetween.DefaultValue = 10.0f;
+	GridSize.DefaultValue = 1000.0f;
+	SpaceBetween.DefaultValue = 100.0f;
 	PathToItem.DefaultValue = true;
+	ScanRangeMultiplier.DefaultValue = 1.5f;
+}
+
+namespace PathGridHelpers
+{
+	static bool HasPath(const FRecastDebugPathfindingData& NodePool, const NavNodeRef& NodeRef)
+	{
+		FRecastDebugPathfindingNode SearchKey(NodeRef);
+		const FRecastDebugPathfindingNode* MyNode = NodePool.Nodes.Find(SearchKey);
+		return MyNode != nullptr;
+	}
 }
 
-void UEnvQueryGenerator_PathingGrid::PostLoad()
+void UEnvQueryGenerator_PathingGrid::ProjectAndFilterNavPoints(TArray<FNavLocation>& Points, FEnvQueryInstance& QueryInstance) const
 {
-	Super::PostLoad();
+	Super::ProjectAndFilterNavPoints(Points, QueryInstance);
+#if WITH_RECAST
+	UObject* DataOwner = QueryInstance.Owner.Get();
+	PathToItem.BindData(DataOwner, QueryInstance.QueryID);
+	ScanRangeMultiplier.BindData(DataOwner, QueryInstance.QueryID);
+
+	bool bPathToItem = PathToItem.GetValue();
+	float RangeMultiplierValue = ScanRangeMultiplier.GetValue();
+
+	ARecastNavMesh* NavMeshData = (ARecastNavMesh*)FEQSHelpers::FindNavMeshForQuery(QueryInstance);
+	if (!NavMeshData)
+	{
+		return;
+	}
+
+	TArray<FVector> ContextLocations;
+	QueryInstance.PrepareContext(GenerateAround, ContextLocations);
 
-	GridSize = MaxDistance;
+	TSharedPtr<FNavigationQueryFilter> NavigationFilterOb = (NavigationFilter != nullptr)
+		? UNavigationQueryFilter::GetQueryFilter(NavMeshData, NavigationFilter)->GetCopy()
+		: NavMeshData->GetDefaultQueryFilter()->GetCopy();
+	NavigationFilterOb->SetBacktrackingEnabled(!bPathToItem);
+	
+	{
+		TArray<NavNodeRef> Polys;
+		for (int32 ContextIdx = 0; ContextIdx < ContextLocations.Num() && Points.Num(); ContextIdx++)
+		{
+			float CollectDistanceSq = 0.0f;
+			for (int32 Idx = 0; Idx < Points.Num(); Idx++)
+			{
+				const float TestDistanceSq = FVector::DistSquared(Points[Idx].Location, ContextLocations[ContextIdx]);
+				CollectDistanceSq = FMath::Max(CollectDistanceSq, TestDistanceSq);
+			}
+
+			const float MaxPathDistance = FMath::Sqrt(CollectDistanceSq) * RangeMultiplierValue;
+
+			Polys.Reset();
+
+			FRecastDebugPathfindingData NodePoolData;
+			NodePoolData.Flags = ERecastDebugPathfindingFlags::Basic;
+
+			NavMeshData->GetPolysWithinPathingDistance(ContextLocations[ContextIdx], MaxPathDistance, Polys, NavigationFilterOb, nullptr, &NodePoolData);
+
+			for (int32 Idx = Points.Num() - 1; Idx >= 0; Idx--)
+			{
+				const bool bHasPath = PathGridHelpers::HasPath(NodePoolData, Points[Idx].NodeRef);
+				if (!bHasPath)
+				{
+					Points.RemoveAt(Idx);
+				}
+			}
+		}
+	}
+#endif // WITH_RECAST
 }
+
+#undef LOCTEXT_NAMESPACE
