added support for named params in instant EQS queries
#ue4

diff --git a/Engine/Source/Runtime/AIModule/Classes/EnvironmentQuery/EnvQueryManager.h b/Engine/Source/Runtime/AIModule/Classes/EnvironmentQuery/EnvQueryManager.h
--- a/Engine/Source/Runtime/AIModule/Classes/EnvironmentQuery/EnvQueryManager.h
+++ b/Engine/Source/Runtime/AIModule/Classes/EnvironmentQuery/EnvQueryManager.h
@@ -153,6 +153,12 @@ class AIMODULE_API UEnvQueryManager : public UObject, public FTickableGameObject
 	/** cleanup hooks for map loading */
 	virtual void FinishDestroy() override;
 
+	/** add information for data providers about query instance run independently */
+	void RegisterExternalQuery(TSharedPtr<FEnvQueryInstance> QueryInstance);
+
+	/** clear information about query instance run independently */
+	void UnregisterExternalQuery(TSharedPtr<FEnvQueryInstance> QueryInstance);
+
 	/** list of all known item types */
 	static TArray<TSubclassOf<UEnvQueryItemType> > RegisteredItemTypes;
 
@@ -173,6 +179,9 @@ protected:
 	/** currently running queries */
 	TArray<TSharedPtr<FEnvQueryInstance> > RunningQueries;
 
+	/** queries run independently from manager, mapped here for data providers */
+	TMap<int32, TWeakPtr<FEnvQueryInstance>> ExternalQueries;
+
 	/** cache of instances */
 	UPROPERTY(transient)
 	TArray<FEnvQueryInstanceCache> InstanceCache;
diff --git a/Engine/Source/Runtime/AIModule/Private/EnvironmentQuery/EQSTestingPawn.cpp b/Engine/Source/Runtime/AIModule/Private/EnvironmentQuery/EQSTestingPawn.cpp
--- a/Engine/Source/Runtime/AIModule/Private/EnvironmentQuery/EQSTestingPawn.cpp
+++ b/Engine/Source/Runtime/AIModule/Private/EnvironmentQuery/EQSTestingPawn.cpp
@@ -184,6 +184,7 @@ void AEQSTestingPawn::MakeOneStep()
 		FEnvQueryRequest QueryRequest(QueryTemplate, this);
 		QueryRequest.SetNamedParams(QueryParams);
 		QueryInstance = EQS->PrepareQueryInstance(QueryRequest, QueryingMode);
+		EQS->RegisterExternalQuery(QueryInstance);
 	}
 
 	// possible still not valid 
diff --git a/Engine/Source/Runtime/AIModule/Private/EnvironmentQuery/EnvQueryManager.cpp b/Engine/Source/Runtime/AIModule/Private/EnvironmentQuery/EnvQueryManager.cpp
--- a/Engine/Source/Runtime/AIModule/Private/EnvironmentQuery/EnvQueryManager.cpp
+++ b/Engine/Source/Runtime/AIModule/Private/EnvironmentQuery/EnvQueryManager.cpp
@@ -179,11 +179,14 @@ TSharedPtr<FEnvQueryResult> UEnvQueryManager::RunInstantQuery(const FEnvQueryReq
 		return NULL;
 	}
 
+	RegisterExternalQuery(QueryInstance);
 	while (QueryInstance->IsFinished() == false)
 	{
 		QueryInstance->ExecuteOneStep((double)FLT_MAX);
 	}
 
+	UnregisterExternalQuery(QueryInstance);
+
 	UE_VLOG_EQS(*QueryInstance.Get(), LogEQS, All);
 
 #if USE_EQS_DEBUGGER
@@ -315,6 +318,16 @@ void UEnvQueryManager::OnWorldCleanup()
 	}
 }
 
+void UEnvQueryManager::RegisterExternalQuery(TSharedPtr<FEnvQueryInstance> QueryInstance)
+{
+	ExternalQueries.Add(QueryInstance->QueryID, QueryInstance);
+}
+
+void UEnvQueryManager::UnregisterExternalQuery(TSharedPtr<FEnvQueryInstance> QueryInstance)
+{
+	ExternalQueries.Remove(QueryInstance->QueryID);
+}
+
 namespace EnvQueryTestSort
 {
 	struct FAllMatching
@@ -570,13 +583,25 @@ float UEnvQueryManager::FindNamedParam(int32 QueryId, FName ParamName) const
 {
 	float ParamValue = 0.0f;
 
-	for (int32 QueryIndex = 0; QueryIndex < RunningQueries.Num(); QueryIndex++)
+	const TWeakPtr<FEnvQueryInstance>* QueryInstancePtr = ExternalQueries.Find(QueryId);
+	if (QueryInstancePtr)
 	{
-		const TSharedPtr<FEnvQueryInstance>& QueryInstance = RunningQueries[QueryIndex];
-		if (QueryInstance->QueryID == QueryId)
+		TSharedPtr<FEnvQueryInstance> QueryInstance = (*QueryInstancePtr).Pin();
+		if (QueryInstance.IsValid())
 		{
 			ParamValue = QueryInstance->NamedParams.FindRef(ParamName);
-			break;
+		}
+	}
+	else
+	{
+		for (int32 QueryIndex = 0; QueryIndex < RunningQueries.Num(); QueryIndex++)
+		{
+			const TSharedPtr<FEnvQueryInstance>& QueryInstance = RunningQueries[QueryIndex];
+			if (QueryInstance->QueryID == QueryId)
+			{
+				ParamValue = QueryInstance->NamedParams.FindRef(ParamName);
+				break;
+			}
 		}
 	}
 
