compilation fix

diff --git a/Engine/Source/Runtime/AIModule/Classes/EnvironmentQuery/EnvQueryTypes.h b/Engine/Source/Runtime/AIModule/Classes/EnvironmentQuery/EnvQueryTypes.h
--- a/Engine/Source/Runtime/AIModule/Classes/EnvironmentQuery/EnvQueryTypes.h
+++ b/Engine/Source/Runtime/AIModule/Classes/EnvironmentQuery/EnvQueryTypes.h
@@ -3,6 +3,7 @@
 #pragma once
 
 #include "EnvironmentQuery/Items/EnvQueryItemType.h"
+#include "EnvironmentQuery/EnvQueryContext.h"
 #include "DataProviders/AIDataProvider.h"
 #include "EnvQueryTypes.generated.h"
 
