added selecting random results from EQS query in two ranges: best 5% and best 25%
#ue4 UE-11276

diff --git a/Engine/Source/Developer/GameplayDebugger/Private/GameplayDebuggingComponent.cpp b/Engine/Source/Developer/GameplayDebugger/Private/GameplayDebuggingComponent.cpp
--- a/Engine/Source/Developer/GameplayDebugger/Private/GameplayDebuggingComponent.cpp
+++ b/Engine/Source/Developer/GameplayDebugger/Private/GameplayDebuggingComponent.cpp
@@ -1443,7 +1443,7 @@ FPrimitiveSceneProxy* UGameplayDebuggingComponent::CreateSceneProxy()
 				ViewFlagName = TEXT("DebugAI");
 			}
 #endif
-			CompositeProxy->AddChild(new FEQSSceneProxy(this, ViewFlagName, false, CurrentLocalData.SolidSpheres, CurrentLocalData.Texts));
+			CompositeProxy->AddChild(new FEQSSceneProxy(this, ViewFlagName, CurrentLocalData.SolidSpheres, CurrentLocalData.Texts));
 		}
 	}
 #endif // USE_EQS_DEBUGGER
diff --git a/Engine/Source/Runtime/AIModule/Classes/BehaviorTree/Tasks/BTTask_RunEQSQuery.h b/Engine/Source/Runtime/AIModule/Classes/BehaviorTree/Tasks/BTTask_RunEQSQuery.h
--- a/Engine/Source/Runtime/AIModule/Classes/BehaviorTree/Tasks/BTTask_RunEQSQuery.h
+++ b/Engine/Source/Runtime/AIModule/Classes/BehaviorTree/Tasks/BTTask_RunEQSQuery.h
@@ -29,6 +29,10 @@ class AIMODULE_API UBTTask_RunEQSQuery : public UBTTask_BlackboardBase
 	UPROPERTY(Category=Node, EditAnywhere)
 	TArray<FEnvNamedValue> QueryParams;
 
+	/** determines which item will be stored (All = only first matching) */
+	UPROPERTY(Category = Node, EditAnywhere)
+	TEnumAsByte<EEnvQueryRunMode::Type> RunMode;
+
 	virtual EBTNodeResult::Type ExecuteTask(UBehaviorTreeComponent& OwnerComp, uint8* NodeMemory) override;
 	virtual EBTNodeResult::Type AbortTask(UBehaviorTreeComponent& OwnerComp, uint8* NodeMemory) override;
 
diff --git a/Engine/Source/Runtime/AIModule/Classes/EnvironmentQuery/EQSQueryResultSourceInterface.h b/Engine/Source/Runtime/AIModule/Classes/EnvironmentQuery/EQSQueryResultSourceInterface.h
--- a/Engine/Source/Runtime/AIModule/Classes/EnvironmentQuery/EQSQueryResultSourceInterface.h
+++ b/Engine/Source/Runtime/AIModule/Classes/EnvironmentQuery/EQSQueryResultSourceInterface.h
@@ -20,4 +20,5 @@ class IEQSQueryResultSourceInterface
 	// debugging
 	virtual bool GetShouldDebugDrawLabels() const { return true; }
 	virtual bool GetShouldDrawFailedItems() const { return true; }
+	virtual float GetHighlightRangePct() const { return 1.0f; }
 };
diff --git a/Engine/Source/Runtime/AIModule/Classes/EnvironmentQuery/EQSRenderingComponent.h b/Engine/Source/Runtime/AIModule/Classes/EnvironmentQuery/EQSRenderingComponent.h
--- a/Engine/Source/Runtime/AIModule/Classes/EnvironmentQuery/EQSRenderingComponent.h
+++ b/Engine/Source/Runtime/AIModule/Classes/EnvironmentQuery/EQSRenderingComponent.h
@@ -13,8 +13,8 @@ struct FEnvQueryInstance;
 class AIMODULE_API FEQSSceneProxy : public FDebugRenderSceneProxy
 {
 public:
-	FEQSSceneProxy(const UPrimitiveComponent* InComponent, const FString& ViewFlagName = TEXT("DebugAI"), bool bDrawOnlyWhenSelected = true);
-	FEQSSceneProxy(const UPrimitiveComponent* InComponent, const FString& ViewFlagName, bool bDrawOnlyWhenSelected, const TArray<FSphere>& Spheres, const TArray<FText3d>& Texts);
+	FEQSSceneProxy(const UPrimitiveComponent* InComponent, const FString& ViewFlagName = TEXT("DebugAI"));
+	FEQSSceneProxy(const UPrimitiveComponent* InComponent, const FString& ViewFlagName, const TArray<FSphere>& Spheres, const TArray<FText3d>& Texts);
 
 	virtual void DrawDebugLabels(UCanvas* Canvas, APlayerController*) override;
 	
@@ -22,14 +22,14 @@ public:
 
 #if  USE_EQS_DEBUGGER 
 	static void CollectEQSData(const UPrimitiveComponent* InComponent, const IEQSQueryResultSourceInterface* QueryDataSource, TArray<FSphere>& Spheres, TArray<FText3d>& Texts, TArray<EQSDebug::FDebugHelper>& DebugItems);
-	static void CollectEQSData(const FEnvQueryResult* ResultItems, const FEnvQueryInstance* QueryInstance, TArray<FSphere>& Spheres, TArray<FText3d>& Texts, bool ShouldDrawFailedItems, TArray<EQSDebug::FDebugHelper>& DebugItems);
+	static void CollectEQSData(const FEnvQueryResult* ResultItems, const FEnvQueryInstance* QueryInstance, float HighlightRangePct, bool ShouldDrawFailedItems, TArray<FSphere>& Spheres, TArray<FText3d>& Texts, TArray<EQSDebug::FDebugHelper>& DebugItems);
 #endif
 private:
 	FEnvQueryResult QueryResult;	
 	// can be 0
 	AActor* ActorOwner;
 	const IEQSQueryResultSourceInterface* QueryDataSource;
-	bool bDrawOnlyWhenSelected;
+	uint32 bDrawOnlyWhenSelected : 1;
 
 	static const FVector ItemDrawRadius;
 
@@ -42,7 +42,7 @@ class AIMODULE_API UEQSRenderingComponent : public UPrimitiveComponent
 	GENERATED_UCLASS_BODY()
 
 	FString DrawFlagName;
-	bool bDrawOnlyWhenSelected;
+	uint32 bDrawOnlyWhenSelected : 1;
 
 #if  USE_EQS_DEBUGGER || ENABLE_VISUAL_LOG
 	EQSDebug::FQueryData DebugData;
diff --git a/Engine/Source/Runtime/AIModule/Classes/EnvironmentQuery/EQSTestingPawn.h b/Engine/Source/Runtime/AIModule/Classes/EnvironmentQuery/EQSTestingPawn.h
--- a/Engine/Source/Runtime/AIModule/Classes/EnvironmentQuery/EQSTestingPawn.h
+++ b/Engine/Source/Runtime/AIModule/Classes/EnvironmentQuery/EQSTestingPawn.h
@@ -10,6 +10,14 @@
 class UEnvQuery;
 class UEQSRenderingComponent;
 
+UENUM()
+enum class EEnvQueryHightlightMode : uint8
+{
+	All,
+	Best5Pct UMETA(DisplayName = "Best 5%"),
+	Best25Pct UMETA(DisplayName = "Best 25%"),
+};
+
 /** this class is abstract even though it's perfectly functional on its own.
  *	The reason is to stop it from showing as valid player pawn type when configuring 
  *	project's game mode. */
@@ -31,7 +39,10 @@ class AIMODULE_API AEQSTestingPawn : public ACharacter, public IEQSQueryResultSo
 	UPROPERTY(Category=EQS, EditAnywhere)
 	int32 StepToDebugDraw;
 
-	UPROPERTY(Category=EQS, EditAnywhere)
+	UPROPERTY(Category = EQS, EditAnywhere)
+	EEnvQueryHightlightMode HighlightMode;
+
+	UPROPERTY(Category = EQS, EditAnywhere)
 	uint32 bDrawLabels:1;
 
 	UPROPERTY(Category=EQS, EditAnywhere)
@@ -76,6 +87,7 @@ public:
 
 	virtual bool GetShouldDebugDrawLabels() const override { return bDrawLabels; }
 	virtual bool GetShouldDrawFailedItems() const override{ return bDrawFailedItems; }
+	virtual float GetHighlightRangePct() const;
 	// IEQSQueryResultSourceInterface end
 
 	void RunEQSQuery();
diff --git a/Engine/Source/Runtime/AIModule/Classes/EnvironmentQuery/EnvQueryTypes.h b/Engine/Source/Runtime/AIModule/Classes/EnvironmentQuery/EnvQueryTypes.h
--- a/Engine/Source/Runtime/AIModule/Classes/EnvironmentQuery/EnvQueryTypes.h
+++ b/Engine/Source/Runtime/AIModule/Classes/EnvironmentQuery/EnvQueryTypes.h
@@ -125,8 +125,10 @@ namespace EEnvQueryRunMode
 {
 	enum Type
 	{
-		SingleResult,		// weight scoring first, try conditions from best result and stop after first item pass
-		AllMatching,		// conditions first (limit set of items), weight scoring later
+		SingleResult	UMETA(Tooltip="Pick first item with the best score"),
+		RandomBest5Pct	UMETA(Tooltip="Pick random item with score 95% .. 100% of max", DisplayName="Random Best 5%"),
+		RandomBest25Pct	UMETA(Tooltip="Pick random item with score 75% .. 100% of max", DisplayName="Random Best 25%"),
+		AllMatching		UMETA(Tooltip="Get all items that match conditions"),
 	};
 }
 
@@ -469,9 +471,6 @@ struct AIMODULE_API FEnvQueryOptionInstance
 	/** type of generated items */
 	TSubclassOf<UEnvQueryItemType> ItemType;
 
-	/** if set, items will be shuffled after tests */
-	uint32 bShuffleItems : 1;
-
 	/** if set, items will be shuffled after tests */
 	uint32 bHasNavLocations : 1;
 
@@ -624,8 +623,8 @@ protected:
 	/** sort all scores, from highest to lowest */
 	void SortScores();
 
-	/** pick one of items with highest score */
-	void PickBestItem();
+	/** pick one of items with score equal or higher than specified */
+	void PickBestItem(float MinScore);
 
 	/** discard all items but one */
 	void PickSingleItem(int32 ItemIndex);
diff --git a/Engine/Source/Runtime/AIModule/Classes/EnvironmentQuery/Tests/EnvQueryTest_Random.h b/Engine/Source/Runtime/AIModule/Classes/EnvironmentQuery/Tests/EnvQueryTest_Random.h
--- a/Engine/Source/Runtime/AIModule/Classes/EnvironmentQuery/Tests/EnvQueryTest_Random.h
+++ b/Engine/Source/Runtime/AIModule/Classes/EnvironmentQuery/Tests/EnvQueryTest_Random.h
@@ -4,7 +4,7 @@
 #include "EnvironmentQuery/EnvQueryTest.h"
 #include "EnvQueryTest_Random.generated.h"
 
-UCLASS(MinimalAPI)
+UCLASS(MinimalAPI, meta=(DeprecatedNode, DeprecationMessage = "This class is now deprecated, please use RunMode supporting random results instead."))
 class UEnvQueryTest_Random : public UEnvQueryTest
 {
 	GENERATED_UCLASS_BODY()
diff --git a/Engine/Source/Runtime/AIModule/Private/BehaviorTree/Tasks/BTTask_RunEQSQuery.cpp b/Engine/Source/Runtime/AIModule/Private/BehaviorTree/Tasks/BTTask_RunEQSQuery.cpp
--- a/Engine/Source/Runtime/AIModule/Private/BehaviorTree/Tasks/BTTask_RunEQSQuery.cpp
+++ b/Engine/Source/Runtime/AIModule/Private/BehaviorTree/Tasks/BTTask_RunEQSQuery.cpp
@@ -31,7 +31,7 @@ EBTNodeResult::Type UBTTask_RunEQSQuery::ExecuteTask(UBehaviorTreeComponent& Own
 	QueryRequest.SetNamedParams(QueryParams);
 
 	FBTEnvQueryTaskMemory* MyMemory = (FBTEnvQueryTaskMemory*)NodeMemory;
-	MyMemory->RequestID = QueryRequest.Execute(EEnvQueryRunMode::SingleResult, this, &UBTTask_RunEQSQuery::OnQueryFinished);
+	MyMemory->RequestID = QueryRequest.Execute(RunMode, this, &UBTTask_RunEQSQuery::OnQueryFinished);
 
 	const bool bValid = (MyMemory->RequestID >= 0);
 	if (bValid)
diff --git a/Engine/Source/Runtime/AIModule/Private/EnvironmentQuery/EQSRenderingComponent.cpp b/Engine/Source/Runtime/AIModule/Private/EnvironmentQuery/EQSRenderingComponent.cpp
--- a/Engine/Source/Runtime/AIModule/Private/EnvironmentQuery/EQSRenderingComponent.cpp
+++ b/Engine/Source/Runtime/AIModule/Private/EnvironmentQuery/EQSRenderingComponent.cpp
@@ -27,11 +27,10 @@ namespace FEQSRenderingHelper
 //----------------------------------------------------------------------//
 const FVector FEQSSceneProxy::ItemDrawRadius(30,30,30);
 
-FEQSSceneProxy::FEQSSceneProxy(const UPrimitiveComponent* InComponent, const FString& InViewFlagName, bool bInDrawOnlyWhenSelected) 
+FEQSSceneProxy::FEQSSceneProxy(const UPrimitiveComponent* InComponent, const FString& InViewFlagName) 
 	: FDebugRenderSceneProxy(InComponent)
 	, ActorOwner(NULL)
 	, QueryDataSource(NULL)
-	, bDrawOnlyWhenSelected(bInDrawOnlyWhenSelected)
 {
 	DrawType = SolidAndWireMeshes;
 	TextWithoutShadowDistance = 1500;
@@ -54,17 +53,19 @@ FEQSSceneProxy::FEQSSceneProxy(const UPrimitiveComponent* InComponent, const FSt
 		}
 	}
 
+	const UEQSRenderingComponent* MyRenderComp = Cast<const UEQSRenderingComponent>(InComponent);
+	bDrawOnlyWhenSelected = MyRenderComp && MyRenderComp->bDrawOnlyWhenSelected;
+
 #if  USE_EQS_DEBUGGER 
 	TArray<EQSDebug::FDebugHelper> DebugItems;
 	FEQSSceneProxy::CollectEQSData(InComponent, QueryDataSource, Spheres, Texts, DebugItems);
 #endif
 }
 
-FEQSSceneProxy::FEQSSceneProxy(const UPrimitiveComponent* InComponent, const FString& InViewFlagName, bool bInDrawOnlyWhenSelected, const TArray<FSphere>& InSpheres, const TArray<FText3d>& InTexts)
+FEQSSceneProxy::FEQSSceneProxy(const UPrimitiveComponent* InComponent, const FString& InViewFlagName, const TArray<FSphere>& InSpheres, const TArray<FText3d>& InTexts)
 	: FDebugRenderSceneProxy(InComponent)
 	, ActorOwner(NULL)
 	, QueryDataSource(NULL)
-	, bDrawOnlyWhenSelected(bInDrawOnlyWhenSelected)
 {
 	DrawType = SolidAndWireMeshes;
 	TextWithoutShadowDistance = 1500;
@@ -89,6 +90,9 @@ FEQSSceneProxy::FEQSSceneProxy(const UPrimitiveComponent* InComponent, const FSt
 			return;
 		}
 	}
+
+	const UEQSRenderingComponent* MyRenderComp = Cast<const UEQSRenderingComponent>(InComponent);
+	bDrawOnlyWhenSelected = MyRenderComp && MyRenderComp->bDrawOnlyWhenSelected;
 }
 
 #if  USE_EQS_DEBUGGER 
@@ -112,7 +116,9 @@ void FEQSSceneProxy::CollectEQSData(const UPrimitiveComponent* InComponent, cons
 	const FEnvQueryResult* ResultItems = QueryDataSource->GetQueryResult();
 	const FEnvQueryInstance* QueryInstance = QueryDataSource->GetQueryInstance();
 
-	FEQSSceneProxy::CollectEQSData(ResultItems, QueryInstance, Spheres, Texts, QueryDataSource->GetShouldDrawFailedItems(), DebugItems);
+	FEQSSceneProxy::CollectEQSData(ResultItems, QueryInstance,
+		QueryDataSource->GetHighlightRangePct(), QueryDataSource->GetShouldDrawFailedItems(),
+		Spheres, Texts, DebugItems);
 
 	if (ActorOwner && Spheres.Num() > EQSMaxItemsDrawn)
 	{
@@ -123,7 +129,7 @@ void FEQSSceneProxy::CollectEQSData(const UPrimitiveComponent* InComponent, cons
 	}
 }
 
-void FEQSSceneProxy::CollectEQSData(const FEnvQueryResult* ResultItems, const FEnvQueryInstance* QueryInstance, TArray<FSphere>& Spheres, TArray<FText3d>& Texts, bool ShouldDrawFailedItems, TArray<EQSDebug::FDebugHelper>& DebugItems)
+void FEQSSceneProxy::CollectEQSData(const FEnvQueryResult* ResultItems, const FEnvQueryInstance* QueryInstance, float HighlightRangePct, bool ShouldDrawFailedItems, TArray<FSphere>& Spheres, TArray<FText3d>& Texts, TArray<EQSDebug::FDebugHelper>& DebugItems)
 {
 	if (ResultItems == NULL)
 	{
@@ -148,7 +154,7 @@ void FEQSSceneProxy::CollectEQSData(const FEnvQueryResult* ResultItems, const FE
 
 	float MinScore = 0.f;
 	float MaxScore = -BIG_NUMBER;
-	if (bUseMidResults)
+	if (bUseMidResults || HighlightRangePct < 1.0f)
 	{
 		const FEnvQueryItem* ItemInfo = Items.GetData();
 		for (int32 ItemIndex = 0; ItemIndex < Items.Num(); ItemIndex++, ItemInfo++)
@@ -160,7 +166,9 @@ void FEQSSceneProxy::CollectEQSData(const FEnvQueryResult* ResultItems, const FE
 			}
 		}
 	}
+	
 	const float ScoreNormalizer = bUseMidResults && (MaxScore != MinScore) ? 1.f / (MaxScore - MinScore) : 1.f;
+	const float HighlightThreshold = (HighlightRangePct < 1.0f) ? (MaxScore * HighlightRangePct) : FLT_MAX;
 
 	if (bSingleItemResult == false)
 	{
@@ -169,12 +177,14 @@ void FEQSSceneProxy::CollectEQSData(const FEnvQueryResult* ResultItems, const FE
 			if (Items[ItemIndex].IsValid())
 			{
 				const float Score = bNoTestsPerformed ? 1 : Items[ItemIndex].Score * ScoreNormalizer;
+				const bool bLowRadius = (HighlightThreshold < FLT_MAX) && (bNoTestsPerformed || (Items[ItemIndex].Score < HighlightThreshold));
+				const float Radius = ItemDrawRadius.X * (bLowRadius ? 0.2f : 1.0f);
 				const FVector Loc = FEQSRenderingHelper::ExtractLocation(ResultItems->ItemType, RawData, Items, ItemIndex);
-				Spheres.Add(FSphere(ItemDrawRadius.X, Loc, bNoTestsPerformed == false 
+				Spheres.Add(FSphere(Radius, Loc, bNoTestsPerformed == false
 					? FLinearColor(FColor::MakeRedToGreenColorFromScalar(Score)) 
 					: FLinearColor(0.2, 1.0, 1.0, 1)));
 
-				DebugItems.Add(EQSDebug::FDebugHelper(Loc, ItemDrawRadius.X));
+				DebugItems.Add(EQSDebug::FDebugHelper(Loc, Radius));
 
 				const FString Label = bNoTestsPerformed ? TEXT("") : FString::Printf(TEXT("%.2f"), Score);
 				Texts.Add(FText3d(Label, Loc, FLinearColor::White));
@@ -186,10 +196,12 @@ void FEQSSceneProxy::CollectEQSData(const FEnvQueryResult* ResultItems, const FE
 		if (Items[0].IsValid())
 		{
 			const float Score = Items[0].Score * ScoreNormalizer;
+			const bool bLowRadius = false;
+			const float Radius = ItemDrawRadius.X * (bLowRadius ? 0.2f : 1.0f);
 			const FVector Loc = FEQSRenderingHelper::ExtractLocation(ResultItems->ItemType, RawData, Items, 0);
-			Spheres.Add(FSphere(ItemDrawRadius.X, Loc, FLinearColor(0.0, 1.0, 0.12, 1)));
+			Spheres.Add(FSphere(Radius, Loc, FLinearColor(0.0, 1.0, 0.12, 1)));
 
-			DebugItems.Add(EQSDebug::FDebugHelper(Loc, ItemDrawRadius.X));
+			DebugItems.Add(EQSDebug::FDebugHelper(Loc, Radius));
 
 			const FString Label = FString::Printf(TEXT("Winner %.2f"), Score);
 			Texts.Add(FText3d(Label, Loc, FLinearColor::White));
@@ -200,10 +212,12 @@ void FEQSSceneProxy::CollectEQSData(const FEnvQueryResult* ResultItems, const FE
 			if (Items[ItemIndex].IsValid())
 			{
 				const float Score = bNoTestsPerformed ? 1 : Items[ItemIndex].Score * ScoreNormalizer;
+				const bool bLowRadius = (HighlightThreshold < FLT_MAX) && (bNoTestsPerformed || (Items[ItemIndex].Score < HighlightThreshold));
+				const float Radius = ItemDrawRadius.X * (bLowRadius ? 0.2f : 1.0f);
 				const FVector Loc = FEQSRenderingHelper::ExtractLocation(ResultItems->ItemType, RawData, Items, ItemIndex);
-				Spheres.Add(FSphere(ItemDrawRadius.X, Loc, FLinearColor(0.0, 0.2, 0.025, 1)));
+				Spheres.Add(FSphere(Radius, Loc, FLinearColor(0.0, 0.2, 0.025, 1)));
 
-				DebugItems.Add(EQSDebug::FDebugHelper(Loc, ItemDrawRadius.X));
+				DebugItems.Add(EQSDebug::FDebugHelper(Loc, Radius));
 				
 				const FString Label = bNoTestsPerformed ? TEXT("") : FString::Printf(TEXT("%.2f"), Score);
 				Texts.Add(FText3d(Label, Loc, FLinearColor::White));
@@ -227,10 +241,12 @@ void FEQSSceneProxy::CollectEQSData(const FEnvQueryResult* ResultItems, const FE
 			}
 
 			const float Score = bNoTestsPerformed ? 1 : Items[ItemIndex].Score * ScoreNormalizer;
+			const bool bLowRadius = (HighlightThreshold < FLT_MAX) && (bNoTestsPerformed || (Items[ItemIndex].Score < HighlightThreshold));
+			const float Radius = ItemDrawRadius.X * (bLowRadius ? 0.2f : 1.0f);
 			const FVector Loc = FEQSRenderingHelper::ExtractLocation(QueryInstance->ItemType, InstanceDebugData.RawData, DebugQueryItems, ItemIndex);
-			Spheres.Add(FSphere(ItemDrawRadius.X, Loc, FLinearColor(0.0, 0.0, 0.6, 0.6)));
+			Spheres.Add(FSphere(Radius, Loc, FLinearColor(0.0, 0.0, 0.6, 0.6)));
 
-			auto& DebugHelper = DebugItems[DebugItems.Add(EQSDebug::FDebugHelper(Loc, ItemDrawRadius.X))];
+			auto& DebugHelper = DebugItems[DebugItems.Add(EQSDebug::FDebugHelper(Loc, Radius))];
 			DebugHelper.AdditionalInformation = Details[ItemIndex].FailedDescription;
 			if (Details[ItemIndex].FailedTestIndex != INDEX_NONE)
 			{
@@ -291,10 +307,10 @@ FPrimitiveSceneProxy* UEQSRenderingComponent::CreateSceneProxy()
 #if  USE_EQS_DEBUGGER || ENABLE_VISUAL_LOG
 	if (DebugData.SolidSpheres.Num() > 0 || DebugData.Texts.Num() > 0)
 	{
-		return new FEQSSceneProxy(this, DrawFlagName, bDrawOnlyWhenSelected, DebugData.SolidSpheres, DebugData.Texts);
+		return new FEQSSceneProxy(this, DrawFlagName, DebugData.SolidSpheres, DebugData.Texts);
 	}
 #endif
-	return new FEQSSceneProxy(this, DrawFlagName, bDrawOnlyWhenSelected);
+	return new FEQSSceneProxy(this, DrawFlagName);
 }
 
 FBoxSphereBounds UEQSRenderingComponent::CalcBounds(const FTransform& LocalToWorld) const
diff --git a/Engine/Source/Runtime/AIModule/Private/EnvironmentQuery/EQSTestingPawn.cpp b/Engine/Source/Runtime/AIModule/Private/EnvironmentQuery/EQSTestingPawn.cpp
--- a/Engine/Source/Runtime/AIModule/Private/EnvironmentQuery/EQSTestingPawn.cpp
+++ b/Engine/Source/Runtime/AIModule/Private/EnvironmentQuery/EQSTestingPawn.cpp
@@ -88,6 +88,11 @@ AEQSTestingPawn::AEQSTestingPawn(const FObjectInitializer& ObjectInitializer)
 #endif // WITH_EDITOR
 }
 
+float AEQSTestingPawn::GetHighlightRangePct() const
+{
+	return (HighlightMode == EEnvQueryHightlightMode::Best25Pct) ? 0.75f : (HighlightMode == EEnvQueryHightlightMode::Best5Pct) ? 0.95f : 1.0f;
+}
+
 void AEQSTestingPawn::OnEditorSelectionChanged(UObject* NewSelection)
 {
 	/* This is totally busted and should not depend on gameplay debugger.
@@ -246,7 +251,6 @@ void AEQSTestingPawn::PostEditChangeProperty( FPropertyChangedEvent& PropertyCha
 	static const FName NAME_QueryTemplate = GET_MEMBER_NAME_CHECKED(AEQSTestingPawn, QueryTemplate);
 	static const FName NAME_StepToDebugDraw = GET_MEMBER_NAME_CHECKED(AEQSTestingPawn, StepToDebugDraw);
 	static const FName NAME_QueryParams = GET_MEMBER_NAME_CHECKED(AEQSTestingPawn, QueryParams);
-	static const FName NAME_DrawFailedItems = GET_MEMBER_NAME_CHECKED(AEQSTestingPawn, bDrawFailedItems);
 	static const FName NAME_ShouldBeVisibleInGame = GET_MEMBER_NAME_CHECKED(AEQSTestingPawn, bShouldBeVisibleInGame);
 	static const FName NAME_QueryingMode = GET_MEMBER_NAME_CHECKED(AEQSTestingPawn, QueryingMode);
 
@@ -267,7 +271,8 @@ void AEQSTestingPawn::PostEditChangeProperty( FPropertyChangedEvent& PropertyCha
 			StepToDebugDraw  = FMath::Clamp(StepToDebugDraw, 0, StepResults.Num() - 1 );
 			UpdateDrawing();
 		}
-		else if (PropName == NAME_DrawFailedItems)
+		else if (PropName == GET_MEMBER_NAME_CHECKED(AEQSTestingPawn, bDrawFailedItems) ||
+			PropName == GET_MEMBER_NAME_CHECKED(AEQSTestingPawn, HighlightMode))
 		{
 			UpdateDrawing();
 		}
diff --git a/Engine/Source/Runtime/AIModule/Private/EnvironmentQuery/EnvQueryDebugHelpers.cpp b/Engine/Source/Runtime/AIModule/Private/EnvironmentQuery/EnvQueryDebugHelpers.cpp
--- a/Engine/Source/Runtime/AIModule/Private/EnvironmentQuery/EnvQueryDebugHelpers.cpp
+++ b/Engine/Source/Runtime/AIModule/Private/EnvironmentQuery/EnvQueryDebugHelpers.cpp
@@ -43,7 +43,7 @@ void UEnvQueryDebugHelpers::QueryToDebugData(FEnvQueryInstance& Query, EQSDebug:
 	// step 1: data for rendering component
 	EQSLocalData.Reset();
 
-	FEQSSceneProxy::CollectEQSData(&Query, &Query, EQSLocalData.SolidSpheres, EQSLocalData.Texts, true, EQSLocalData.RenderDebugHelpers);
+	FEQSSceneProxy::CollectEQSData(&Query, &Query, 1.0f, true, EQSLocalData.SolidSpheres, EQSLocalData.Texts, EQSLocalData.RenderDebugHelpers);
 
 	// step 2: detailed scoring data for HUD
 	const int32 MaxDetailedItems = 10;
diff --git a/Engine/Source/Runtime/AIModule/Private/EnvironmentQuery/EnvQueryInstance.cpp b/Engine/Source/Runtime/AIModule/Private/EnvironmentQuery/EnvQueryInstance.cpp
--- a/Engine/Source/Runtime/AIModule/Private/EnvironmentQuery/EnvQueryInstance.cpp
+++ b/Engine/Source/Runtime/AIModule/Private/EnvironmentQuery/EnvQueryInstance.cpp
@@ -417,20 +417,6 @@ void FEnvQueryInstance::NormalizeScores()
 void FEnvQueryInstance::SortScores()
 {
 	const int32 NumItems = Items.Num();
-	if (Options[OptionIndex].bShuffleItems)
-	{
-		for (int32 ItemIndex = 0; ItemIndex < NumItems; ItemIndex++)
-		{
-			const int32 Idx1 = FMath::RandHelper(NumItems);
-			const int32 Idx2 = FMath::RandHelper(NumItems);
-			Items.Swap(Idx1, Idx2);
-
-#if USE_EQS_DEBUGGER
-			ItemDetails.Swap(Idx1, Idx2);
-#endif
-		}
-	}
-
 #if USE_EQS_DEBUGGER
 	struct FSortHelperForDebugData
 	{
@@ -472,13 +458,13 @@ void FEnvQueryInstance::StripRedundantData()
 	Items.SetNum(NumValidItems);
 }
 
-void FEnvQueryInstance::PickBestItem()
+void FEnvQueryInstance::PickBestItem(float MinScore)
 {
-	// find first valid item with score worse than best one
+	// find first valid item with score worse than best range
 	int32 NumBestItems = NumValidItems;
 	for (int32 ItemIndex = 1; ItemIndex < NumValidItems; ItemIndex++)
 	{
-		if (Items[ItemIndex].Score < Items[0].Score)
+		if (Items[ItemIndex].Score < MinScore)
 		{
 			NumBestItems = ItemIndex;
 			break;
@@ -536,9 +522,15 @@ void FEnvQueryInstance::FinalizeQuery()
 			if (bFoundSingleResult == false && bPassOnSingleResult == false)
 			{
 				SortScores();
-				PickBestItem();
+				PickBestItem(Items[0].Score);
 			}
 		}
+		else if (Mode == EEnvQueryRunMode::RandomBest5Pct || Mode == EEnvQueryRunMode::RandomBest25Pct)
+		{
+			SortScores();
+			const float ScoreRangePct = (Mode == EEnvQueryRunMode::RandomBest5Pct) ? 0.95f : 0.75f;
+			PickBestItem(Items[0].Score * ScoreRangePct);
+		}
 		else
 		{
 			SortScores();
diff --git a/Engine/Source/Runtime/AIModule/Private/EnvironmentQuery/EnvQueryManager.cpp b/Engine/Source/Runtime/AIModule/Private/EnvironmentQuery/EnvQueryManager.cpp
--- a/Engine/Source/Runtime/AIModule/Private/EnvironmentQuery/EnvQueryManager.cpp
+++ b/Engine/Source/Runtime/AIModule/Private/EnvironmentQuery/EnvQueryManager.cpp
@@ -507,6 +507,8 @@ TSharedPtr<FEnvQueryInstance> UEnvQueryManager::CreateQueryInstance(const UEnvQu
 					SortedTests.Sort(EnvQueryTestSort::FSingleResult(HighestCost));
 					break;
 
+				case EEnvQueryRunMode::RandomBest5Pct:
+				case EEnvQueryRunMode::RandomBest25Pct:
 				case EEnvQueryRunMode::AllMatching:
 					SortedTests.Sort(EnvQueryTestSort::FAllMatching());
 					break;
@@ -540,7 +542,6 @@ void UEnvQueryManager::CreateOptionInstance(UEnvQueryOption* OptionTemplate, con
 	FEnvQueryOptionInstance OptionInstance;
 	OptionInstance.Generator = OptionTemplate->Generator;
 	OptionInstance.ItemType = OptionTemplate->Generator->ItemType;
-	OptionInstance.bShuffleItems = true;
 	OptionInstance.bHasNavLocations = false;
 
 	OptionInstance.Tests.AddZeroed(SortedTests.Num());
@@ -548,15 +549,6 @@ void UEnvQueryManager::CreateOptionInstance(UEnvQueryOption* OptionTemplate, con
 	{
 		UEnvQueryTest* TestOb = SortedTests[TestIndex];
 		OptionInstance.Tests[TestIndex] = TestOb;
-
-		// HACK!  TODO: Is this the correct replacement here?  or should it check just if SCORING ONLY?
-		// always randomize when asking for single result
-		// otherwise, can skip randomization if test wants to score every item
-		if ((TestOb->TestPurpose != EEnvTestPurpose::Filter) && // We ARE scoring, regardless of whether or not we're filtering
-			(Instance.Mode != EEnvQueryRunMode::SingleResult))
-		{
-			OptionInstance.bShuffleItems = false;
-		}
 	}
 
 	DEC_MEMORY_STAT_BY(STAT_AI_EQS_InstanceMemory, Instance.Options.GetAllocatedSize());
