fix for missing WITH_RECAST to allow recast to be compiled out

diff --git a/Engine/Source/Runtime/AIModule/Private/EnvironmentQuery/Generators/EnvQueryGenerator_PathingGrid.cpp b/Engine/Source/Runtime/AIModule/Private/EnvironmentQuery/Generators/EnvQueryGenerator_PathingGrid.cpp
--- a/Engine/Source/Runtime/AIModule/Private/EnvironmentQuery/Generators/EnvQueryGenerator_PathingGrid.cpp
+++ b/Engine/Source/Runtime/AIModule/Private/EnvironmentQuery/Generators/EnvQueryGenerator_PathingGrid.cpp
@@ -14,6 +14,7 @@ UEnvQueryGenerator_PathingGrid::UEnvQueryGenerator_PathingGrid(const FObjectInit
 	ScanRangeMultiplier.DefaultValue = 1.5f;
 }
 
+#if WITH_RECAST
 namespace PathGridHelpers
 {
 	static bool HasPath(const FRecastDebugPathfindingData& NodePool, const NavNodeRef& NodeRef)
@@ -23,6 +24,7 @@ namespace PathGridHelpers
 		return MyNode != nullptr;
 	}
 }
+#endif
 
 void UEnvQueryGenerator_PathingGrid::ProjectAndFilterNavPoints(TArray<FNavLocation>& Points, FEnvQueryInstance& QueryInstance) const
 {
