composite EQS generator can be used from editor, custom display names for existing generators
#ue4

diff --git a/Engine/Source/Editor/AIGraph/Classes/AIGraphTypes.h b/Engine/Source/Editor/AIGraph/Classes/AIGraphTypes.h
--- a/Engine/Source/Editor/AIGraph/Classes/AIGraphTypes.h
+++ b/Engine/Source/Editor/AIGraph/Classes/AIGraphTypes.h
@@ -15,6 +15,7 @@ struct AIGRAPH_API FGraphNodeClassData
 	FString ToString() const;
 	FString GetClassName() const;
 	FString GetCategory() const;
+	FString GetDisplayName() const;
 	UClass* GetClass(bool bSilent = false);
 	bool IsAbstract() const;
 
diff --git a/Engine/Source/Editor/AIGraph/Private/AIGraphTypes.cpp b/Engine/Source/Editor/AIGraph/Private/AIGraphTypes.cpp
--- a/Engine/Source/Editor/AIGraph/Private/AIGraphTypes.cpp
+++ b/Engine/Source/Editor/AIGraph/Private/AIGraphTypes.cpp
@@ -30,6 +30,12 @@ FGraphNodeClassData::FGraphNodeClassData(const FString& InAssetName, const FStri
 
 FString FGraphNodeClassData::ToString() const
 {
+	FString ShortName = GetDisplayName();
+	if (!ShortName.IsEmpty())
+	{
+		return ShortName;
+	}
+
 	UClass* MyClass = Class.Get();
 	if (MyClass)
 	{
@@ -57,6 +63,11 @@ FString FGraphNodeClassData::GetClassName() const
 	return Class.IsValid() ? Class->GetName() : ClassName;
 }
 
+FString FGraphNodeClassData::GetDisplayName() const
+{
+	return Class.IsValid() ? Class->GetMetaData(TEXT("DisplayName")) : FString();
+}
+
 FString FGraphNodeClassData::GetCategory() const
 {
 	return Class.IsValid() ? Class->GetMetaData(TEXT("Category")) : Category;
diff --git a/Engine/Source/Editor/EnvironmentQueryEditor/Classes/EnvironmentQueryGraphNode_Option.h b/Engine/Source/Editor/EnvironmentQueryEditor/Classes/EnvironmentQueryGraphNode_Option.h
--- a/Engine/Source/Editor/EnvironmentQueryEditor/Classes/EnvironmentQueryGraphNode_Option.h
+++ b/Engine/Source/Editor/EnvironmentQueryEditor/Classes/EnvironmentQueryGraphNode_Option.h
@@ -19,6 +19,7 @@ class UEnvironmentQueryGraphNode_Option : public UEnvironmentQueryGraphNode
 	void CreateAddTestSubMenu(class FMenuBuilder& MenuBuilder, UEdGraph* Graph) const;
 
 	void CalculateWeights();
+	void UpdateNodeData();
 
 protected:
 
diff --git a/Engine/Source/Editor/EnvironmentQueryEditor/Private/EnvironmentQueryGraph.cpp b/Engine/Source/Editor/EnvironmentQueryEditor/Private/EnvironmentQueryGraph.cpp
--- a/Engine/Source/Editor/EnvironmentQueryEditor/Private/EnvironmentQueryGraph.cpp
+++ b/Engine/Source/Editor/EnvironmentQueryEditor/Private/EnvironmentQueryGraph.cpp
@@ -60,8 +60,10 @@ void UEnvironmentQueryGraph::UpdateAsset(int32 UpdateFlags)
 			UEnvironmentQueryGraphNode_Option* OptionNode = Cast<UEnvironmentQueryGraphNode_Option>(MyPin->LinkedTo[Idx]->GetOwningNode());
 			if (OptionNode)
 			{
+				OptionNode->UpdateNodeData();
+
 				UEnvQueryOption* OptionInstance = Cast<UEnvQueryOption>(OptionNode->NodeInstance);
-				if (OptionInstance != NULL)
+				if (OptionInstance && OptionInstance->Generator)
 				{
 					OptionInstance->Tests.Reset();
 
diff --git a/Engine/Source/Editor/EnvironmentQueryEditor/Private/EnvironmentQueryGraphNode_Option.cpp b/Engine/Source/Editor/EnvironmentQueryEditor/Private/EnvironmentQueryGraphNode_Option.cpp
--- a/Engine/Source/Editor/EnvironmentQueryEditor/Private/EnvironmentQueryGraphNode_Option.cpp
+++ b/Engine/Source/Editor/EnvironmentQueryEditor/Private/EnvironmentQueryGraphNode_Option.cpp
@@ -137,4 +137,17 @@ void UEnvironmentQueryGraphNode_Option::CalculateWeights()
 	}
 }
 
+void UEnvironmentQueryGraphNode_Option::UpdateNodeData()
+{
+	UEnvQueryOption* OptionInstance = Cast<UEnvQueryOption>(NodeInstance);
+	UEnvQueryGenerator_Composite* CompositeGenerator = OptionInstance ? Cast<UEnvQueryGenerator_Composite>(OptionInstance->Generator) : nullptr;
+	if (CompositeGenerator)
+	{
+		CompositeGenerator->VerifyItemTypes();
+
+		ErrorMessage = CompositeGenerator->bHasMatchingItemType ? TEXT("") : LOCTEXT("NestedGeneratorMismatch", "Nested generators must work on exactly the same item types!").ToString();
+	}
+}
+
+
 #undef LOCTEXT_NAMESPACE
diff --git a/Engine/Source/Runtime/AIModule/Classes/EnvironmentQuery/EnvQueryGenerator.h b/Engine/Source/Runtime/AIModule/Classes/EnvironmentQuery/EnvQueryGenerator.h
--- a/Engine/Source/Runtime/AIModule/Classes/EnvironmentQuery/EnvQueryGenerator.h
+++ b/Engine/Source/Runtime/AIModule/Classes/EnvironmentQuery/EnvQueryGenerator.h
@@ -16,7 +16,7 @@ namespace EnvQueryGeneratorVersion
 	const int32 Latest = DataProviders;
 }
 
-UCLASS(Abstract, meta=(Category="Generators"))
+UCLASS(EditInlineNew, Abstract, meta = (Category = "Generators"))
 class AIMODULE_API UEnvQueryGenerator : public UEnvQueryNode
 {
 	GENERATED_UCLASS_BODY()
diff --git a/Engine/Source/Runtime/AIModule/Classes/EnvironmentQuery/Generators/EnvQueryGenerator_ActorsOfClass.h b/Engine/Source/Runtime/AIModule/Classes/EnvironmentQuery/Generators/EnvQueryGenerator_ActorsOfClass.h
--- a/Engine/Source/Runtime/AIModule/Classes/EnvironmentQuery/Generators/EnvQueryGenerator_ActorsOfClass.h
+++ b/Engine/Source/Runtime/AIModule/Classes/EnvironmentQuery/Generators/EnvQueryGenerator_ActorsOfClass.h
@@ -7,7 +7,7 @@
 class AActor;
 class UEnvQueryContext;
 
-UCLASS()	
+UCLASS(meta = (DisplayName = "Actors Of Class"))
 class UEnvQueryGenerator_ActorsOfClass : public UEnvQueryGenerator
 {
 	GENERATED_UCLASS_BODY()
diff --git a/Engine/Source/Runtime/AIModule/Classes/EnvironmentQuery/Generators/EnvQueryGenerator_Composite.h b/Engine/Source/Runtime/AIModule/Classes/EnvironmentQuery/Generators/EnvQueryGenerator_Composite.h
--- a/Engine/Source/Runtime/AIModule/Classes/EnvironmentQuery/Generators/EnvQueryGenerator_Composite.h
+++ b/Engine/Source/Runtime/AIModule/Classes/EnvironmentQuery/Generators/EnvQueryGenerator_Composite.h
@@ -9,17 +9,22 @@ struct FEnvQueryInstance;
 
 /**
  * Composite generator allows using multiple generators in single query option
- * Currently it's available only from code
+ * All child generators must produce exactly the same item type!
  */
 
-UCLASS(MinimalAPI, meta = (HiddenNode))
-class UEnvQueryGenerator_Composite : public UEnvQueryGenerator
+UCLASS(meta = (DisplayName = "Composite"))
+class AIMODULE_API UEnvQueryGenerator_Composite : public UEnvQueryGenerator
 {
 	GENERATED_UCLASS_BODY()
 
-	UPROPERTY()
+	UPROPERTY(EditDefaultsOnly, Instanced, Category = Generator)
 	TArray<UEnvQueryGenerator*> Generators;
 
+	UPROPERTY()
+	bool bHasMatchingItemType;
+
 	virtual void GenerateItems(FEnvQueryInstance& QueryInstance) const override;
 	virtual FText GetDescriptionTitle() const override;
+
+	void VerifyItemTypes();
 };
diff --git a/Engine/Source/Runtime/AIModule/Classes/EnvironmentQuery/Generators/EnvQueryGenerator_Donut.h b/Engine/Source/Runtime/AIModule/Classes/EnvironmentQuery/Generators/EnvQueryGenerator_Donut.h
--- a/Engine/Source/Runtime/AIModule/Classes/EnvironmentQuery/Generators/EnvQueryGenerator_Donut.h
+++ b/Engine/Source/Runtime/AIModule/Classes/EnvironmentQuery/Generators/EnvQueryGenerator_Donut.h
@@ -6,7 +6,7 @@
 
 struct FEnvQueryInstance;
 
-UCLASS()
+UCLASS(meta = (DisplayName = "Points: Donut"))
 class AIMODULE_API UEnvQueryGenerator_Donut : public UEnvQueryGenerator_ProjectedPoints
 {
 	GENERATED_UCLASS_BODY()
diff --git a/Engine/Source/Runtime/AIModule/Classes/EnvironmentQuery/Generators/EnvQueryGenerator_OnCircle.h b/Engine/Source/Runtime/AIModule/Classes/EnvironmentQuery/Generators/EnvQueryGenerator_OnCircle.h
--- a/Engine/Source/Runtime/AIModule/Classes/EnvironmentQuery/Generators/EnvQueryGenerator_OnCircle.h
+++ b/Engine/Source/Runtime/AIModule/Classes/EnvironmentQuery/Generators/EnvQueryGenerator_OnCircle.h
@@ -6,7 +6,7 @@
 
 struct FEnvQueryInstance;
 
-UCLASS()	
+UCLASS(meta = (DisplayName = "Points: Circle"))
 class AIMODULE_API UEnvQueryGenerator_OnCircle : public UEnvQueryGenerator_ProjectedPoints
 {
 	GENERATED_UCLASS_BODY()
diff --git a/Engine/Source/Runtime/AIModule/Classes/EnvironmentQuery/Generators/EnvQueryGenerator_PathingGrid.h b/Engine/Source/Runtime/AIModule/Classes/EnvironmentQuery/Generators/EnvQueryGenerator_PathingGrid.h
--- a/Engine/Source/Runtime/AIModule/Classes/EnvironmentQuery/Generators/EnvQueryGenerator_PathingGrid.h
+++ b/Engine/Source/Runtime/AIModule/Classes/EnvironmentQuery/Generators/EnvQueryGenerator_PathingGrid.h
@@ -11,7 +11,7 @@ class UNavigationQueryFilter;
  *  with paths to/from context no further than given limit
  */
 
-UCLASS()
+UCLASS(meta = (DisplayName = "Points: Pathing Grid"))
 class UEnvQueryGenerator_PathingGrid : public UEnvQueryGenerator_SimpleGrid
 {
 	GENERATED_UCLASS_BODY()
diff --git a/Engine/Source/Runtime/AIModule/Classes/EnvironmentQuery/Generators/EnvQueryGenerator_SimpleGrid.h b/Engine/Source/Runtime/AIModule/Classes/EnvironmentQuery/Generators/EnvQueryGenerator_SimpleGrid.h
--- a/Engine/Source/Runtime/AIModule/Classes/EnvironmentQuery/Generators/EnvQueryGenerator_SimpleGrid.h
+++ b/Engine/Source/Runtime/AIModule/Classes/EnvironmentQuery/Generators/EnvQueryGenerator_SimpleGrid.h
@@ -10,7 +10,7 @@ class UEnvQueryContext;
  *  Simple grid, generates points in 2D square around context
  */
 
-UCLASS()
+UCLASS(meta = (DisplayName = "Points: Grid"))
 class UEnvQueryGenerator_SimpleGrid : public UEnvQueryGenerator_ProjectedPoints
 {
 	GENERATED_UCLASS_BODY()
diff --git a/Engine/Source/Runtime/AIModule/Private/EnvironmentQuery/Generators/EnvQueryGenerator_Composite.cpp b/Engine/Source/Runtime/AIModule/Private/EnvironmentQuery/Generators/EnvQueryGenerator_Composite.cpp
--- a/Engine/Source/Runtime/AIModule/Private/EnvironmentQuery/Generators/EnvQueryGenerator_Composite.cpp
+++ b/Engine/Source/Runtime/AIModule/Private/EnvironmentQuery/Generators/EnvQueryGenerator_Composite.cpp
@@ -7,16 +7,20 @@
 
 UEnvQueryGenerator_Composite::UEnvQueryGenerator_Composite(const FObjectInitializer& ObjectInitializer) : Super(ObjectInitializer)
 {
-
+	ItemType = UEnvQueryItemType_Point::StaticClass();
+	bHasMatchingItemType = true;
 }
 
 void UEnvQueryGenerator_Composite::GenerateItems(FEnvQueryInstance& QueryInstance) const 
 {
-	for (int32 Idx = 0; Idx < Generators.Num(); Idx++)
+	if (bHasMatchingItemType)
 	{
-		if (Generators[Idx])
+		for (int32 Idx = 0; Idx < Generators.Num(); Idx++)
 		{
-			Generators[Idx]->GenerateItems(QueryInstance);
+			if (Generators[Idx])
+			{
+				Generators[Idx]->GenerateItems(QueryInstance);
+			}
 		}
 	}
 }
@@ -35,4 +39,31 @@ FText UEnvQueryGenerator_Composite::GetDescriptionTitle() const
 	return Desc;
 };
 
+void UEnvQueryGenerator_Composite::VerifyItemTypes()
+{
+	TSubclassOf<UEnvQueryItemType> CommonItemType = nullptr;
+	bHasMatchingItemType = true;
+
+	for (int32 Idx = 0; Idx < Generators.Num(); Idx++)
+	{
+		if (Generators[Idx])
+		{
+			if (CommonItemType)
+			{
+				if (CommonItemType != Generators[Idx]->ItemType)
+				{
+					bHasMatchingItemType = false;
+					break;
+				}
+			}
+			else
+			{
+				CommonItemType = Generators[Idx]->ItemType;
+			}
+		}
+	}
+
+	ItemType = bHasMatchingItemType ? CommonItemType : UEnvQueryItemType_Point::StaticClass();
+}
+
 #undef LOCTEXT_NAMESPACE
\ No newline at end of file
