missing includes

diff --git a/Engine/Source/Editor/EnvironmentQueryEditor/Private/EnvironmentQueryGraphNode_Option.cpp b/Engine/Source/Editor/EnvironmentQueryEditor/Private/EnvironmentQueryGraphNode_Option.cpp
--- a/Engine/Source/Editor/EnvironmentQueryEditor/Private/EnvironmentQueryGraphNode_Option.cpp
+++ b/Engine/Source/Editor/EnvironmentQueryEditor/Private/EnvironmentQueryGraphNode_Option.cpp
@@ -5,6 +5,7 @@
 #include "SGraphEditorActionMenuAI.h"
 #include "EnvironmentQuery/EnvQueryTest.h"
 #include "EnvironmentQuery/EnvQueryOption.h"
+#include "EnvironmentQuery/Generators/EnvQueryGenerator_Composite.h"
 
 #define LOCTEXT_NAMESPACE "EnvironmentQueryEditor"
 
diff --git a/Engine/Source/Runtime/AIModule/Private/EnvironmentQuery/Generators/EnvQueryGenerator_Composite.cpp b/Engine/Source/Runtime/AIModule/Private/EnvironmentQuery/Generators/EnvQueryGenerator_Composite.cpp
--- a/Engine/Source/Runtime/AIModule/Private/EnvironmentQuery/Generators/EnvQueryGenerator_Composite.cpp
+++ b/Engine/Source/Runtime/AIModule/Private/EnvironmentQuery/Generators/EnvQueryGenerator_Composite.cpp
@@ -1,6 +1,7 @@
 // Copyright 1998-2015 Epic Games, Inc. All Rights Reserved.
 
 #include "AIModulePrivate.h"
+#include "EnvironmentQuery/Items/EnvQueryItemType_Point.h"
 #include "EnvironmentQuery/Generators/EnvQueryGenerator_Composite.h"
 
 #define LOCTEXT_NAMESPACE "EnvQueryGenerator"
