another compilation fix for composite EQS generator

diff --git a/Engine/Source/Runtime/AIModule/Private/EnvironmentQuery/Generators/EnvQueryGenerator_Composite.cpp b/Engine/Source/Runtime/AIModule/Private/EnvironmentQuery/Generators/EnvQueryGenerator_Composite.cpp
--- a/Engine/Source/Runtime/AIModule/Private/EnvironmentQuery/Generators/EnvQueryGenerator_Composite.cpp
+++ b/Engine/Source/Runtime/AIModule/Private/EnvironmentQuery/Generators/EnvQueryGenerator_Composite.cpp
@@ -64,7 +64,14 @@ void UEnvQueryGenerator_Composite::VerifyItemTypes()
 		}
 	}
 
-	ItemType = bHasMatchingItemType ? CommonItemType : UEnvQueryItemType_Point::StaticClass();
+	if (bHasMatchingItemType)
+	{
+		ItemType = CommonItemType;
+	}
+	else
+	{
+		ItemType = UEnvQueryItemType_Point::StaticClass();
+	}
 }
 
 #undef LOCTEXT_NAMESPACE
\ No newline at end of file
