fixed missing properties in instanced behavior tree nodes
#ue4 UE-11777

diff --git a/Engine/Source/Runtime/AIModule/Private/BehaviorTree/BehaviorTreeTypes.cpp b/Engine/Source/Runtime/AIModule/Private/BehaviorTree/BehaviorTreeTypes.cpp
--- a/Engine/Source/Runtime/AIModule/Private/BehaviorTree/BehaviorTreeTypes.cpp
+++ b/Engine/Source/Runtime/AIModule/Private/BehaviorTree/BehaviorTreeTypes.cpp
@@ -17,7 +17,14 @@ void FBehaviorTreeInstance::Initialize(UBehaviorTreeComponent& OwnerComp, UBTCom
 		Node.Services[ServiceIndex]->InitializeInSubtree(OwnerComp, Node.Services[ServiceIndex]->GetNodeMemory<uint8>(*this), InstancedIndex, InitType);
 	}
 
-	Node.InitializeInSubtree(OwnerComp, Node.GetNodeMemory<uint8>(*this), InstancedIndex, InitType);
+	uint8* NodeMemory = Node.GetNodeMemory<uint8>(*this);
+	Node.InitializeInSubtree(OwnerComp, NodeMemory, InstancedIndex, InitType);
+
+	UBTCompositeNode* InstancedComposite = Cast<UBTCompositeNode>(Node.GetNodeInstance(OwnerComp, NodeMemory));
+	if (InstancedComposite)
+	{
+		InstancedComposite->InitializeComposite(Node.GetLastExecutionIndex());
+	}
 
 	for (int32 ChildIndex = 0; ChildIndex < Node.Children.Num(); ChildIndex++)
 	{
@@ -25,7 +32,15 @@ void FBehaviorTreeInstance::Initialize(UBehaviorTreeComponent& OwnerComp, UBTCom
 
 		for (int32 DecoratorIndex = 0; DecoratorIndex < ChildInfo.Decorators.Num(); DecoratorIndex++)
 		{
-			ChildInfo.Decorators[DecoratorIndex]->InitializeInSubtree(OwnerComp, ChildInfo.Decorators[DecoratorIndex]->GetNodeMemory<uint8>(*this), InstancedIndex, InitType);
+			UBTDecorator* DecoratorOb = ChildInfo.Decorators[DecoratorIndex];
+			uint8* DecoratorMemory = DecoratorOb->GetNodeMemory<uint8>(*this);
+			DecoratorOb->InitializeInSubtree(OwnerComp, DecoratorMemory, InstancedIndex, InitType);
+
+			UBTDecorator* InstancedDecoratorOb = Cast<UBTDecorator>(DecoratorOb->GetNodeInstance(OwnerComp, DecoratorMemory));
+			if (InstancedDecoratorOb)
+			{
+				InstancedDecoratorOb->InitializeDecorator(DecoratorOb->GetChildIndex());
+			}
 		}
 
 		if (ChildInfo.ChildComposite)
