fixed ValueChange mode in blackboard decorator
#ue4

diff --git a/Engine/Source/Runtime/AIModule/Private/BehaviorTree/Decorators/BTDecorator_Blackboard.cpp b/Engine/Source/Runtime/AIModule/Private/BehaviorTree/Decorators/BTDecorator_Blackboard.cpp
--- a/Engine/Source/Runtime/AIModule/Private/BehaviorTree/Decorators/BTDecorator_Blackboard.cpp
+++ b/Engine/Source/Runtime/AIModule/Private/BehaviorTree/Decorators/BTDecorator_Blackboard.cpp
@@ -89,14 +89,7 @@ void UBTDecorator_Blackboard::OnBlackboardChange(const UBlackboardComponent& Bla
 				*UBehaviorTreeTypes::DescribeNodeHelper(this),
 				*Blackboard.GetKeyName(ChangedKeyID).ToString());
 
-			// force result Aborted to restart from this decorator
-			// can't use helper function
-
-			const int32 InstanceIdx = BehaviorComp->FindInstanceContainingNode((UBTNode*)GetParentNode());
-			if (InstanceIdx != INDEX_NONE)
-			{
-				BehaviorComp->RequestExecution(GetParentNode(), InstanceIdx, this, GetChildIndex(), EBTNodeResult::Aborted);
-			}
+			BehaviorComp->RequestExecution(this);
 		}
 	}
 }
