Fixed EQS caching invalid query templates (like without any options for example) #UE4

diff --git a/Engine/Source/Runtime/AIModule/Private/EnvironmentQuery/EQSTestingPawn.cpp b/Engine/Source/Runtime/AIModule/Private/EnvironmentQuery/EQSTestingPawn.cpp
--- a/Engine/Source/Runtime/AIModule/Private/EnvironmentQuery/EQSTestingPawn.cpp
+++ b/Engine/Source/Runtime/AIModule/Private/EnvironmentQuery/EQSTestingPawn.cpp
@@ -189,7 +189,10 @@ void AEQSTestingPawn::MakeOneStep()
 		FEnvQueryRequest QueryRequest(QueryTemplate, this);
 		QueryRequest.SetNamedParams(QueryParams);
 		QueryInstance = EQS->PrepareQueryInstance(QueryRequest, QueryingMode);
-		EQS->RegisterExternalQuery(QueryInstance);
+		if (QueryInstance.IsValid())
+		{
+			EQS->RegisterExternalQuery(QueryInstance);
+		}
 	}
 
 	// possible still not valid 
diff --git a/Engine/Source/Runtime/AIModule/Private/EnvironmentQuery/EnvQueryManager.cpp b/Engine/Source/Runtime/AIModule/Private/EnvironmentQuery/EnvQueryManager.cpp
--- a/Engine/Source/Runtime/AIModule/Private/EnvironmentQuery/EnvQueryManager.cpp
+++ b/Engine/Source/Runtime/AIModule/Private/EnvironmentQuery/EnvQueryManager.cpp
@@ -418,9 +418,10 @@ UEnvQuery* UEnvQueryManager::FindQueryTemplate(const FString& QueryName) const
 
 TSharedPtr<FEnvQueryInstance> UEnvQueryManager::CreateQueryInstance(const UEnvQuery* Template, EEnvQueryRunMode::Type RunMode)
 {
-	if (Template == NULL)
+	if (Template == nullptr || Template->Options.Num() == 0)
 	{
-		return NULL;
+		UE_CLOG(Template != nullptr && Template->Options.Num() == 0, LogEQS, Warning, TEXT("Query [%s] doesn't have any valid options!"), *Template->GetName());
+		return nullptr;
 	}
 
 	// try to find entry in cache
@@ -525,11 +526,7 @@ TSharedPtr<FEnvQueryInstance> UEnvQueryManager::CreateQueryInstance(const UEnvQu
 			CreateOptionInstance(LocalOption, SortedTests, *InstanceTemplate);
 		}
 
-		if (InstanceTemplate->Options.Num() == 0)
-		{
-			UE_LOG(LogEQS, Warning, TEXT("Query [%s] doesn't have any valid options!"), *GetNameSafe(LocalTemplate));
-			return NULL;
-		}
+		check(InstanceTemplate->Options.Num() == 0);
 	}
 
 	// create new instance
