Fixed a buggy operator in a check in UEnvQueryManager::CreateQueryInstance #UE4

diff --git a/Engine/Source/Runtime/AIModule/Private/EnvironmentQuery/EnvQueryManager.cpp b/Engine/Source/Runtime/AIModule/Private/EnvironmentQuery/EnvQueryManager.cpp
--- a/Engine/Source/Runtime/AIModule/Private/EnvironmentQuery/EnvQueryManager.cpp
+++ b/Engine/Source/Runtime/AIModule/Private/EnvironmentQuery/EnvQueryManager.cpp
@@ -526,7 +526,7 @@ TSharedPtr<FEnvQueryInstance> UEnvQueryManager::CreateQueryInstance(const UEnvQu
 			CreateOptionInstance(LocalOption, SortedTests, *InstanceTemplate);
 		}
 
-		check(InstanceTemplate->Options.Num() == 0);
+		check(InstanceTemplate->Options.Num() != 0);
 	}
 
 	// create new instance
