Fixed BTTask_PlayAnimation note re-enabling animation BP once done #UE4

diff --git a/Engine/Source/Runtime/AIModule/Classes/BehaviorTree/Tasks/BTTask_PlayAnimation.h b/Engine/Source/Runtime/AIModule/Classes/BehaviorTree/Tasks/BTTask_PlayAnimation.h
--- a/Engine/Source/Runtime/AIModule/Classes/BehaviorTree/Tasks/BTTask_PlayAnimation.h
+++ b/Engine/Source/Runtime/AIModule/Classes/BehaviorTree/Tasks/BTTask_PlayAnimation.h
@@ -30,6 +30,11 @@ class AIMODULE_API UBTTask_PlayAnimation : public UBTTaskNode
 	UPROPERTY()
 	UBehaviorTreeComponent* MyOwnerComp;
 
+	UPROPERTY()
+	USkeletalMeshComponent* CachedSkelMesh;
+
+	EAnimationMode::Type PreviousAnimationMode;
+
 	FTimerDelegate TimerDelegate;
 	FTimerHandle TimerHandle;
 
@@ -38,8 +43,11 @@ class AIMODULE_API UBTTask_PlayAnimation : public UBTTaskNode
 	virtual FString GetStaticDescription() const override;
 
 	void OnAnimationTimerDone();
-
+	
 #if WITH_EDITOR
 	virtual FName GetNodeIconName() const override;
 #endif // WITH_EDITOR
+
+protected:
+	void CleanUp(UBehaviorTreeComponent& OwnerComp);
 };
diff --git a/Engine/Source/Runtime/AIModule/Private/BehaviorTree/Tasks/BTTask_PlayAnimation.cpp b/Engine/Source/Runtime/AIModule/Private/BehaviorTree/Tasks/BTTask_PlayAnimation.cpp
--- a/Engine/Source/Runtime/AIModule/Private/BehaviorTree/Tasks/BTTask_PlayAnimation.cpp
+++ b/Engine/Source/Runtime/AIModule/Private/BehaviorTree/Tasks/BTTask_PlayAnimation.cpp
@@ -20,6 +20,7 @@ UBTTask_PlayAnimation::UBTTask_PlayAnimation(const FObjectInitializer& ObjectIni
 	bNonBlocking = false;
 
 	TimerDelegate = FTimerDelegate::CreateUObject(this, &UBTTask_PlayAnimation::OnAnimationTimerDone);
+	PreviousAnimationMode = EAnimationMode::AnimationBlueprint;
 }
 
 EBTNodeResult::Type UBTTask_PlayAnimation::ExecuteTask(UBehaviorTreeComponent& OwnerComp, uint8* NodeMemory)
@@ -46,6 +47,9 @@ EBTNodeResult::Type UBTTask_PlayAnimation::ExecuteTask(UBehaviorTreeComponent& O
 
 		if (SkelMesh != nullptr)
 		{
+			PreviousAnimationMode = SkelMesh->GetAnimationMode();
+			CachedSkelMesh = SkelMesh;
+
 			SkelMesh->PlayAnimation(AnimationToPlay, bLooping);
 			const float FinishDelay = AnimationToPlay->GetMaxCurrentTime();
 
@@ -80,6 +84,8 @@ EBTNodeResult::Type UBTTask_PlayAnimation::AbortTask(UBehaviorTreeComponent& Own
 
 	TimerHandle.Invalidate();
 
+	CleanUp(OwnerComp);
+
 	return EBTNodeResult::Aborted;
 }
 
@@ -94,10 +100,19 @@ void UBTTask_PlayAnimation::OnAnimationTimerDone()
 {
 	if (MyOwnerComp)
 	{
+		CleanUp(*MyOwnerComp);
 		FinishLatentTask(*MyOwnerComp, EBTNodeResult::Succeeded);
 	}
 }
 
+void UBTTask_PlayAnimation::CleanUp(UBehaviorTreeComponent& OwnerComp)
+{
+	if (CachedSkelMesh != nullptr && PreviousAnimationMode == EAnimationMode::AnimationBlueprint)
+	{
+		CachedSkelMesh->SetAnimationMode(EAnimationMode::AnimationBlueprint);
+	}
+}
+
 #if WITH_EDITOR
 
 FName UBTTask_PlayAnimation::GetNodeIconName() const
