Added a way to run EQS queries from blueprints! #UE4

Change includes creation of a EQS query instance wrapper for purposes of BP API. This wrapper can be expanded by game projects.
This is pass 1. The biggest thing missing is there's no support for named params in EQS BP API yet.

diff --git a/Engine/Source/Runtime/AIModule/Classes/EnvironmentQuery/EnvQueryInstanceBlueprintWrapper.h b/Engine/Source/Runtime/AIModule/Classes/EnvironmentQuery/EnvQueryInstanceBlueprintWrapper.h
new file mode 100644
--- /dev/null
+++ b/Engine/Source/Runtime/AIModule/Classes/EnvironmentQuery/EnvQueryInstanceBlueprintWrapper.h
@@ -0,0 +1,52 @@
+// Copyright 1998-2015 Epic Games, Inc. All Rights Reserved.
+
+#pragma once
+
+#include "EnvQueryTypes.h"
+#include "EnvironmentQuery/EQSQueryResultSourceInterface.h"
+#include "EnvQueryInstanceBlueprintWrapper.generated.h"
+
+UCLASS(BlueprintType, meta = (DisplayName = "EQS Query Instance"))
+class AIMODULE_API UEnvQueryInstanceBlueprintWrapper : public UObject, public IEQSQueryResultSourceInterface
+{
+	GENERATED_BODY()
+
+	DECLARE_DYNAMIC_MULTICAST_DELEGATE_TwoParams(FEQSQueryDoneSignature, UEnvQueryInstanceBlueprintWrapper*, QueryInstance, EEnvQueryStatus::Type, QueryStatus);
+
+protected:
+	UPROPERTY(BlueprintReadOnly, Category = "EQS")
+	int32 QueryID;
+
+	EEnvQueryRunMode::Type RunMode;
+
+	TSharedPtr<FEnvQueryResult> QueryResult;
+
+	UPROPERTY(BlueprintReadOnly, Category = "EQS")
+	TSubclassOf<UEnvQueryItemType> ItemType;
+
+	/** index of query option, that generated items */
+	UPROPERTY(BlueprintReadOnly, Category = "EQS")
+	int32 OptionIndex;
+	
+	UPROPERTY(BlueprintAssignable, Category = "AI|EQS", meta = (DisplayName = "OnQueryFinished"))
+	FEQSQueryDoneSignature OnQueryFinishedEvent;
+
+public:
+	UEnvQueryInstanceBlueprintWrapper(const FObjectInitializer& ObjectInitializer);
+
+	void OnQueryFinished(TSharedPtr<FEnvQueryResult> Result);
+
+	void SetRunMode(EEnvQueryRunMode::Type InRunMode) { RunMode = InRunMode; }
+	void SetQueryID(int32 InQueryID) { QueryID = InQueryID; }
+
+	UFUNCTION(BlueprintPure, Category = "AI|EQS")
+	float GetItemScore(int32 ItemIndex);
+
+	/** return an array filled with resulting actors. Not that it makes sense only if ItemType is a EnvQueryItemType_ActorBase-derived type */
+	UFUNCTION(BlueprintPure, Category = "AI|EQS")
+	TArray<AActor*> GetResultsAsActors();
+
+	/** returns an array of locations generated by the query. If the query generated Actors the the array is filled with their locations */
+	UFUNCTION(BlueprintPure, Category = "AI|EQS")
+	TArray<FVector> GetResultsAsLocations();
+};
\ No newline at end of file
diff --git a/Engine/Source/Runtime/AIModule/Classes/EnvironmentQuery/EnvQueryManager.h b/Engine/Source/Runtime/AIModule/Classes/EnvironmentQuery/EnvQueryManager.h
--- a/Engine/Source/Runtime/AIModule/Classes/EnvironmentQuery/EnvQueryManager.h
+++ b/Engine/Source/Runtime/AIModule/Classes/EnvironmentQuery/EnvQueryManager.h
@@ -3,6 +3,7 @@
 #pragma once
 #include "EnvQueryTypes.h"
 #include "Tickable.h"
+#include "EnvironmentQuery/EnvQueryInstanceBlueprintWrapper.h"
 #include "EnvQueryManager.generated.h"
 
 class UObject;
@@ -164,6 +165,9 @@ class AIMODULE_API UEnvQueryManager : public UObject, public FTickableGameObject
 
 	static UEnvQueryManager* GetCurrent(UWorld* World);
 	static UEnvQueryManager* GetCurrent(UObject* WorldContextObject);
+	
+	UFUNCTION(BlueprintCallable, Category = "AI|EQS", meta = (WorldContext = "WorldContext", AdvancedDisplay = "WrapperClass"))
+	static UEnvQueryInstanceBlueprintWrapper* RunEQSQuery(UObject* WorldContext, UEnvQuery* QueryTemplate, UObject* Querier, TEnumAsByte<EEnvQueryRunMode::Type> RunMode, TSubclassOf<UEnvQueryInstanceBlueprintWrapper> WrapperClass);
 
 #if USE_EQS_DEBUGGER
 	static void NotifyAssetUpdate(UEnvQuery* Query);
diff --git a/Engine/Source/Runtime/AIModule/Classes/EnvironmentQuery/EnvQueryTypes.h b/Engine/Source/Runtime/AIModule/Classes/EnvironmentQuery/EnvQueryTypes.h
--- a/Engine/Source/Runtime/AIModule/Classes/EnvironmentQuery/EnvQueryTypes.h
+++ b/Engine/Source/Runtime/AIModule/Classes/EnvironmentQuery/EnvQueryTypes.h
@@ -125,9 +125,9 @@ namespace EEnvQueryRunMode
 {
 	enum Type
 	{
-		SingleResult	UMETA(Tooltip="Pick first item with the best score"),
-		RandomBest5Pct	UMETA(Tooltip="Pick random item with score 95% .. 100% of max", DisplayName="Random Best 5%"),
-		RandomBest25Pct	UMETA(Tooltip="Pick random item with score 75% .. 100% of max", DisplayName="Random Best 25%"),
+		SingleResult	UMETA(Tooltip="Pick first item with the best score", DisplayName="Single Best Item"),
+		RandomBest5Pct	UMETA(Tooltip="Pick random item with score 95% .. 100% of max", DisplayName="Single Random Item from Best 5%"),
+		RandomBest25Pct	UMETA(Tooltip="Pick random item with score 75% .. 100% of max", DisplayName="Single Random Item from Best 25%"),
 		AllMatching		UMETA(Tooltip="Get all items that match conditions"),
 	};
 }
@@ -352,11 +352,15 @@ struct AIMODULE_API FEnvQueryItem
 
 template <> struct TIsZeroConstructType<FEnvQueryItem> { enum { Value = true }; };
 
+USTRUCT()
 struct AIMODULE_API FEnvQueryResult
 {
+	GENERATED_USTRUCT_BODY()
+
 	TArray<FEnvQueryItem> Items;
 
 	/** type of generated items */
+	UPROPERTY(BlueprintReadOnly, Category = "EQS")
 	TSubclassOf<UEnvQueryItemType> ItemType;
 
 	/** raw data of items */
@@ -364,13 +368,15 @@ struct AIMODULE_API FEnvQueryResult
 
 private:
 	/** query status */
-	TEnumAsByte<EEnvQueryStatus::Type> Status;
+	EEnvQueryStatus::Type Status;
 
 public:
 	/** index of query option, that generated items */
+	UPROPERTY(BlueprintReadOnly, Category = "EQS")
 	int32 OptionIndex;
 
 	/** instance ID */
+	UPROPERTY(BlueprintReadOnly, Category = "EQS")
 	int32 QueryID;
 
 	/** instance owner. Mind that it doesn't have to be the query's "Querier". This is just the object that is responsible for this query instance. */
@@ -382,6 +388,10 @@ public:
 	AActor* GetItemAsActor(int32 Index) const;
 	FVector GetItemAsLocation(int32 Index) const;
 
+	/** note that this function does not strip out the null-actors to not mess up results of GetItemScore(Index) calls*/
+	void GetAllAsActors(TArray<AActor*>& OutActors) const;
+	void GetAllAsLocations(TArray<FVector>& OutLocations) const;
+
 	FEnvQueryResult() : ItemType(NULL), Status(EEnvQueryStatus::Processing), OptionIndex(0) {}
 	FEnvQueryResult(const EEnvQueryStatus::Type& InStatus) : ItemType(NULL), Status(InStatus), OptionIndex(0) {}
 
@@ -392,6 +402,8 @@ public:
 	FORCEINLINE void MarkAsFailed() { Status = EEnvQueryStatus::Failed; }
 	FORCEINLINE void MarkAsFinishedWithoutIssues() { Status = EEnvQueryStatus::Success; }
 	FORCEINLINE void MarkAsOwnerLost() { Status = EEnvQueryStatus::OwnerLost; }
+
+	FORCEINLINE EEnvQueryStatus::Type GetRawStatus() const { return Status; }
 };
 
 
diff --git a/Engine/Source/Runtime/AIModule/Private/EnvironmentQuery/EnvQueryInstanceBlueprintWrapper.cpp b/Engine/Source/Runtime/AIModule/Private/EnvironmentQuery/EnvQueryInstanceBlueprintWrapper.cpp
new file mode 100644
--- /dev/null
+++ b/Engine/Source/Runtime/AIModule/Private/EnvironmentQuery/EnvQueryInstanceBlueprintWrapper.cpp
@@ -0,0 +1,64 @@
+// Copyright 1998-2015 Epic Games, Inc. All Rights Reserved.
+
+#include "AIModulePrivate.h"
+#include "EnvironmentQuery/EnvQueryInstanceBlueprintWrapper.h"
+
+UEnvQueryInstanceBlueprintWrapper::UEnvQueryInstanceBlueprintWrapper(const FObjectInitializer& ObjectInitializer)
+	: Super(ObjectInitializer)
+	, QueryID(INDEX_NONE)
+{
+
+}
+
+void UEnvQueryInstanceBlueprintWrapper::OnQueryFinished(TSharedPtr<FEnvQueryResult> Result)
+{
+	check(Result.IsValid());
+	QueryResult = Result;
+	ItemType = Result->ItemType;
+	OptionIndex = Result->OptionIndex;
+
+	OnQueryFinishedEvent.Broadcast(this, Result->GetRawStatus());
+}
+
+float UEnvQueryInstanceBlueprintWrapper::GetItemScore(int32 ItemIndex)
+{
+	return QueryResult.IsValid() ? QueryResult->GetItemScore(ItemIndex) : -1.f;
+}
+
+TArray<AActor*> UEnvQueryInstanceBlueprintWrapper::GetResultsAsActors()
+{
+	TArray<AActor*> Results;
+
+	if (QueryResult.IsValid() && ItemType->IsChildOf(UEnvQueryItemType_ActorBase::StaticClass()))
+	{
+		if (RunMode != EEnvQueryRunMode::AllMatching)
+		{
+			Results.Add(QueryResult->GetItemAsActor(0));
+		}
+		else
+		{
+			QueryResult->GetAllAsActors(Results);
+		}
+	}
+
+	return Results;
+}
+
+TArray<FVector> UEnvQueryInstanceBlueprintWrapper::GetResultsAsLocations()
+{
+	TArray<FVector> Results;
+
+	if (QueryResult.IsValid())
+	{
+		if (RunMode != EEnvQueryRunMode::AllMatching)
+		{
+			Results.Add(QueryResult->GetItemAsLocation(0));
+		}
+		else
+		{
+			QueryResult->GetAllAsLocations(Results);
+		}
+	}
+
+	return Results;
+}
\ No newline at end of file
diff --git a/Engine/Source/Runtime/AIModule/Private/EnvironmentQuery/EnvQueryManager.cpp b/Engine/Source/Runtime/AIModule/Private/EnvironmentQuery/EnvQueryManager.cpp
--- a/Engine/Source/Runtime/AIModule/Private/EnvironmentQuery/EnvQueryManager.cpp
+++ b/Engine/Source/Runtime/AIModule/Private/EnvironmentQuery/EnvQueryManager.cpp
@@ -9,6 +9,7 @@
 #include "EnvironmentQuery/EnvQueryContext.h"
 #include "EnvironmentQuery/EQSTestingPawn.h"
 #include "EnvironmentQuery/EnvQueryDebugHelpers.h"
+#include "EnvironmentQuery/EnvQueryInstanceBlueprintWrapper.h"
 #if WITH_EDITOR
 #include "UnrealEd.h"
 #include "Engine/Brush.h"
@@ -596,6 +597,33 @@ float UEnvQueryManager::FindNamedParam(int32 QueryId, FName ParamName) const
 	return ParamValue;
 }
 
+//----------------------------------------------------------------------//
+// BP functions
+//----------------------------------------------------------------------//
+UEnvQueryInstanceBlueprintWrapper* UEnvQueryManager::RunEQSQuery(UObject* WorldContext, UEnvQuery* QueryTemplate, UObject* Querier, TEnumAsByte<EEnvQueryRunMode::Type> RunMode, TSubclassOf<UEnvQueryInstanceBlueprintWrapper> WrapperClass)
+{ 
+	if (QueryTemplate == nullptr)
+	{
+		return nullptr;
+	}
+
+	UEnvQueryManager* EQSManager = GetCurrent(WorldContext);
+	UEnvQueryInstanceBlueprintWrapper* QueryInstanceWrapper = nullptr;
+
+	if (EQSManager)
+	{
+		QueryInstanceWrapper = NewObject<UEnvQueryInstanceBlueprintWrapper>(WrapperClass ? WrapperClass : UEnvQueryInstanceBlueprintWrapper::StaticClass());
+		check(QueryInstanceWrapper);
+		FEnvQueryRequest QueryRequest(QueryTemplate, Querier);
+		// @todo named params still missing support
+		//QueryRequest.SetNamedParams(QueryParams);
+
+		QueryInstanceWrapper->SetRunMode(RunMode);
+		QueryInstanceWrapper->SetQueryID(QueryRequest.Execute(RunMode, QueryInstanceWrapper, &UEnvQueryInstanceBlueprintWrapper::OnQueryFinished));
+	}
+	
+	return QueryInstanceWrapper;
+}
 
 //----------------------------------------------------------------------//
 // FEQSDebugger
diff --git a/Engine/Source/Runtime/AIModule/Private/EnvironmentQuery/EnvQueryTypes.cpp b/Engine/Source/Runtime/AIModule/Private/EnvironmentQuery/EnvQueryTypes.cpp
--- a/Engine/Source/Runtime/AIModule/Private/EnvironmentQuery/EnvQueryTypes.cpp
+++ b/Engine/Source/Runtime/AIModule/Private/EnvironmentQuery/EnvQueryTypes.cpp
@@ -11,16 +11,19 @@
 #define LOCTEXT_NAMESPACE "EnvQueryGenerator"
 float UEnvQueryTypes::SkippedItemValue = -FLT_MAX;
 
+//----------------------------------------------------------------------//
+// FEnvQueryResult
+//----------------------------------------------------------------------//
 AActor* FEnvQueryResult::GetItemAsActor(int32 Index) const
 {
 	if (Items.IsValidIndex(Index) &&
 		ItemType->IsChildOf(UEnvQueryItemType_ActorBase::StaticClass()))
 	{
-		UEnvQueryItemType_ActorBase* DefTypeOb = (UEnvQueryItemType_ActorBase*)ItemType->GetDefaultObject();
+		UEnvQueryItemType_ActorBase* DefTypeOb = static_cast<UEnvQueryItemType_ActorBase*>(ItemType->GetDefaultObject());
 		return DefTypeOb->GetActor(RawData.GetData() + Items[Index].DataOffset);
 	}
 
-	return NULL;
+	return nullptr;
 }
 
 FVector FEnvQueryResult::GetItemAsLocation(int32 Index) const
@@ -28,13 +31,46 @@ FVector FEnvQueryResult::GetItemAsLocation(int32 Index) const
 	if (Items.IsValidIndex(Index) &&
 		ItemType->IsChildOf(UEnvQueryItemType_VectorBase::StaticClass()))
 	{
-		UEnvQueryItemType_VectorBase* DefTypeOb = (UEnvQueryItemType_VectorBase*)ItemType->GetDefaultObject();
+		UEnvQueryItemType_VectorBase* DefTypeOb = static_cast<UEnvQueryItemType_VectorBase*>(ItemType->GetDefaultObject());
 		return DefTypeOb->GetItemLocation(RawData.GetData() + Items[Index].DataOffset);
 	}
 
 	return FVector::ZeroVector;
 }
 
+void FEnvQueryResult::GetAllAsActors(TArray<AActor*>& OutActors) const
+{
+	if (ItemType->IsChildOf(UEnvQueryItemType_ActorBase::StaticClass()) && Items.Num() > 0)
+	{
+		UEnvQueryItemType_ActorBase* DefTypeOb = static_cast<UEnvQueryItemType_ActorBase*>(ItemType->GetDefaultObject());
+		
+		OutActors.Reserve(OutActors.Num() + Items.Num());
+
+		for (const FEnvQueryItem& Item : Items)
+		{
+			OutActors.Add(DefTypeOb->GetActor(RawData.GetData() + Item.DataOffset));
+		}
+	}
+}
+
+void FEnvQueryResult::GetAllAsLocations(TArray<FVector>& OutLocations) const
+{
+	if (ItemType->IsChildOf(UEnvQueryItemType_VectorBase::StaticClass()) && Items.Num() > 0)
+	{
+		UEnvQueryItemType_VectorBase* DefTypeOb = static_cast<UEnvQueryItemType_VectorBase*>(ItemType->GetDefaultObject());
+
+		OutLocations.Reserve(OutLocations.Num() + Items.Num());
+
+		for (const FEnvQueryItem& Item : Items)
+		{
+			OutLocations.Add(DefTypeOb->GetItemLocation(RawData.GetData() + Item.DataOffset));
+		}
+	}
+}
+
+//----------------------------------------------------------------------//
+// UEnvQueryTypes
+//----------------------------------------------------------------------//
 FText UEnvQueryTypes::GetShortTypeName(const UObject* Ob)
 {
 	if (Ob == NULL)
