Fix Mac compile error due to missing enum definition from SkeletalMeshComponent.

diff --git a/Engine/Source/Runtime/AIModule/Classes/BehaviorTree/Tasks/BTTask_PlayAnimation.h b/Engine/Source/Runtime/AIModule/Classes/BehaviorTree/Tasks/BTTask_PlayAnimation.h
--- a/Engine/Source/Runtime/AIModule/Classes/BehaviorTree/Tasks/BTTask_PlayAnimation.h
+++ b/Engine/Source/Runtime/AIModule/Classes/BehaviorTree/Tasks/BTTask_PlayAnimation.h
@@ -2,6 +2,7 @@
 
 #pragma once
 #include "BehaviorTree/BTTaskNode.h"
+#include "Components/SkeletalMeshComponent.h"
 #include "BTTask_PlayAnimation.generated.h"
 
 class UAnimationAsset;
