Fixed compilation error on android.

diff --git a/Engine/Source/Runtime/AIModule/Private/EnvironmentQuery/EnvQueryManager.cpp b/Engine/Source/Runtime/AIModule/Private/EnvironmentQuery/EnvQueryManager.cpp
--- a/Engine/Source/Runtime/AIModule/Private/EnvironmentQuery/EnvQueryManager.cpp
+++ b/Engine/Source/Runtime/AIModule/Private/EnvironmentQuery/EnvQueryManager.cpp
@@ -612,7 +612,7 @@ UEnvQueryInstanceBlueprintWrapper* UEnvQueryManager::RunEQSQuery(UObject* WorldC
 
 	if (EQSManager)
 	{
-		QueryInstanceWrapper = NewObject<UEnvQueryInstanceBlueprintWrapper>(WrapperClass ? WrapperClass : UEnvQueryInstanceBlueprintWrapper::StaticClass());
+		QueryInstanceWrapper = NewObject<UEnvQueryInstanceBlueprintWrapper>((*WrapperClass) ? WrapperClass : UEnvQueryInstanceBlueprintWrapper::StaticClass());
 		check(QueryInstanceWrapper);
 		FEnvQueryRequest QueryRequest(QueryTemplate, Querier);
 		// @todo named params still missing support
