Eliminate GWorld in EnvQueryManager

diff --git a/Engine/Source/Runtime/AIModule/Classes/EnvironmentQuery/EnvQueryManager.h b/Engine/Source/Runtime/AIModule/Classes/EnvironmentQuery/EnvQueryManager.h
--- a/Engine/Source/Runtime/AIModule/Classes/EnvironmentQuery/EnvQueryManager.h
+++ b/Engine/Source/Runtime/AIModule/Classes/EnvironmentQuery/EnvQueryManager.h
@@ -94,7 +94,7 @@ struct AIMODULE_API FEQSDebugger
 		TSharedPtr<FEnvQueryInstance> Instance;
 	};
 
-	void StoreQuery(TSharedPtr<FEnvQueryInstance>& Query);
+	void StoreQuery(UWorld* InWorld, TSharedPtr<FEnvQueryInstance>& Query);
 	TArray<FEnvQueryInfo>&  GetAllQueriesForOwner(const UObject* Owner);
 
 protected:
diff --git a/Engine/Source/Runtime/AIModule/Private/EnvironmentQuery/EnvQueryManager.cpp b/Engine/Source/Runtime/AIModule/Private/EnvironmentQuery/EnvQueryManager.cpp
--- a/Engine/Source/Runtime/AIModule/Private/EnvironmentQuery/EnvQueryManager.cpp
+++ b/Engine/Source/Runtime/AIModule/Private/EnvironmentQuery/EnvQueryManager.cpp
@@ -191,7 +191,7 @@ TSharedPtr<FEnvQueryResult> UEnvQueryManager::RunInstantQuery(const FEnvQueryReq
 	UE_VLOG_EQS(*QueryInstance.Get(), LogEQS, All);
 
 #if USE_EQS_DEBUGGER
-	EQSDebugger.StoreQuery(QueryInstance);
+	EQSDebugger.StoreQuery(GetWorld(), QueryInstance);
 #endif // USE_EQS_DEBUGGER
 
 	return QueryInstance;
@@ -287,7 +287,7 @@ void UEnvQueryManager::Tick(float DeltaTime)
 				UE_VLOG_EQS(*QueryInstance.Get(), LogEQS, All);
 
 #if USE_EQS_DEBUGGER
-				EQSDebugger.StoreQuery(QueryInstance);
+				EQSDebugger.StoreQuery(GetWorld(), QueryInstance);
 #endif // USE_EQS_DEBUGGER
 
 				QueryInstance->FinishDelegate.ExecuteIfBound(QueryInstance);
@@ -630,7 +630,7 @@ UEnvQueryInstanceBlueprintWrapper* UEnvQueryManager::RunEQSQuery(UObject* WorldC
 //----------------------------------------------------------------------//
 #if USE_EQS_DEBUGGER
 
-void FEQSDebugger::StoreQuery(TSharedPtr<FEnvQueryInstance>& Query)
+void FEQSDebugger::StoreQuery(UWorld* InWorld, TSharedPtr<FEnvQueryInstance>& Query)
 {
 	StoredQueries.Remove(NULL);
 	if (!Query.IsValid())
@@ -646,7 +646,7 @@ void FEQSDebugger::StoreQuery(TSharedPtr<FEnvQueryInstance>& Query)
 		if (CurrentQuery.Instance.IsValid() && Query->QueryName == CurrentQuery.Instance->QueryName)
 		{
 			CurrentQuery.Instance = Query;
-			CurrentQuery.Timestamp = GWorld->GetTimeSeconds();
+			CurrentQuery.Timestamp = InWorld->GetTimeSeconds();
 			bFoundQuery = true;
 			break;
 		}
@@ -655,7 +655,7 @@ void FEQSDebugger::StoreQuery(TSharedPtr<FEnvQueryInstance>& Query)
 	{
 		FEnvQueryInfo Info;
 		Info.Instance = Query;
-		Info.Timestamp = GWorld->GetTimeSeconds();
+		Info.Timestamp = InWorld->GetTimeSeconds();
 		AllQueries.AddUnique(Info);
 	}
 }
