Fix non-unity compile

diff --git a/Engine/Source/Runtime/AIModule/Private/EnvironmentQuery/EnvQueryInstanceBlueprintWrapper.cpp b/Engine/Source/Runtime/AIModule/Private/EnvironmentQuery/EnvQueryInstanceBlueprintWrapper.cpp
--- a/Engine/Source/Runtime/AIModule/Private/EnvironmentQuery/EnvQueryInstanceBlueprintWrapper.cpp
+++ b/Engine/Source/Runtime/AIModule/Private/EnvironmentQuery/EnvQueryInstanceBlueprintWrapper.cpp
@@ -2,6 +2,7 @@
 
 #include "AIModulePrivate.h"
 #include "EnvironmentQuery/EnvQueryInstanceBlueprintWrapper.h"
+#include "EnvironmentQuery/Items/EnvQueryItemType_ActorBase.h"
 
 UEnvQueryInstanceBlueprintWrapper::UEnvQueryInstanceBlueprintWrapper(const FObjectInitializer& ObjectInitializer)
 	: Super(ObjectInitializer)
diff --git a/Engine/Source/Runtime/Engine/Private/DemoNetDriver.cpp b/Engine/Source/Runtime/Engine/Private/DemoNetDriver.cpp
--- a/Engine/Source/Runtime/Engine/Private/DemoNetDriver.cpp
+++ b/Engine/Source/Runtime/Engine/Private/DemoNetDriver.cpp
@@ -9,6 +9,7 @@
 #include "Engine/DemoNetConnection.h"
 #include "Engine/DemoPendingNetGame.h"
 #include "Engine/ActorChannel.h"
+#include "Engine/PackageMapClient.h"
 #include "RepLayout.h"
 #include "GameFramework/SpectatorPawn.h"
 #include "Engine/LevelStreamingKismet.h"
