Fixed AIPerceptionComponent repeatedly reporting actors' perception update even if there was no change in actors' perception #UE4

While in there added an out parameters "Stimulus Strength" to AISense_Sight::ShouldAutomaticallySeeTarget to make it consistent with IAISightTargetInterface::CanBeSeenFrom
Renamed FAIStimulus::bLastSensingResult to more explicit bSuccessfullySensed and exposed it to BP

diff --git a/Engine/Source/Runtime/AIModule/Classes/Perception/AIPerceptionTypes.h b/Engine/Source/Runtime/AIModule/Classes/Perception/AIPerceptionTypes.h
--- a/Engine/Source/Runtime/AIModule/Classes/Perception/AIPerceptionTypes.h
+++ b/Engine/Source/Runtime/AIModule/Classes/Perception/AIPerceptionTypes.h
@@ -86,8 +86,11 @@ public:
 	FVector ReceiverLocation;
 
 	FAISenseID Type;
+
 protected:
-	uint32 bLastSensingResult:1; // currently used only for marking failed sight tests
+	UPROPERTY(BlueprintReadWrite, Category = "AI|Perception")
+	uint32 bSuccessfullySensed:1; // currently used only for marking failed sight tests
+
 	/** this means the stimulus was originally created with a "time limit" and this time has passed. 
 	 *	Expiration also results in calling MarkNoLongerSensed */
 	uint32 bExpired:1;	
@@ -101,13 +104,13 @@ public:
 		: Age(0.f), ExpirationAge(NeverHappenedAge)
 		, Strength(Result == SensingSucceeded ? StimulusStrength : -1.f)
 		, StimulusLocation(InStimulusLocation)
-		, ReceiverLocation(InReceiverLocation), Type(SenseType), bLastSensingResult(Result == SensingSucceeded), bExpired(false)
+		, ReceiverLocation(InReceiverLocation), Type(SenseType), bSuccessfullySensed(Result == SensingSucceeded), bExpired(false)
 	{}
 
 	// default constructor
 	FAIStimulus()
 		: Age(NeverHappenedAge), ExpirationAge(NeverHappenedAge), Strength(-1.f), StimulusLocation(FAISystem::InvalidLocation)
-		, ReceiverLocation(FAISystem::InvalidLocation), Type(FAISenseID::InvalidID()), bLastSensingResult(false), bExpired(false)
+		, ReceiverLocation(FAISystem::InvalidLocation), Type(FAISenseID::InvalidID()), bSuccessfullySensed(false), bExpired(false)
 	{}
 
 	FAIStimulus& SetExpirationAge(float InExpirationAge) { ExpirationAge = InExpirationAge; return *this; }
@@ -120,9 +123,9 @@ public:
 		Age += ConstPerceptionAgingRate; 
 		return Age < ExpirationAge;
 	}
-	FORCEINLINE bool WasSuccessfullySensed() const { return bLastSensingResult; }
+	FORCEINLINE bool WasSuccessfullySensed() const { return bSuccessfullySensed; }
 	FORCEINLINE bool IsExpired() const { return bExpired; }
-	FORCEINLINE void MarkNoLongerSensed() { bLastSensingResult = false; }
+	FORCEINLINE void MarkNoLongerSensed() { bSuccessfullySensed = false; }
 	FORCEINLINE void MarkExpired() { bExpired = true; MarkNoLongerSensed(); }
 	FORCEINLINE bool IsActive() const { return WasSuccessfullySensed() == true && GetAge() < NeverHappenedAge; }
 };
diff --git a/Engine/Source/Runtime/AIModule/Classes/Perception/AISense_Sight.h b/Engine/Source/Runtime/AIModule/Classes/Perception/AISense_Sight.h
--- a/Engine/Source/Runtime/AIModule/Classes/Perception/AISense_Sight.h
+++ b/Engine/Source/Runtime/AIModule/Classes/Perception/AISense_Sight.h
@@ -151,7 +151,7 @@ public:
 protected:
 	virtual float Update() override;
 
-	virtual bool ShouldAutomaticallySeeTarget(const FDigestedSightProperties& PropDigest, FAISightQuery* SightQuery, FPerceptionListener& Listener, AActor* TargetActor) const;
+	virtual bool ShouldAutomaticallySeeTarget(const FDigestedSightProperties& PropDigest, FAISightQuery* SightQuery, FPerceptionListener& Listener, AActor* TargetActor, float& OutStimulusStrength) const;
 
 	void OnNewListenerImpl(const FPerceptionListener& NewListener);
 	void OnListenerUpdateImpl(const FPerceptionListener& UpdatedListener);
diff --git a/Engine/Source/Runtime/AIModule/Private/Perception/AIPerceptionComponent.cpp b/Engine/Source/Runtime/AIModule/Private/Perception/AIPerceptionComponent.cpp
--- a/Engine/Source/Runtime/AIModule/Private/Perception/AIPerceptionComponent.cpp
+++ b/Engine/Source/Runtime/AIModule/Private/Perception/AIPerceptionComponent.cpp
@@ -408,7 +408,7 @@ void UAIPerceptionComponent::ProcessStimuli()
 		FAIStimulus& StimulusStore = PerceptualInfo->LastSensedStimuli[SourcedStimulus->Stimulus.Type];
 
 		// if the new stimulus is "valid" or it's info that "no longer sensed" and it used to be sensed successfully
-		if (SourcedStimulus->Stimulus.WasSuccessfullySensed() || StimulusStore.WasSuccessfullySensed())
+		if (SourcedStimulus->Stimulus.WasSuccessfullySensed() != StimulusStore.WasSuccessfullySensed())
 		{
 			UpdatedActors.AddUnique(SourcedStimulus->Source);
 		}
diff --git a/Engine/Source/Runtime/AIModule/Private/Perception/AIPerceptionTypes.cpp b/Engine/Source/Runtime/AIModule/Private/Perception/AIPerceptionTypes.cpp
--- a/Engine/Source/Runtime/AIModule/Private/Perception/AIPerceptionTypes.cpp
+++ b/Engine/Source/Runtime/AIModule/Private/Perception/AIPerceptionTypes.cpp
@@ -13,7 +13,7 @@ const float FAIStimulus::NeverHappenedAge = FLT_MAX;
 FAIStimulus::FAIStimulus(const UAISense& Sense, float StimulusStrength, const FVector& InStimulusLocation, const FVector& InReceiverLocation, FResult Result)
 	: Age(0.f), Strength(Result == SensingSucceeded ? StimulusStrength : -1.f)
 	, StimulusLocation(InStimulusLocation)
-	, ReceiverLocation(InReceiverLocation), bLastSensingResult(Result == SensingSucceeded)
+	, ReceiverLocation(InReceiverLocation), bSuccessfullySensed(Result == SensingSucceeded)
 	, bExpired(false)
 {
 	Type = Sense.GetSenseID();
diff --git a/Engine/Source/Runtime/AIModule/Private/Perception/AISense_Sight.cpp b/Engine/Source/Runtime/AIModule/Private/Perception/AISense_Sight.cpp
--- a/Engine/Source/Runtime/AIModule/Private/Perception/AISense_Sight.cpp
+++ b/Engine/Source/Runtime/AIModule/Private/Perception/AISense_Sight.cpp
@@ -103,8 +103,10 @@ void UAISense_Sight::PostInitProperties()
 	HighImportanceDistanceSquare = FMath::Square(HighImportanceQueryDistanceThreshold);
 }
 
-bool UAISense_Sight::ShouldAutomaticallySeeTarget(const FDigestedSightProperties& PropDigest, FAISightQuery* SightQuery, FPerceptionListener& Listener, AActor* TargetActor) const
+bool UAISense_Sight::ShouldAutomaticallySeeTarget(const FDigestedSightProperties& PropDigest, FAISightQuery* SightQuery, FPerceptionListener& Listener, AActor* TargetActor, float& OutStimulusStrength) const
 {
+	OutStimulusStrength = 1.0f;
+
 	if ((PropDigest.AutoSuccessRangeSqFromLastSeenLocation != FAISystem::InvalidRange) && (SightQuery->LastSeenLocation != FAISystem::InvalidLocation))
 	{
 		const float DistanceToLastSeenLocationSq = FVector::DistSquared(TargetActor->GetActorLocation(), SightQuery->LastSeenLocation);
@@ -155,11 +157,13 @@ float UAISense_Sight::Update()
 				const FVector TargetLocation = TargetActor->GetActorLocation();
 				const FDigestedSightProperties& PropDigest = DigestedProperties[SightQuery->ObserverId];
 				const float SightRadiusSq = SightQuery->bLastResult ? PropDigest.LoseSightRadiusSq : PropDigest.SightRadiusSq;
+				
+				float StimulusStrength = 1.f;
 
-				if (ShouldAutomaticallySeeTarget(PropDigest, SightQuery, Listener, TargetActor))
+				if (ShouldAutomaticallySeeTarget(PropDigest, SightQuery, Listener, TargetActor, StimulusStrength))
 				{
 					// Pretend like we've seen this target where we last saw them
-					Listener.RegisterStimulus(TargetActor, FAIStimulus(*this, 1.f, SightQuery->LastSeenLocation, Listener.CachedLocation));
+					Listener.RegisterStimulus(TargetActor, FAIStimulus(*this, StimulusStrength, SightQuery->LastSeenLocation, Listener.CachedLocation));
 					SightQuery->bLastResult = true;
 				}
 				else if (CheckIsTargetInSightPie(Listener, PropDigest, TargetLocation, SightRadiusSq))
@@ -172,10 +176,9 @@ float UAISense_Sight::Update()
 					{
 						int32 NumberOfLoSChecksPerformed = 0;
 						// defaulting to 1 to have "full strength" by default instead of "no strength"
-						float SightStrength = 1.f;
-						if (Target.SightTargetInterface->CanBeSeenFrom(Listener.CachedLocation, OutSeenLocation, NumberOfLoSChecksPerformed, SightStrength, Listener.Listener->GetBodyActor()) == true)
+						if (Target.SightTargetInterface->CanBeSeenFrom(Listener.CachedLocation, OutSeenLocation, NumberOfLoSChecksPerformed, StimulusStrength, Listener.Listener->GetBodyActor()) == true)
 						{
-							Listener.RegisterStimulus(TargetActor, FAIStimulus(*this, SightStrength, OutSeenLocation, Listener.CachedLocation));
+							Listener.RegisterStimulus(TargetActor, FAIStimulus(*this, StimulusStrength, OutSeenLocation, Listener.CachedLocation));
 							SightQuery->bLastResult = true;
 							SightQuery->LastSeenLocation = OutSeenLocation;
 						}
