Fixes to AIPerceptionSystem #UE4

Including:
* Better control over AIController::PerceptionComponent access
* implemented team affiliation handling to sight and hearing senses
* added a way to configure a sense if it wants to send out notifications on every perception even, or just on perception value change (like visibility lost or gained)

diff --git a/Engine/Source/Developer/GameplayDebugger/Private/GameplayDebuggingHUDComponent.cpp b/Engine/Source/Developer/GameplayDebugger/Private/GameplayDebuggingHUDComponent.cpp
--- a/Engine/Source/Developer/GameplayDebugger/Private/GameplayDebuggingHUDComponent.cpp
+++ b/Engine/Source/Developer/GameplayDebugger/Private/GameplayDebuggingHUDComponent.cpp
@@ -643,9 +643,9 @@ void AGameplayDebuggingHUDComponent::DrawPerception(APlayerController* PC, class
 		if (BTAI)
 		{
 			// standalone only
-			if (BTAI->PerceptionComponent && DefaultContext.Canvas != NULL)
+			if (BTAI->GetAIPerceptionComponent() && DefaultContext.Canvas != NULL)
 			{
-				BTAI->PerceptionComponent->DrawDebugInfo(DefaultContext.Canvas);
+				BTAI->GetAIPerceptionComponent()->DrawDebugInfo(DefaultContext.Canvas);
 
 				const FVector AILocation = MyPawn->GetActorLocation();
 				const FVector Facing = MyPawn->GetActorRotation().Vector();
diff --git a/Engine/Source/Runtime/AIModule/Classes/AIController.h b/Engine/Source/Runtime/AIModule/Classes/AIController.h
--- a/Engine/Source/Runtime/AIModule/Classes/AIController.h
+++ b/Engine/Source/Runtime/AIModule/Classes/AIController.h
@@ -110,7 +110,8 @@ public:
 	UPROPERTY(BlueprintReadWrite, Category = AI)
 	UBrainComponent* BrainComponent;
 
-	UPROPERTY(BlueprintReadOnly, Category = AI)
+	DEPRECATED_FORGAME(4.8, "PerceptionComponent should not be accessed directly, please use GetAIPerceptionComponent() function instead. PerceptionComponent will soon be private and your code will not compile.")
+	UPROPERTY()
 	UAIPerceptionComponent* PerceptionComponent;
 
 private_subobject:
@@ -328,10 +329,11 @@ public:
 	 *	@param InPriority focus priority to clear. If you don't know what to use you probably mean EAIFocusPriority::Gameplay*/
 	virtual void ClearFocus(EAIFocusPriority::Type InPriority);
 
+	void SetPerceptionComponent(UAIPerceptionComponent& InPerceptionComponent);
 	//----------------------------------------------------------------------//
 	// IAIPerceptionListenerInterface
 	//----------------------------------------------------------------------//
-	virtual UAIPerceptionComponent* GetPerceptionComponent() override { return PerceptionComponent; }
+	virtual UAIPerceptionComponent* GetPerceptionComponent() override { return GetAIPerceptionComponent(); }
 
 	//----------------------------------------------------------------------//
 	// Actions
@@ -356,8 +358,15 @@ public:
 	UPathFollowingComponent* GetPathFollowingComponent() const;
 	/** Returns ActionsComp subobject **/
 	UPawnActionsComponent* GetActionsComp() const;
+	UFUNCTION(BlueprintPure, Category = "AI|Perception")
+	UAIPerceptionComponent* GetAIPerceptionComponent();
+
+	const UAIPerceptionComponent* GetAIPerceptionComponent() const;
 };
 
+//----------------------------------------------------------------------//
+// forceinlines
+//----------------------------------------------------------------------//
 namespace FAISystem
 {
 	FORCEINLINE bool IsValidControllerAndHasValidPawn(const AController* Controller)
diff --git a/Engine/Source/Runtime/AIModule/Classes/GenericTeamAgentInterface.h b/Engine/Source/Runtime/AIModule/Classes/GenericTeamAgentInterface.h
--- a/Engine/Source/Runtime/AIModule/Classes/GenericTeamAgentInterface.h
+++ b/Engine/Source/Runtime/AIModule/Classes/GenericTeamAgentInterface.h
@@ -42,8 +42,23 @@ public:
 	
 	static FGenericTeamId GetTeamIdentifier(const AActor* TeamMember);
 	static ETeamAttitude::Type GetAttitude(const AActor* A, const AActor* B);
+	static ETeamAttitude::Type GetAttitude(FGenericTeamId TeamA, FGenericTeamId TeamB)
+	{
+		return AttitudeSolverImpl ? (*AttitudeSolverImpl)(TeamA, TeamB) : ETeamAttitude::Neutral;
+	}
+
+	typedef ETeamAttitude::Type FAttitudeSolverFunction(FGenericTeamId, FGenericTeamId);
+	
+	static void SetAttitudeSolver(FAttitudeSolverFunction* Solver);
 
-	static FGenericTeamId NoTeam;
+protected:
+	// the default implementation makes all teams hostile
+	// @note that for consistency IGenericTeamAgentInterface should be using the same function 
+	//	(by default it does)
+	static FAttitudeSolverFunction* AttitudeSolverImpl;
+
+public:
+	static const FGenericTeamId NoTeam;
 };
 
 UINTERFACE()
@@ -57,11 +72,16 @@ class AIMODULE_API IGenericTeamAgentInterface
 	GENERATED_IINTERFACE_BODY()
 
 	/** Assigns Team Agent to given TeamID */
-	virtual void SetGenericTeamId(const FGenericTeamId& TemaID) {}
+	virtual void SetGenericTeamId(const FGenericTeamId& TeamID) {}
 	
 	/** Retrieve team identifier in form of FGenericTeamId */
 	virtual FGenericTeamId GetGenericTeamId() const { return FGenericTeamId::NoTeam; }
 
 	/** Retrieved owner attitude toward given Other object */
-	virtual ETeamAttitude::Type GetTeamAttitudeTowards(const AActor& Other) const { return ETeamAttitude::Neutral; }
+	virtual ETeamAttitude::Type GetTeamAttitudeTowards(const AActor& Other) const
+	{ 
+		const IGenericTeamAgentInterface* OtherTeamAgent = Cast<const IGenericTeamAgentInterface>(&Other);
+		return OtherTeamAgent ? FGenericTeamId::GetAttitude(GetGenericTeamId(), OtherTeamAgent->GetGenericTeamId())
+			: ETeamAttitude::Neutral;
+	}
 };
diff --git a/Engine/Source/Runtime/AIModule/Classes/Perception/AIPerceptionListenerInterface.h b/Engine/Source/Runtime/AIModule/Classes/Perception/AIPerceptionListenerInterface.h
--- a/Engine/Source/Runtime/AIModule/Classes/Perception/AIPerceptionListenerInterface.h
+++ b/Engine/Source/Runtime/AIModule/Classes/Perception/AIPerceptionListenerInterface.h
@@ -16,6 +16,6 @@ class AIMODULE_API IAIPerceptionListenerInterface
 {
 	GENERATED_IINTERFACE_BODY()
 
-	virtual UAIPerceptionComponent* GetPerceptionComponent() { return NULL; }
+	virtual UAIPerceptionComponent* GetPerceptionComponent() { return nullptr; }
 };
 
diff --git a/Engine/Source/Runtime/AIModule/Classes/Perception/AIPerceptionTypes.h b/Engine/Source/Runtime/AIModule/Classes/Perception/AIPerceptionTypes.h
--- a/Engine/Source/Runtime/AIModule/Classes/Perception/AIPerceptionTypes.h
+++ b/Engine/Source/Runtime/AIModule/Classes/Perception/AIPerceptionTypes.h
@@ -25,6 +25,13 @@ typedef FAIGenericID<FPerceptionListenerCounter> FPerceptionListenerID;
 
 //////////////////////////////////////////////////////////////////////////
 
+UENUM()
+enum class EAISenseNotifyType : uint8
+{
+	OnEveryPerception,	// continuous update whenever target is perceived
+	OnPerceptionChange, // from "visible" to "not visible" or vice versa
+};
+
 struct FPerceptionChannelFilter
 {
 	static const uint32 AllChannels = uint32(-1);
@@ -88,6 +95,8 @@ public:
 	FAISenseID Type;
 
 protected:
+	uint32 bWantsToNotifyOnlyOnValueChange : 1;
+
 	UPROPERTY(BlueprintReadWrite, Category = "AI|Perception")
 	uint32 bSuccessfullySensed:1; // currently used only for marking failed sight tests
 
@@ -100,21 +109,24 @@ public:
 	/** this is the recommended constructor. Use others if you know what you're doing. */
 	FAIStimulus(const UAISense& Sense, float StimulusStrength, const FVector& InStimulusLocation, const FVector& InReceiverLocation, FResult Result = SensingSucceeded);
 
-	FAIStimulus(FAISenseID SenseType, float StimulusStrength, const FVector& InStimulusLocation, const FVector& InReceiverLocation, FResult Result = SensingSucceeded)
+	/*FAIStimulus(FAISenseID SenseType, float StimulusStrength, const FVector& InStimulusLocation, const FVector& InReceiverLocation, FResult Result = SensingSucceeded)
 		: Age(0.f), ExpirationAge(NeverHappenedAge)
 		, Strength(Result == SensingSucceeded ? StimulusStrength : -1.f)
 		, StimulusLocation(InStimulusLocation)
-		, ReceiverLocation(InReceiverLocation), Type(SenseType), bSuccessfullySensed(Result == SensingSucceeded), bExpired(false)
-	{}
+		, ReceiverLocation(InReceiverLocation), Type(SenseType), bWantsToNotifyOnlyOnValueChange(false)
+		, bSuccessfullySensed(Result == SensingSucceeded), bExpired(false)
+	{}*/
 
 	// default constructor
 	FAIStimulus()
 		: Age(NeverHappenedAge), ExpirationAge(NeverHappenedAge), Strength(-1.f), StimulusLocation(FAISystem::InvalidLocation)
-		, ReceiverLocation(FAISystem::InvalidLocation), Type(FAISenseID::InvalidID()), bSuccessfullySensed(false), bExpired(false)
+		, ReceiverLocation(FAISystem::InvalidLocation), Type(FAISenseID::InvalidID()), bWantsToNotifyOnlyOnValueChange(false)
+		, bSuccessfullySensed(false), bExpired(false)
 	{}
 
 	FAIStimulus& SetExpirationAge(float InExpirationAge) { ExpirationAge = InExpirationAge; return *this; }
 	FAIStimulus& SetStimulusAge(float StimulusAge) { Age = StimulusAge; return *this; }
+	FAIStimulus& SetWantsNotifyOnlyOnValueChange(bool InEnable) { bWantsToNotifyOnlyOnValueChange = InEnable; return *this; }
 	
 	FORCEINLINE float GetAge() const { return Strength > 0 ? Age : NeverHappenedAge; }
 	/** @return false when this stimulus is no longer valid, when it is Expired */
@@ -128,6 +140,7 @@ public:
 	FORCEINLINE void MarkNoLongerSensed() { bSuccessfullySensed = false; }
 	FORCEINLINE void MarkExpired() { bExpired = true; MarkNoLongerSensed(); }
 	FORCEINLINE bool IsActive() const { return WasSuccessfullySensed() == true && GetAge() < NeverHappenedAge; }
+	FORCEINLINE bool WantsToNotifyOnlyOnPerceptionChange() const { return bWantsToNotifyOnlyOnValueChange; }
 };
 
 USTRUCT()
@@ -145,6 +158,22 @@ struct AIMODULE_API FAISenseAffiliationFilter
 	uint32 bDetectFriendlies : 1;
 	
 	uint8 GetAsFlags() const { return (bDetectEnemies << ETeamAttitude::Hostile) | (bDetectNeutrals << ETeamAttitude::Neutral) | (bDetectFriendlies << ETeamAttitude::Friendly); }
+	FORCEINLINE bool ShouldDetectAll() const { return (bDetectEnemies & bDetectNeutrals & bDetectFriendlies); }
+
+	static FORCEINLINE uint8 DetectAllFlags() { return (1 << ETeamAttitude::Hostile) | (1 << ETeamAttitude::Neutral) | (1 << ETeamAttitude::Friendly); }
+
+	static bool ShouldSenseTeam(FGenericTeamId TeamA, FGenericTeamId TeamB, uint8 AffiliationFlags)
+	{
+		static const uint8 AllFlags = DetectAllFlags();
+		return AffiliationFlags == AllFlags || ((1 << FGenericTeamId::GetAttitude(TeamA, TeamB)) & AffiliationFlags);
+	}
+
+	static bool ShouldSenseTeam(const IGenericTeamAgentInterface* TeamAgent, const AActor& TargetActor, uint8 AffiliationFlags)
+	{
+		static const uint8 AllFlags = DetectAllFlags();
+		return AffiliationFlags == AllFlags 
+			|| (TeamAgent == nullptr ? (AffiliationFlags & (1 << ETeamAttitude::Neutral)) : ((1 << TeamAgent->GetTeamAttitudeTowards(TargetActor)) & AffiliationFlags));
+	}
 };
 
 /** Should contain only cached information common to all senses. Sense-specific data needs to be stored by senses themselves */
diff --git a/Engine/Source/Runtime/AIModule/Classes/Perception/AISense.h b/Engine/Source/Runtime/AIModule/Classes/Perception/AISense.h
--- a/Engine/Source/Runtime/AIModule/Classes/Perception/AISense.h
+++ b/Engine/Source/Runtime/AIModule/Classes/Perception/AISense.h
@@ -28,6 +28,9 @@ protected:
 	/** age past which stimulus of this sense are "forgotten"*/
 	UPROPERTY(EditAnywhere, BlueprintReadOnly, Category = "AI Perception", config)
 	float DefaultExpirationAge;
+
+	UPROPERTY(EditDefaultsOnly, BlueprintReadOnly, Category = "AI Perception", config)
+	EAISenseNotifyType NotifyType;
 	
 private:
 	UPROPERTY()
@@ -63,6 +66,8 @@ public:
 	}
 	FORCEINLINE FAISenseID GetSenseID() const { return SenseID; }
 
+	FORCEINLINE bool WantsUpdateOnlyOnPerceptionValueChange() const { return (NotifyType == EAISenseNotifyType::OnPerceptionChange); }
+
 	virtual void PostInitProperties() override;
 
 	/** 
diff --git a/Engine/Source/Runtime/AIModule/Classes/Perception/AISense_Sight.h b/Engine/Source/Runtime/AIModule/Classes/Perception/AISense_Sight.h
--- a/Engine/Source/Runtime/AIModule/Classes/Perception/AISense_Sight.h
+++ b/Engine/Source/Runtime/AIModule/Classes/Perception/AISense_Sight.h
@@ -126,18 +126,18 @@ public:
 	TArray<FAISightQuery> SightQueryQueue;
 
 protected:
-	UPROPERTY(config)
+	UPROPERTY(EditDefaultsOnly, Category = "AI Perception", config)
 	int32 MaxTracesPerTick;
 
-	UPROPERTY(config)
+	UPROPERTY(EditDefaultsOnly, Category = "AI Perception", config)
 	float HighImportanceQueryDistanceThreshold;
 
 	float HighImportanceDistanceSquare;
 
-	UPROPERTY(config)
+	UPROPERTY(EditDefaultsOnly, Category = "AI Perception", config)
 	float MaxQueryImportance;
 
-	UPROPERTY(config)
+	UPROPERTY(EditDefaultsOnly, Category = "AI Perception", config)
 	float SightLimitQueryImportance;
 
 public:
diff --git a/Engine/Source/Runtime/AIModule/Private/AIController.cpp b/Engine/Source/Runtime/AIModule/Private/AIController.cpp
--- a/Engine/Source/Runtime/AIModule/Private/AIController.cpp
+++ b/Engine/Source/Runtime/AIModule/Private/AIController.cpp
@@ -276,6 +276,15 @@ void AAIController::ClearFocus(EAIFocusPriority::Type InPriority)
 	}
 }
 
+void AAIController::SetPerceptionComponent(UAIPerceptionComponent& InPerceptionComponent)
+{
+	if (PerceptionComponent != nullptr)
+	{
+		UE_VLOG(this, LogAIPerception, Warning, TEXT("Setting perception component while AIController already has one!"));
+	}
+	PerceptionComponent = &InPerceptionComponent;
+}
+
 bool AAIController::LineOfSightTo(const AActor* Other, FVector ViewPoint, bool bAlternateChecks) const
 {
 	if (Other == nullptr)
@@ -485,10 +494,10 @@ EPathFollowingRequestResult::Type AAIController::MoveToActor(AActor* Goal, float
 	MoveReq.SetCanStrafe(bCanStrafe);
 
 	return MoveTo(MoveReq);
-	}
+}
 
 EPathFollowingRequestResult::Type AAIController::MoveToLocation(const FVector& Dest, float AcceptanceRadius, bool bStopOnOverlap, bool bUsePathfinding, bool bProjectDestinationToNavigation, bool bCanStrafe, TSubclassOf<UNavigationQueryFilter> FilterClass, bool bAllowPartialPaths)
-		{
+{
 	FAIMoveRequest MoveReq(Dest);
 	MoveReq.SetUsePathfinding(bUsePathfinding);
 	MoveReq.SetAllowPartialPath(bAllowPartialPaths);
@@ -499,7 +508,7 @@ EPathFollowingRequestResult::Type AAIController::MoveToLocation(const FVector& D
 	MoveReq.SetCanStrafe(bCanStrafe);
 
 	return MoveTo(MoveReq);
-	}
+}
 
 EPathFollowingRequestResult::Type AAIController::MoveTo(const FAIMoveRequest& MoveRequest)
 {
@@ -874,3 +883,13 @@ FString AAIController::GetDebugIcon() const
 UPathFollowingComponent* AAIController::GetPathFollowingComponent() const { return PathFollowingComponent; }
 /** Returns ActionsComp subobject **/
 UPawnActionsComponent* AAIController::GetActionsComp() const { return ActionsComp; }
+
+UAIPerceptionComponent* AAIController::GetAIPerceptionComponent()
+{
+	return PerceptionComponent;
+}
+
+const UAIPerceptionComponent* AAIController::GetAIPerceptionComponent() const 
+{
+	return PerceptionComponent;
+}
\ No newline at end of file
diff --git a/Engine/Source/Runtime/AIModule/Private/AIInterfaces.cpp b/Engine/Source/Runtime/AIModule/Private/AIInterfaces.cpp
--- a/Engine/Source/Runtime/AIModule/Private/AIInterfaces.cpp
+++ b/Engine/Source/Runtime/AIModule/Private/AIInterfaces.cpp
@@ -9,7 +9,7 @@ UAIResourceInterface::UAIResourceInterface(const FObjectInitializer& ObjectIniti
 {
 }
 
-FGenericTeamId FGenericTeamId::NoTeam(FGenericTeamId::NoTeamId);
+const FGenericTeamId FGenericTeamId::NoTeam(FGenericTeamId::NoTeamId);
 
 FGenericTeamId FGenericTeamId::GetTeamIdentifier(const AActor* TeamMember)
 {
@@ -22,6 +22,19 @@ FGenericTeamId FGenericTeamId::GetTeamIdentifier(const AActor* TeamMember)
 	return FGenericTeamId::NoTeam;
 }
 
+//----------------------------------------------------------------------//
+// FGenericTeamId
+//----------------------------------------------------------------------//
+namespace
+{
+	ETeamAttitude::Type DefaultTeamAttitudeSolver(FGenericTeamId A, FGenericTeamId B)
+	{
+		return A != B ? ETeamAttitude::Hostile : ETeamAttitude::Friendly;
+	}
+}
+
+FGenericTeamId::FAttitudeSolverFunction* FGenericTeamId::AttitudeSolverImpl = &DefaultTeamAttitudeSolver;
+
 ETeamAttitude::Type FGenericTeamId::GetAttitude(const AActor* A, const AActor* B)
 {
 	const IGenericTeamAgentInterface* TeamAgentA = Cast<const IGenericTeamAgentInterface>(A);
@@ -29,6 +42,14 @@ ETeamAttitude::Type FGenericTeamId::GetAttitude(const AActor* A, const AActor* B
 	return TeamAgentA == NULL || B == NULL ? ETeamAttitude::Neutral : TeamAgentA->GetTeamAttitudeTowards(*B);
 }
 
+void FGenericTeamId::SetAttitudeSolver(FGenericTeamId::FAttitudeSolverFunction* Solver)
+{
+	AttitudeSolverImpl = Solver;
+}
+
+//----------------------------------------------------------------------//
+// UGenericTeamAgentInterface
+//----------------------------------------------------------------------//
 UGenericTeamAgentInterface::UGenericTeamAgentInterface(const FObjectInitializer& ObjectInitializer)
 	: Super(ObjectInitializer)
 {
diff --git a/Engine/Source/Runtime/AIModule/Private/AIModulePrivate.h b/Engine/Source/Runtime/AIModule/Private/AIModulePrivate.h
--- a/Engine/Source/Runtime/AIModule/Private/AIModulePrivate.h
+++ b/Engine/Source/Runtime/AIModule/Private/AIModulePrivate.h
@@ -39,6 +39,7 @@
 #include "EnvironmentQuery/EnvQueryTypes.h"
 #include "EnvironmentQuery/EnvQueryGenerator.h"
 #include "EnvironmentQuery/EnvQueryTest.h"
+#include "EnvironmentQuery/Items/EnvQueryItemType_ActorBase.h"
 #include "BehaviorTree/Blackboard/BlackboardKeyType.h"
 #include "Actions/PawnAction.h"
 #include "Actions/PawnActionsComponent.h"
@@ -71,4 +72,4 @@
 
 DECLARE_LOG_CATEGORY_EXTERN(LogAINavigation, Warning, All);
 DECLARE_LOG_CATEGORY_EXTERN(LogBehaviorTree, Display, All);
-DECLARE_LOG_CATEGORY_EXTERN(LogCrowdFollowing, Warning, All);
+DECLARE_LOG_CATEGORY_EXTERN(LogCrowdFollowing, Warning, All);
\ No newline at end of file
diff --git a/Engine/Source/Runtime/AIModule/Private/Perception/AIPerceptionComponent.cpp b/Engine/Source/Runtime/AIModule/Private/Perception/AIPerceptionComponent.cpp
--- a/Engine/Source/Runtime/AIModule/Private/Perception/AIPerceptionComponent.cpp
+++ b/Engine/Source/Runtime/AIModule/Private/Perception/AIPerceptionComponent.cpp
@@ -183,7 +183,7 @@ void UAIPerceptionComponent::OnRegister()
 	AAIController* AIOwner = Cast<AAIController>(GetOwner());
 	if (AIOwner)
 	{
-		AIOwner->PerceptionComponent = this;
+		AIOwner->SetPerceptionComponent(*this);
 	}
 }
 
@@ -408,7 +408,7 @@ void UAIPerceptionComponent::ProcessStimuli()
 		FAIStimulus& StimulusStore = PerceptualInfo->LastSensedStimuli[SourcedStimulus->Stimulus.Type];
 
 		// if the new stimulus is "valid" or it's info that "no longer sensed" and it used to be sensed successfully
-		if (SourcedStimulus->Stimulus.WasSuccessfullySensed() != StimulusStore.WasSuccessfullySensed())
+		if (SourcedStimulus->Stimulus.WantsToNotifyOnlyOnPerceptionChange() == false || SourcedStimulus->Stimulus.WasSuccessfullySensed() != StimulusStore.WasSuccessfullySensed())
 		{
 			UpdatedActors.AddUnique(SourcedStimulus->Source);
 		}
diff --git a/Engine/Source/Runtime/AIModule/Private/Perception/AIPerceptionTypes.cpp b/Engine/Source/Runtime/AIModule/Private/Perception/AIPerceptionTypes.cpp
--- a/Engine/Source/Runtime/AIModule/Private/Perception/AIPerceptionTypes.cpp
+++ b/Engine/Source/Runtime/AIModule/Private/Perception/AIPerceptionTypes.cpp
@@ -13,7 +13,9 @@ const float FAIStimulus::NeverHappenedAge = FLT_MAX;
 FAIStimulus::FAIStimulus(const UAISense& Sense, float StimulusStrength, const FVector& InStimulusLocation, const FVector& InReceiverLocation, FResult Result)
 	: Age(0.f), Strength(Result == SensingSucceeded ? StimulusStrength : -1.f)
 	, StimulusLocation(InStimulusLocation)
-	, ReceiverLocation(InReceiverLocation), bSuccessfullySensed(Result == SensingSucceeded)
+	, ReceiverLocation(InReceiverLocation)
+	, bWantsToNotifyOnlyOnValueChange(Sense.WantsUpdateOnlyOnPerceptionValueChange())
+	, bSuccessfullySensed(Result == SensingSucceeded)
 	, bExpired(false)
 {
 	Type = Sense.GetSenseID();
diff --git a/Engine/Source/Runtime/AIModule/Private/Perception/AISense_Hearing.cpp b/Engine/Source/Runtime/AIModule/Private/Perception/AISense_Hearing.cpp
--- a/Engine/Source/Runtime/AIModule/Private/Perception/AISense_Hearing.cpp
+++ b/Engine/Source/Runtime/AIModule/Private/Perception/AISense_Hearing.cpp
@@ -124,10 +124,8 @@ float UAISense_Hearing::Update()
 		{
 			const FAINoiseEvent& Event = NoiseEvents[EventIndex];
 		
-			// @todo implement some kind of TeamIdentifierType that would supply comparison operator 
-			// @todo use PropDigest.AffiliationFlags
-			if (Listener.TeamIdentifier == Event.TeamIdentifier 
-				|| FVector::DistSquared(Event.NoiseLocation, Listener.CachedLocation) > PropDigest.HearingRangeSq * Event.Loudness)
+			if (FVector::DistSquared(Event.NoiseLocation, Listener.CachedLocation) > PropDigest.HearingRangeSq * Event.Loudness
+				|| FAISenseAffiliationFilter::ShouldSenseTeam(Listener.TeamIdentifier, Event.TeamIdentifier, PropDigest.AffiliationFlags) == false)
 			{
 				continue;
 			}
diff --git a/Engine/Source/Runtime/AIModule/Private/Perception/AISense_Sight.cpp b/Engine/Source/Runtime/AIModule/Private/Perception/AISense_Sight.cpp
--- a/Engine/Source/Runtime/AIModule/Private/Perception/AISense_Sight.cpp
+++ b/Engine/Source/Runtime/AIModule/Private/Perception/AISense_Sight.cpp
@@ -88,6 +88,7 @@ UAISense_Sight::UAISense_Sight(const FObjectInitializer& ObjectInitializer)
 
 	DebugDrawColor = FColor::Green;
 	DebugName = TEXT("Sight");
+	NotifyType = EAISenseNotifyType::OnPerceptionChange;
 }
 
 FORCEINLINE_DEBUGGABLE float UAISense_Sight::CalcQueryImportance(const FPerceptionListener& Listener, const FVector& TargetLocation, const float SightRadiusSq) const
@@ -313,16 +314,18 @@ bool UAISense_Sight::RegisterTarget(AActor& TargetActor, FQueriesOperationPostPr
 		const FPerceptionListener& Listener = ItListener->Value;
 		const IGenericTeamAgentInterface* ListenersTeamAgent = Listener.GetTeamAgent();
 
-		// @todo add configuration here
-		if (Listener.HasSense(GetSenseID()) && (ListenersTeamAgent == NULL || ListenersTeamAgent->GetTeamAttitudeTowards(TargetActor) == ETeamAttitude::Hostile))
+		if (Listener.HasSense(GetSenseID()))
 		{
-			// create a sight query		
-			FAISightQuery SightQuery(ItListener->Key, SightTarget->TargetId);
 			const FDigestedSightProperties& PropDigest = DigestedProperties[Listener.GetListenerID()];
-			SightQuery.Importance = CalcQueryImportance(ItListener->Value, TargetLocation, PropDigest.SightRadiusSq);
+			if (FAISenseAffiliationFilter::ShouldSenseTeam(ListenersTeamAgent, TargetActor, PropDigest.AffiliationFlags))
+			{
+				// create a sight query		
+				FAISightQuery SightQuery(ItListener->Key, SightTarget->TargetId);
+				SightQuery.Importance = CalcQueryImportance(ItListener->Value, TargetLocation, PropDigest.SightRadiusSq);
 
-			SightQueryQueue.Add(SightQuery);
-			bNewQueriesAdded = true;
+				SightQueryQueue.Add(SightQuery);
+				bNewQueriesAdded = true;
+			}
 		}
 	}
 
@@ -361,8 +364,7 @@ void UAISense_Sight::GenerateQueriesForListener(const FPerceptionListener& Liste
 			continue;
 		}
 
-		// @todo this should be configurable - some AI might want to observe Neutrals and Friendlies as well
-		if (ListenersTeamAgent == NULL || ListenersTeamAgent->GetTeamAttitudeTowards(*TargetActor) == ETeamAttitude::Hostile)
+		if (FAISenseAffiliationFilter::ShouldSenseTeam(ListenersTeamAgent, *TargetActor, PropertyDigest.AffiliationFlags))
 		{
 			// create a sight query		
 			FAISightQuery SightQuery(Listener.GetListenerID(), ItTarget->Key);
