Brought back a FNavigationPath::ShortcutNodeRefs copying operation removed by PVS #UE4

But placed it in a better spot!

diff --git a/Engine/Source/Runtime/AIModule/Private/Navigation/PathFollowingComponent.cpp b/Engine/Source/Runtime/AIModule/Private/Navigation/PathFollowingComponent.cpp
--- a/Engine/Source/Runtime/AIModule/Private/Navigation/PathFollowingComponent.cpp
+++ b/Engine/Source/Runtime/AIModule/Private/Navigation/PathFollowingComponent.cpp
@@ -711,6 +711,8 @@ int32 UPathFollowingComponent::OptimizeSegmentVisibility(int32 StartIndex)
 		{
 			Path->ShortcutNodeRefs.Reserve(RaycastResult.CorridorPolysCount);
 			Path->ShortcutNodeRefs.SetNumUninitialized(RaycastResult.CorridorPolysCount);
+
+			FPlatformMemory::Memcmp(Path->ShortcutNodeRefs.GetData(), RaycastResult.CorridorPolys, RaycastResult.CorridorPolysCount * sizeof(NavNodeRef));
 		}
 #else
 		const bool RaycastHitResult = NavData->Raycast(PawnLocation, AdjustedDestination, HitLocation, QueryFilter);
