Fixed PawnActions not being aborted on PawnActionComponent's unregistering #UE4

This fixed PawnAction_Move holding a reference to navigation path after AI's death, which made the path stick around longer and impacting navigation paths updating time.

diff --git a/Engine/Source/Runtime/AIModule/Classes/Actions/PawnActionsComponent.h b/Engine/Source/Runtime/AIModule/Classes/Actions/PawnActionsComponent.h
--- a/Engine/Source/Runtime/AIModule/Classes/Actions/PawnActionsComponent.h
+++ b/Engine/Source/Runtime/AIModule/Classes/Actions/PawnActionsComponent.h
@@ -90,6 +90,11 @@ private:
 	uint32 ActionEventIndex;
 
 public:
+	//----------------------------------------------------------------------//
+	// UActorComponent
+	//----------------------------------------------------------------------//
+	virtual void OnUnregister() override;
+
 	//----------------------------------------------------------------------//
 	// blueprint interface
 	//----------------------------------------------------------------------//
diff --git a/Engine/Source/Runtime/AIModule/Private/Actions/PawnActionsComponent.cpp b/Engine/Source/Runtime/AIModule/Private/Actions/PawnActionsComponent.cpp
--- a/Engine/Source/Runtime/AIModule/Private/Actions/PawnActionsComponent.cpp
+++ b/Engine/Source/Runtime/AIModule/Private/Actions/PawnActionsComponent.cpp
@@ -162,6 +162,22 @@ UPawnActionsComponent::UPawnActionsComponent(const FObjectInitializer& ObjectIni
 	ActionStacks.AddZeroed(EAIRequestPriority::MAX);
 }
 
+void UPawnActionsComponent::OnUnregister()
+{
+	Super::OnUnregister();
+
+	// call for every regular priority 
+	for (int32 PriorityIndex = 0; PriorityIndex < EAIRequestPriority::MAX; ++PriorityIndex)
+	{
+		UPawnAction* Action = ActionStacks[PriorityIndex].GetTop();
+		while (Action)
+		{
+			Action->Abort(EAIForceParam::Force);
+			Action = Action->ParentAction;
+		}
+	}
+}
+
 void UPawnActionsComponent::TickComponent(float DeltaTime, enum ELevelTick TickType, FActorComponentTickFunction *ThisTickFunction)
 {
 	Super::TickComponent(DeltaTime, TickType, ThisTickFunction);
