Un-virtualized AI BlueprintImplementableEvents #UE4

UE-12506

diff --git a/Engine/Source/Runtime/AIModule/Classes/BehaviorTree/Decorators/BTDecorator_BlueprintBase.h b/Engine/Source/Runtime/AIModule/Classes/BehaviorTree/Decorators/BTDecorator_BlueprintBase.h
--- a/Engine/Source/Runtime/AIModule/Classes/BehaviorTree/Decorators/BTDecorator_BlueprintBase.h
+++ b/Engine/Source/Runtime/AIModule/Classes/BehaviorTree/Decorators/BTDecorator_BlueprintBase.h
@@ -116,79 +116,79 @@ protected:
 	 *	@Note that if both generic and AI event versions are implemented only the more
 	 *	suitable one will be called, meaning the AI version if called for AI, generic one otherwise */
 	UFUNCTION(BlueprintImplementableEvent)
-	virtual void ReceiveTick(AActor* OwnerActor, float DeltaSeconds);
+	void ReceiveTick(AActor* OwnerActor, float DeltaSeconds);
 
 	/** called on execution of underlying node 
 	 *	@Note that if both generic and AI event versions are implemented only the more
 	 *	suitable one will be called, meaning the AI version if called for AI, generic one otherwise */
 	UFUNCTION(BlueprintImplementableEvent)
-	virtual void ReceiveExecutionStart(AActor* OwnerActor);
+	void ReceiveExecutionStart(AActor* OwnerActor);
 
 	/** called when execution of underlying node is finished 
 	 *	@Note that if both generic and AI event versions are implemented only the more
 	 *	suitable one will be called, meaning the AI version if called for AI, generic one otherwise */
 	UFUNCTION(BlueprintImplementableEvent)
-	virtual void ReceiveExecutionFinish(AActor* OwnerActor, enum EBTNodeResult::Type NodeResult);
+	void ReceiveExecutionFinish(AActor* OwnerActor, enum EBTNodeResult::Type NodeResult);
 
 	/** called when observer is activated (flow controller) 
 	 *	@Note that if both generic and AI event versions are implemented only the more
 	 *	suitable one will be called, meaning the AI version if called for AI, generic one otherwise */
 	UFUNCTION(BlueprintImplementableEvent)
-	virtual void ReceiveObserverActivated(AActor* OwnerActor);
+	void ReceiveObserverActivated(AActor* OwnerActor);
 
 	/** called when observer is deactivated (flow controller) 
 	 *	@Note that if both generic and AI event versions are implemented only the more
 	 *	suitable one will be called, meaning the AI version if called for AI, generic one otherwise */
 	UFUNCTION(BlueprintImplementableEvent)
-	virtual void ReceiveObserverDeactivated(AActor* OwnerActor);
+	void ReceiveObserverDeactivated(AActor* OwnerActor);
 
 	/** called when testing if underlying node can be executed, must call FinishConditionCheck
 	 *	@Note that if both generic and AI event versions are implemented only the more
 	 *	suitable one will be called, meaning the AI version if called for AI, generic one otherwise */
 	UFUNCTION(BlueprintImplementableEvent)
-	virtual bool PerformConditionCheck(AActor* OwnerActor);
+	bool PerformConditionCheck(AActor* OwnerActor);
 
 	/** Alternative AI version of ReceiveTick
 	 *	@see ReceiveTick for more details
 	 *	@Note that if both generic and AI event versions are implemented only the more
 	 *	suitable one will be called, meaning the AI version if called for AI, generic one otherwise */
 	UFUNCTION(BlueprintImplementableEvent, Category = AI)
-	virtual void ReceiveTickAI(AAIController* OwnerController, APawn* ControlledPawn, float DeltaSeconds);
+	void ReceiveTickAI(AAIController* OwnerController, APawn* ControlledPawn, float DeltaSeconds);
 
 	/** Alternative AI version of ReceiveExecutionStart
 	 *	@see ReceiveExecutionStart for more details
 	 *	@Note that if both generic and AI event versions are implemented only the more
 	 *	suitable one will be called, meaning the AI version if called for AI, generic one otherwise */
 	UFUNCTION(BlueprintImplementableEvent, Category = AI)
-	virtual void ReceiveExecutionStartAI(AAIController* OwnerController, APawn* ControlledPawn);
+	void ReceiveExecutionStartAI(AAIController* OwnerController, APawn* ControlledPawn);
 
 	/** Alternative AI version of ReceiveExecutionFinish
 	 *	@see ReceiveExecutionFinish for more details
 	 *	@Note that if both generic and AI event versions are implemented only the more
 	 *	suitable one will be called, meaning the AI version if called for AI, generic one otherwise */
 	UFUNCTION(BlueprintImplementableEvent, Category = AI)
-	virtual void ReceiveExecutionFinishAI(AAIController* OwnerController, APawn* ControlledPawn, enum EBTNodeResult::Type NodeResult);
+	void ReceiveExecutionFinishAI(AAIController* OwnerController, APawn* ControlledPawn, enum EBTNodeResult::Type NodeResult);
 
 	/** Alternative AI version of ReceiveObserverActivated
 	 *	@see ReceiveObserverActivated for more details
 	 *	@Note that if both generic and AI event versions are implemented only the more
 	 *	suitable one will be called, meaning the AI version if called for AI, generic one otherwise */
 	UFUNCTION(BlueprintImplementableEvent, Category = AI)
-	virtual void ReceiveObserverActivatedAI(AAIController* OwnerController, APawn* ControlledPawn);
+	void ReceiveObserverActivatedAI(AAIController* OwnerController, APawn* ControlledPawn);
 
 	/** Alternative AI version of ReceiveObserverDeactivated
 	 *	@see ReceiveObserverDeactivated for more details
 	 *	@Note that if both generic and AI event versions are implemented only the more
 	 *	suitable one will be called, meaning the AI version if called for AI, generic one otherwise */
 	UFUNCTION(BlueprintImplementableEvent, Category = AI)
-	virtual void ReceiveObserverDeactivatedAI(AAIController* OwnerController, APawn* ControlledPawn);
+	void ReceiveObserverDeactivatedAI(AAIController* OwnerController, APawn* ControlledPawn);
 
 	/** Alternative AI version of ReceiveConditionCheck
 	 *	@see ReceiveConditionCheck for more details
 	 *	@Note that if both generic and AI event versions are implemented only the more
 	 *	suitable one will be called, meaning the AI version if called for AI, generic one otherwise */
 	UFUNCTION(BlueprintImplementableEvent, Category = AI)
-	virtual bool PerformConditionCheckAI(AAIController* OwnerController, APawn* ControlledPawn);
+	bool PerformConditionCheckAI(AAIController* OwnerController, APawn* ControlledPawn);
 		
 	/** check if decorator is part of currently active branch */
 	UFUNCTION(BlueprintCallable, Category="AI|BehaviorTree")
@@ -212,5 +212,5 @@ protected:
 
 	DEPRECATED(4.7, "This function is deprecated. To implement your condition checking implement PerformConditionCheck or PerformConditionCheckAI function in your Blueprint.")
 	UFUNCTION(BlueprintImplementableEvent, meta = (DeprecatedFunction, DeprecationMessage = "To implement your condition checking implement PerformConditionCheck or PerformConditionCheckAI function in your Blueprint."))
-	virtual void ReceiveConditionCheck(AActor* OwnerActor);
+	void ReceiveConditionCheck(AActor* OwnerActor);
 };
diff --git a/Engine/Source/Runtime/AIModule/Classes/BehaviorTree/Services/BTService_BlueprintBase.h b/Engine/Source/Runtime/AIModule/Classes/BehaviorTree/Services/BTService_BlueprintBase.h
--- a/Engine/Source/Runtime/AIModule/Classes/BehaviorTree/Services/BTService_BlueprintBase.h
+++ b/Engine/Source/Runtime/AIModule/Classes/BehaviorTree/Services/BTService_BlueprintBase.h
@@ -74,53 +74,53 @@ protected:
 	 *	@Note that if both generic and AI event versions are implemented only the more
 	 *	suitable one will be called, meaning the AI version if called for AI, generic one otherwise */
 	UFUNCTION(BlueprintImplementableEvent)
-	virtual void ReceiveTick(AActor* OwnerActor, float DeltaSeconds);
+	void ReceiveTick(AActor* OwnerActor, float DeltaSeconds);
 
 	/** task search enters branch of tree
 	 *	@Note that if both generic and AI event versions are implemented only the more
 	 *	suitable one will be called, meaning the AI version if called for AI, generic one otherwise */
 	UFUNCTION(BlueprintImplementableEvent)
-	virtual void ReceiveSearchStart(AActor* OwnerActor);
+	void ReceiveSearchStart(AActor* OwnerActor);
 
 	/** service became active
 	 *	@Note that if both generic and AI event versions are implemented only the more
 	 *	suitable one will be called, meaning the AI version if called for AI, generic one otherwise */
 	UFUNCTION(BlueprintImplementableEvent)
-	virtual void ReceiveActivation(AActor* OwnerActor);
+	void ReceiveActivation(AActor* OwnerActor);
 
 	/** service became inactive
 	 *	@Note that if both generic and AI event versions are implemented only the more
 	 *	suitable one will be called, meaning the AI version if called for AI, generic one otherwise */
 	UFUNCTION(BlueprintImplementableEvent)
-	virtual void ReceiveDeactivation(AActor* OwnerActor);
+	void ReceiveDeactivation(AActor* OwnerActor);
 
 	/** Alternative AI version of ReceiveTick function.
 	 *	@see ReceiveTick for more details
 	 *	@Note that if both generic and AI event versions are implemented only the more
 	 *	suitable one will be called, meaning the AI version if called for AI, generic one otherwise */
 	UFUNCTION(BlueprintImplementableEvent, Category = AI)
-	virtual void ReceiveTickAI(AAIController* OwnerController, APawn* ControlledPawn, float DeltaSeconds);
+	void ReceiveTickAI(AAIController* OwnerController, APawn* ControlledPawn, float DeltaSeconds);
 
 	/** Alternative AI version of ReceiveSearchStart function.
 	 *	@see ReceiveSearchStart for more details
 	 *	@Note that if both generic and AI event versions are implemented only the more
 	 *	suitable one will be called, meaning the AI version if called for AI, generic one otherwise */
 	UFUNCTION(BlueprintImplementableEvent, Category = AI)
-	virtual void ReceiveSearchStartAI(AAIController* OwnerController, APawn* ControlledPawn);
+	void ReceiveSearchStartAI(AAIController* OwnerController, APawn* ControlledPawn);
 
 	/** Alternative AI version of ReceiveActivation function.
 	 *	@see ReceiveActivation for more details
 	 *	@Note that if both generic and AI event versions are implemented only the more
 	 *	suitable one will be called, meaning the AI version if called for AI, generic one otherwise */
 	UFUNCTION(BlueprintImplementableEvent, Category = AI)
-	virtual void ReceiveActivationAI(AAIController* OwnerController, APawn* ControlledPawn);
+	void ReceiveActivationAI(AAIController* OwnerController, APawn* ControlledPawn);
 
 	/** Alternative AI version of ReceiveDeactivation function.
 	 *	@see ReceiveDeactivation for more details
 	 *	@Note that if both generic and AI event versions are implemented only the more
 	 *	suitable one will be called, meaning the AI version if called for AI, generic one otherwise */
 	UFUNCTION(BlueprintImplementableEvent, Category = AI)
-	virtual void ReceiveDeactivationAI(AAIController* OwnerController, APawn* ControlledPawn);
+	void ReceiveDeactivationAI(AAIController* OwnerController, APawn* ControlledPawn);
 
 	/** check if service is currently being active */
 	UFUNCTION(BlueprintCallable, Category="AI|BehaviorTree")
diff --git a/Engine/Source/Runtime/AIModule/Classes/BehaviorTree/Tasks/BTTask_BlueprintBase.h b/Engine/Source/Runtime/AIModule/Classes/BehaviorTree/Tasks/BTTask_BlueprintBase.h
--- a/Engine/Source/Runtime/AIModule/Classes/BehaviorTree/Tasks/BTTask_BlueprintBase.h
+++ b/Engine/Source/Runtime/AIModule/Classes/BehaviorTree/Tasks/BTTask_BlueprintBase.h
@@ -72,40 +72,40 @@ protected:
 	 *	@Note that if both generic and AI event versions are implemented only the more 
 	 *	suitable one will be called, meaning the AI version if called for AI, generic one otherwise */
 	UFUNCTION(BlueprintImplementableEvent)
-	virtual void ReceiveExecute(AActor* OwnerActor);
+	void ReceiveExecute(AActor* OwnerActor);
 
 	/** if blueprint graph contains this event, task will stay active until FinishAbort is called
 	 *	@Note that if both generic and AI event versions are implemented only the more
 	 *	suitable one will be called, meaning the AI version if called for AI, generic one otherwise */
 	UFUNCTION(BlueprintImplementableEvent)
-	virtual void ReceiveAbort(AActor* OwnerActor);
+	void ReceiveAbort(AActor* OwnerActor);
 
 	/** tick function
 	 *	@Note that if both generic and AI event versions are implemented only the more
 	 *	suitable one will be called, meaning the AI version if called for AI, generic one otherwise */
 	UFUNCTION(BlueprintImplementableEvent)
-	virtual void ReceiveTick(AActor* OwnerActor, float DeltaSeconds);
+	void ReceiveTick(AActor* OwnerActor, float DeltaSeconds);
 
 	/** Alternative AI version of ReceiveExecute
 	*	@see ReceiveExecute for more details
 	 *	@Note that if both generic and AI event versions are implemented only the more
 	 *	suitable one will be called, meaning the AI version if called for AI, generic one otherwise */
 	UFUNCTION(BlueprintImplementableEvent, Category = AI)
-	virtual void ReceiveExecuteAI(AAIController* OwnerController, APawn* ControlledPawn);
+	void ReceiveExecuteAI(AAIController* OwnerController, APawn* ControlledPawn);
 
 	/** Alternative AI version of ReceiveAbort
 	 *	@see ReceiveAbort for more details
 	 *	@Note that if both generic and AI event versions are implemented only the more
 	 *	suitable one will be called, meaning the AI version if called for AI, generic one otherwise */
 	UFUNCTION(BlueprintImplementableEvent, Category = AI)
-	virtual void ReceiveAbortAI(AAIController* OwnerController, APawn* ControlledPawn);
+	void ReceiveAbortAI(AAIController* OwnerController, APawn* ControlledPawn);
 
 	/** Alternative AI version of tick function.
 	 *	@see ReceiveTick for more details
 	 *	@Note that if both generic and AI event versions are implemented only the more
 	 *	suitable one will be called, meaning the AI version if called for AI, generic one otherwise */
 	UFUNCTION(BlueprintImplementableEvent, Category = AI)
-	virtual void ReceiveTickAI(AAIController* OwnerController, APawn* ControlledPawn, float DeltaSeconds);
+	void ReceiveTickAI(AAIController* OwnerController, APawn* ControlledPawn, float DeltaSeconds);
 
 	/** finishes task execution with Success or Fail result */
 	UFUNCTION(BlueprintCallable, Category="AI|BehaviorTree")
diff --git a/Engine/Source/Runtime/AIModule/Classes/EnvironmentQuery/Contexts/EnvQueryContext_BlueprintBase.h b/Engine/Source/Runtime/AIModule/Classes/EnvironmentQuery/Contexts/EnvQueryContext_BlueprintBase.h
--- a/Engine/Source/Runtime/AIModule/Classes/EnvironmentQuery/Contexts/EnvQueryContext_BlueprintBase.h
+++ b/Engine/Source/Runtime/AIModule/Classes/EnvironmentQuery/Contexts/EnvQueryContext_BlueprintBase.h
@@ -31,14 +31,14 @@ class UEnvQueryContext_BlueprintBase : public UEnvQueryContext
 	virtual void ProvideContext(FEnvQueryInstance& QueryInstance, FEnvQueryContextData& ContextData) const override;
 
 	UFUNCTION(BlueprintImplementableEvent)
-	virtual void ProvideSingleActor(AActor* QuerierActor, AActor*& ResultingActor) const;
+	void ProvideSingleActor(AActor* QuerierActor, AActor*& ResultingActor) const;
 
 	UFUNCTION(BlueprintImplementableEvent)
-	virtual void ProvideSingleLocation(AActor* QuerierActor, FVector& ResultingLocation) const;
+	void ProvideSingleLocation(AActor* QuerierActor, FVector& ResultingLocation) const;
 
 	UFUNCTION(BlueprintImplementableEvent)
-	virtual void ProvideActorsSet(AActor* QuerierActor, TArray<AActor*>& ResultingActorsSet) const;
+	void ProvideActorsSet(AActor* QuerierActor, TArray<AActor*>& ResultingActorsSet) const;
 
 	UFUNCTION(BlueprintImplementableEvent)
-	virtual void ProvideLocationsSet(AActor* QuerierActor, TArray<FVector>& ResultingLocationSet) const;
+	void ProvideLocationsSet(AActor* QuerierActor, TArray<FVector>& ResultingLocationSet) const;
 };
diff --git a/Engine/Source/Runtime/AIModule/Classes/EnvironmentQuery/Generators/EnvQueryGenerator_BlueprintBase.h b/Engine/Source/Runtime/AIModule/Classes/EnvironmentQuery/Generators/EnvQueryGenerator_BlueprintBase.h
--- a/Engine/Source/Runtime/AIModule/Classes/EnvironmentQuery/Generators/EnvQueryGenerator_BlueprintBase.h
+++ b/Engine/Source/Runtime/AIModule/Classes/EnvironmentQuery/Generators/EnvQueryGenerator_BlueprintBase.h
@@ -29,7 +29,7 @@ class AIMODULE_API UEnvQueryGenerator_BlueprintBase : public UEnvQueryGenerator
 	virtual UWorld* GetWorld() const;
 
 	UFUNCTION(BlueprintImplementableEvent, Category = Generator)
-	virtual void DoItemGeneration(const TArray<FVector>& ContextLocations) const;
+	void DoItemGeneration(const TArray<FVector>& ContextLocations) const;
 
 	virtual void GenerateItems(FEnvQueryInstance& QueryInstance) const override;
 
diff --git a/Engine/Source/Runtime/AIModule/Classes/Perception/AISense_Blueprint.h b/Engine/Source/Runtime/AIModule/Classes/Perception/AISense_Blueprint.h
--- a/Engine/Source/Runtime/AIModule/Classes/Perception/AISense_Blueprint.h
+++ b/Engine/Source/Runtime/AIModule/Classes/Perception/AISense_Blueprint.h
@@ -35,25 +35,25 @@ public:
 	/** returns requested amount of time to pass until next frame. 
 	 *	Return 0 to get update every frame (WARNING: hits performance) */
 	UFUNCTION(BlueprintImplementableEvent)
-	virtual float OnUpdate(const TArray<UAISenseEvent*>& EventsToProcess);
+	float OnUpdate(const TArray<UAISenseEvent*>& EventsToProcess);
 
 	/**
 	 *	@param PerceptionComponent is ActorListener's AIPerceptionComponent instance
 	 */
 	UFUNCTION(BlueprintImplementableEvent)
-	virtual void OnListenerRegistered(AActor* ActorListener, UAIPerceptionComponent* PerceptionComponent);
+	void OnListenerRegistered(AActor* ActorListener, UAIPerceptionComponent* PerceptionComponent);
 
 	/**
 	 *	@param PerceptionComponent is ActorListener's AIPerceptionComponent instance
 	 */
 	UFUNCTION(BlueprintImplementableEvent)
-	virtual void OnListenerUpdated(AActor* ActorListener, UAIPerceptionComponent* PerceptionComponent);
+	void OnListenerUpdated(AActor* ActorListener, UAIPerceptionComponent* PerceptionComponent);
 
 	/** called when a listener unregistered from this sense. Most often this is called due to actor's death
 	 *	@param PerceptionComponent is ActorListener's AIPerceptionComponent instance
 	 */
 	UFUNCTION(BlueprintImplementableEvent)
-	virtual void OnListenerUnregistered(AActor* ActorListener, UAIPerceptionComponent* PerceptionComponent);
+	void OnListenerUnregistered(AActor* ActorListener, UAIPerceptionComponent* PerceptionComponent);
 	
 	UFUNCTION(BlueprintCallable, Category = "AI|Perception")
 	void GetAllListenerActors(TArray<AActor*>& ListenerActors) const;
diff --git a/Engine/Source/Runtime/Engine/Classes/AI/Navigation/NavLinkProxy.h b/Engine/Source/Runtime/Engine/Classes/AI/Navigation/NavLinkProxy.h
--- a/Engine/Source/Runtime/Engine/Classes/AI/Navigation/NavLinkProxy.h
+++ b/Engine/Source/Runtime/Engine/Classes/AI/Navigation/NavLinkProxy.h
@@ -73,7 +73,7 @@ public:
 
 	/** called when agent reaches smart link during path following, use ResumePathFollowing() to give control back */
 	UFUNCTION(BlueprintImplementableEvent)
-	virtual void ReceiveSmartLinkReached(AActor* Agent, const FVector& Destination);
+	void ReceiveSmartLinkReached(AActor* Agent, const FVector& Destination);
 
 	/** resume normal path following */
 	UFUNCTION(BlueprintCallable, Category="AI|Navigation")
