fixed crowd simulation updates for moving without pathfinding
#ue4 UE-12422

diff --git a/Engine/Source/Runtime/AIModule/Private/Navigation/CrowdFollowingComponent.cpp b/Engine/Source/Runtime/AIModule/Private/Navigation/CrowdFollowingComponent.cpp
--- a/Engine/Source/Runtime/AIModule/Private/Navigation/CrowdFollowingComponent.cpp
+++ b/Engine/Source/Runtime/AIModule/Private/Navigation/CrowdFollowingComponent.cpp
@@ -581,7 +581,6 @@ void UCrowdFollowingComponent::SetMoveSegment(int32 SegmentStartIndex)
 		MoveSegmentDirection = FVector::ZeroVector;
 
 		CurrentDestination.Set(Path->GetBaseActor(), CurrentTargetPt);
-		SuspendCrowdSteering(false);
 
 		LogPathPartHelper(GetOwner(), NavMeshPath, PathStartIndex, PathPartEndIdx);
 		UE_VLOG_SEGMENT(GetOwner(), LogCrowdFollowing, Log, MovementComp->GetActorFeetLocation(), CurrentTargetPt, FColor::Red, TEXT("path part"));
@@ -604,11 +603,10 @@ void UCrowdFollowingComponent::SetMoveSegment(int32 SegmentStartIndex)
 
 		bFinalPathPart = true;
 		bCheckMovementAngle = true;
-		bUpdateDirectMoveVelocity = DestinationActor.IsValid();
+		bUpdateDirectMoveVelocity = true;
 		CurrentDestination.Set(Path->GetBaseActor(), CurrentTargetPt);
 		CrowdAgentMoveDirection = (CurrentTargetPt - AgentLoc).GetSafeNormal();
 		MoveSegmentDirection = CrowdAgentMoveDirection;
-		SuspendCrowdSteering(true);
 
 		UE_VLOG(GetOwner(), LogCrowdFollowing, Log, TEXT("SetMoveSegment, direct move"));
 		UE_VLOG_SEGMENT(GetOwner(), LogCrowdFollowing, Log, AgentLoc, CurrentTargetPt, FColor::Red, TEXT("path"));
@@ -702,20 +700,22 @@ void UCrowdFollowingComponent::FollowPathSegment(float DeltaTime)
 		return;
 	}
 
-	if (bUpdateDirectMoveVelocity && DestinationActor.IsValid())
+	if (bUpdateDirectMoveVelocity)
 	{
 		const FVector CurrentTargetPt = DestinationActor->GetActorLocation();
-		const float DistSq = (CurrentTargetPt - GetCurrentTargetLocation()).SizeSquared();
-		if (DistSq > FMath::Square(10.0f))
-		{
-			UCrowdManager* Manager = UCrowdManager::GetCurrent(GetWorld());
-			const FVector AgentLoc = GetCrowdAgentLocation();
+		const FVector AgentLoc = GetCrowdAgentLocation();
+		const FVector NewDirection = (CurrentTargetPt - AgentLoc).GetSafeNormal();
 
+		const bool bDirectionChanged = !NewDirection.Equals(CrowdAgentMoveDirection);
+		if (bDirectionChanged)
+		{
 			CurrentDestination.Set(Path->GetBaseActor(), CurrentTargetPt);
-			CrowdAgentMoveDirection = (CurrentTargetPt - AgentLoc).GetSafeNormal();
-			MoveSegmentDirection = CrowdAgentMoveDirection;
+			CrowdAgentMoveDirection = NewDirection;
+			MoveSegmentDirection = NewDirection;
+
+			UCrowdManager* Manager = UCrowdManager::GetCurrent(GetWorld());
+			Manager->SetAgentMoveDirection(this, NewDirection);
 
-			Manager->SetAgentMoveDirection(this, MoveSegmentDirection);
 			UE_VLOG(GetOwner(), LogCrowdFollowing, Log, TEXT("Updated direct move direction for crowd agent."));
 		}
 	}
diff --git a/Engine/Source/Runtime/AIModule/Private/Navigation/CrowdManager.cpp b/Engine/Source/Runtime/AIModule/Private/Navigation/CrowdManager.cpp
--- a/Engine/Source/Runtime/AIModule/Private/Navigation/CrowdManager.cpp
+++ b/Engine/Source/Runtime/AIModule/Private/Navigation/CrowdManager.cpp
@@ -260,7 +260,7 @@ void UCrowdManager::Tick(float DeltaTime)
 					if (AgentData.bIsSimulated && AgentData.IsValid())
 					{
 						UCrowdFollowingComponent* CrowdComponent = Cast<UCrowdFollowingComponent>(It.Key());
-						if (CrowdComponent && CrowdComponent->IsCrowdSimulationActive())
+						if (CrowdComponent && CrowdComponent->IsCrowdSimulationEnabled())
 						{
 							ApplyVelocity(CrowdComponent, AgentData.AgentIndex);
 						}
