TArray::Init deprecated, as it's not obvious from the name that it leaves elements unconstructed.
TArray::Init use replaced with SetNumUninitialized.

diff --git a/Engine/Plugins/Messaging/UdpMessaging/Source/UdpMessaging/Private/Tunnel/UdpMessageTunnel.cpp b/Engine/Plugins/Messaging/UdpMessaging/Source/UdpMessaging/Private/Tunnel/UdpMessageTunnel.cpp
--- a/Engine/Plugins/Messaging/UdpMessaging/Source/UdpMessaging/Private/Tunnel/UdpMessageTunnel.cpp
+++ b/Engine/Plugins/Messaging/UdpMessaging/Source/UdpMessaging/Private/Tunnel/UdpMessageTunnel.cpp
@@ -259,7 +259,7 @@ void FUdpMessageTunnel::UdpToTcp( FSocket* Socket )
 	while (Socket->HasPendingData(DatagramSize))
 	{
 		FArrayReaderPtr Datagram = MakeShareable(new FArrayReader(true));
-		Datagram->Init(FMath::Min(DatagramSize, 65507u));
+		Datagram->SetNumUninitialized(FMath::Min(DatagramSize, 65507u));
 
 		int32 BytesRead = 0;
 
diff --git a/Engine/Plugins/Messaging/UdpMessaging/Source/UdpMessaging/Private/Tunnel/UdpMessageTunnelConnection.cpp b/Engine/Plugins/Messaging/UdpMessaging/Source/UdpMessaging/Private/Tunnel/UdpMessageTunnelConnection.cpp
--- a/Engine/Plugins/Messaging/UdpMessaging/Source/UdpMessaging/Private/Tunnel/UdpMessageTunnelConnection.cpp
+++ b/Engine/Plugins/Messaging/UdpMessaging/Source/UdpMessaging/Private/Tunnel/UdpMessageTunnelConnection.cpp
@@ -154,7 +154,7 @@ bool FUdpMessageTunnelConnection::ReceivePayloads()
 		int32 BytesRead = 0;
 			
 		FArrayReader PayloadSizeData = FArrayReader(true);
-		PayloadSizeData.Init(sizeof(uint16));
+		PayloadSizeData.SetNumUninitialized(sizeof(uint16));
 			
 		// ... read it from the stream without removing it...
 		if (!Socket->Recv(PayloadSizeData.GetData(), sizeof(uint16), BytesRead, ESocketReceiveFlags::Peek))
@@ -181,7 +181,7 @@ bool FUdpMessageTunnelConnection::ReceivePayloads()
 
 			// ... receive the payload
 			FArrayReaderPtr Payload = MakeShareable(new FArrayReader(true));
-			Payload->Init(PayloadSize);
+			Payload->SetNumUninitialized(PayloadSize);
 
 			if (!Socket->Recv(Payload->GetData(), Payload->Num(), BytesRead))
 			{
diff --git a/Engine/Plugins/Runtime/LiveEditorListenServer/Source/LiveEditorListenServer/Private/LiveEditorListenServer.cpp b/Engine/Plugins/Runtime/LiveEditorListenServer/Source/LiveEditorListenServer/Private/LiveEditorListenServer.cpp
--- a/Engine/Plugins/Runtime/LiveEditorListenServer/Source/LiveEditorListenServer/Private/LiveEditorListenServer.cpp
+++ b/Engine/Plugins/Runtime/LiveEditorListenServer/Source/LiveEditorListenServer/Private/LiveEditorListenServer.cpp
@@ -329,7 +329,7 @@ void FLiveEditorListenServer::Tick( float DeltaTime )
 		while ( Client->HasPendingData(DataSize) )
 		{
 			FArrayReaderPtr Datagram = MakeShareable(new FArrayReader(true));
-			Datagram->Init(FMath::Min(DataSize, 65507u));
+			Datagram->SetNumUninitialized(FMath::Min(DataSize, 65507u));
 
 			int32 BytesRead = 0;
 			if ( Client->Recv(Datagram->GetData(), Datagram->Num(), BytesRead) )
diff --git a/Engine/Source/Developer/MeshUtilities/Private/DisjointSet.h b/Engine/Source/Developer/MeshUtilities/Private/DisjointSet.h
--- a/Engine/Source/Developer/MeshUtilities/Private/DisjointSet.h
+++ b/Engine/Source/Developer/MeshUtilities/Private/DisjointSet.h
@@ -18,7 +18,7 @@ private:
 
 FDisjointSet::FDisjointSet( const uint32 Size )
 {
-	Parents.Init( Size );
+	Parents.SetNumUninitialized( Size );
 	for( uint32 i = 0; i < Size; i++ )
 	{
 		Parents[i] = i;
diff --git a/Engine/Source/Developer/MeshUtilities/Private/LayoutUV.cpp b/Engine/Source/Developer/MeshUtilities/Private/LayoutUV.cpp
--- a/Engine/Source/Developer/MeshUtilities/Private/LayoutUV.cpp
+++ b/Engine/Source/Developer/MeshUtilities/Private/LayoutUV.cpp
@@ -27,8 +27,8 @@ void FLayoutUV::FindCharts( const TMultiMap<int32,int32>& OverlappingCorners )
 	uint32 NumTris = NumIndexes / 3;
 
 	TArray< int32 > TranslatedMatches;
-	TranslatedMatches.Init( NumIndexes );
-	TexCoords.Init( NumIndexes );
+	TranslatedMatches.SetNumUninitialized( NumIndexes );
+	TexCoords.SetNumUninitialized( NumIndexes );
 	for( uint32 i = 0; i < NumIndexes; i++ )
 	{
 		TranslatedMatches[i] = -1;
@@ -125,7 +125,7 @@ void FLayoutUV::FindCharts( const TMultiMap<int32,int32>& OverlappingCorners )
 	}
 
 	// Sort tris by chart
-	SortedTris.Init( NumTris );
+	SortedTris.SetNumUninitialized( NumTris );
 	for( uint32 i = 0; i < NumTris; i++ )
 	{
 		// Flatten disjoint set path
@@ -997,7 +997,7 @@ void FLayoutUV::RasterizeChart( const FMeshChart& Chart, uint32 RectW, uint32 Re
 void FLayoutUV::CommitPackedUVs()
 {
 	// Alloc new UV channel
-	RawMesh->WedgeTexCoords[ DstChannel ].Init( TexCoords.Num() );
+	RawMesh->WedgeTexCoords[ DstChannel ].SetNumUninitialized( TexCoords.Num() );
 
 	// Commit chart UVs
 	for( int32 i = 0; i < Charts.Num(); i++ )
diff --git a/Engine/Source/Developer/TextureFormatPVR/Private/TextureFormatPVR.cpp b/Engine/Source/Developer/TextureFormatPVR/Private/TextureFormatPVR.cpp
--- a/Engine/Source/Developer/TextureFormatPVR/Private/TextureFormatPVR.cpp
+++ b/Engine/Source/Developer/TextureFormatPVR/Private/TextureFormatPVR.cpp
@@ -105,7 +105,7 @@ static bool SquarifyImage(FImage& Image, uint32 MinSquareSize)
     
     // Allocate room to fill out into
     TArray<uint32> SquareRawData;
-	SquareRawData.Init(SquareSize * SquareSize * Image.NumSlices);
+	SquareRawData.SetNumUninitialized(SquareSize * SquareSize * Image.NumSlices);
     
 	int32 SourceSliceSize = Image.SizeX * Image.SizeY;
 	int32 DestSliceSize = SquareSize * SquareSize;
@@ -133,7 +133,7 @@ static bool SquarifyImage(FImage& Image, uint32 MinSquareSize)
 
 	// Put the new image data into the existing Image (copying from uint32 array to uint8 array)
 	Image.RawData.Empty(SquareSize * SquareSize * Image.NumSlices * sizeof(uint32));
-	Image.RawData.Init(SquareSize * SquareSize * Image.NumSlices * sizeof(uint32));
+	Image.RawData.SetNumUninitialized(SquareSize * SquareSize * Image.NumSlices * sizeof(uint32));
 	uint32* FinalData = (uint32*)Image.RawData.GetData();
 	FMemory::Memcpy(Image.RawData.GetData(), SquareRawData.GetData(), SquareSize * SquareSize * Image.NumSlices * sizeof(uint32));
 
@@ -221,7 +221,7 @@ static void UseOriginal(const FImage& InImage, FCompressedImage2D& OutCompressed
     OutCompressedImage.PixelFormat = CompressedPixelFormat;
 	
 	// Output Data
-	OutCompressedImage.RawData.Init(Image.SizeX * Image.SizeY * 4);
+	OutCompressedImage.RawData.SetNumUninitialized(Image.SizeX * Image.SizeY * 4);
     void* MipData = (void*)Image.RawData.GetData();
     FMemory::Memcpy(MipData, Image.RawData.GetData(), Image.SizeX * Image.SizeY * 4);
 }
diff --git a/Engine/Source/Editor/UnrealEd/Private/EditorExporters.cpp b/Engine/Source/Editor/UnrealEd/Private/EditorExporters.cpp
--- a/Engine/Source/Editor/UnrealEd/Private/EditorExporters.cpp
+++ b/Engine/Source/Editor/UnrealEd/Private/EditorExporters.cpp
@@ -2775,7 +2775,7 @@ namespace MaterialExportUtils
 		GetRendererModule().BeginRenderingViewFamily(&Canvas, &ViewFamily);
 
 		// Copy the contents of the remote texture to system memory
-		OutSamples.Init(TargetSize.X*TargetSize.Y);
+		OutSamples.SetNumUninitialized(TargetSize.X*TargetSize.Y);
 		FReadSurfaceDataFlags ReadSurfaceDataFlags;
 		ReadSurfaceDataFlags.SetLinearToGamma(false);
 		RenderTargetResource->ReadPixelsPtr(OutSamples.GetData(), ReadSurfaceDataFlags, FIntRect(0, 0, TargetSize.X, TargetSize.Y));
diff --git a/Engine/Source/Editor/UnrealEd/Private/Factories/ApexDestructibleAssetImport.cpp b/Engine/Source/Editor/UnrealEd/Private/Factories/ApexDestructibleAssetImport.cpp
--- a/Engine/Source/Editor/UnrealEd/Private/Factories/ApexDestructibleAssetImport.cpp
+++ b/Engine/Source/Editor/UnrealEd/Private/Factories/ApexDestructibleAssetImport.cpp
@@ -936,13 +936,13 @@ UNREALED_API bool BuildDestructibleMeshFromFractureSettings(UDestructibleMesh& D
 	if (DestructibleMesh.FractureSettings != NULL)
 	{
 		TArray<UMaterialInterface*> OverrideMaterials;
-		OverrideMaterials.Init(DestructibleMesh.Materials.Num());	//save old materials
+		OverrideMaterials.SetNumUninitialized(DestructibleMesh.Materials.Num());	//save old materials
 		for (int32 MaterialIndex = 0; MaterialIndex < DestructibleMesh.Materials.Num(); ++MaterialIndex)
 		{
 			OverrideMaterials[MaterialIndex] = DestructibleMesh.Materials[MaterialIndex].MaterialInterface;
 		}
 
-		DestructibleMesh.Materials.Init(DestructibleMesh.FractureSettings->Materials.Num());
+		DestructibleMesh.Materials.SetNumUninitialized(DestructibleMesh.FractureSettings->Materials.Num());
 
 		for (int32 MaterialIndex = 0; MaterialIndex < DestructibleMesh.Materials.Num(); ++MaterialIndex)
 		{
diff --git a/Engine/Source/Editor/UnrealEd/Private/Factories/HDRLoader.cpp b/Engine/Source/Editor/UnrealEd/Private/Factories/HDRLoader.cpp
--- a/Engine/Source/Editor/UnrealEd/Private/Factories/HDRLoader.cpp
+++ b/Engine/Source/Editor/UnrealEd/Private/Factories/HDRLoader.cpp
@@ -211,7 +211,7 @@ void FHDRLoadHelper::OldDecompressScanline(uint8* Out, const uint8*& In, uint32
 void FHDRLoadHelper::ExtractDDSInRGBE(TArray<uint8>& OutDDSFile) const
 {
 	// header, one uint32 per texel, no mips
-	OutDDSFile.Init(4 + sizeof(FDDSFileHeader) + GetWidth() * GetHeight() * sizeof(uint32));
+	OutDDSFile.SetNumUninitialized(4 + sizeof(FDDSFileHeader) + GetWidth() * GetHeight() * sizeof(uint32));
 
 	// create the dds header
 	{
diff --git a/Engine/Source/Editor/WorldBrowser/Private/Tiles/WorldTileCollectionModel.cpp b/Engine/Source/Editor/WorldBrowser/Private/Tiles/WorldTileCollectionModel.cpp
--- a/Engine/Source/Editor/WorldBrowser/Private/Tiles/WorldTileCollectionModel.cpp
+++ b/Engine/Source/Editor/WorldBrowser/Private/Tiles/WorldTileCollectionModel.cpp
@@ -1438,7 +1438,7 @@ static bool ReadRawFile(TArray<DataType>& Result, const FString& Filename, uint3
 		}
 		return 0;
 	}
-	Result.Init(Reader->TotalSize()/Result.GetTypeSize());
+	Result.SetNumUninitialized(Reader->TotalSize()/Result.GetTypeSize());
 	Reader->Serialize(Result.GetData(), Result.Num()*Result.GetTypeSize());
 	return Reader->Close();
 }
diff --git a/Engine/Source/Runtime/AIModule/Private/EnvironmentQuery/EnvQueryInstance.cpp b/Engine/Source/Runtime/AIModule/Private/EnvironmentQuery/EnvQueryInstance.cpp
--- a/Engine/Source/Runtime/AIModule/Private/EnvironmentQuery/EnvQueryInstance.cpp
+++ b/Engine/Source/Runtime/AIModule/Private/EnvironmentQuery/EnvQueryInstance.cpp
@@ -83,7 +83,7 @@ bool FEnvQueryInstance::PrepareContext(UClass* Context, TArray<FEnvQuerySpatialD
 		const uint16 DefTypeValueSize = DefTypeOb->GetValueSize();
 		uint8* RawData = ContextData.RawData.GetData();
 
-		Data.Init(ContextData.NumValues);
+		Data.SetNumUninitialized(ContextData.NumValues);
 		for (int32 ValueIndex = 0; ValueIndex < ContextData.NumValues; ValueIndex++)
 		{
 			Data[ValueIndex].Location = DefTypeOb->GetItemLocation(RawData);
@@ -111,7 +111,7 @@ bool FEnvQueryInstance::PrepareContext(UClass* Context, TArray<FVector>& Data)
 		const uint16 DefTypeValueSize = DefTypeOb->GetValueSize();
 		uint8* RawData = (uint8*)ContextData.RawData.GetData();
 
-		Data.Init(ContextData.NumValues);
+		Data.SetNumUninitialized(ContextData.NumValues);
 		for (int32 ValueIndex = 0; ValueIndex < ContextData.NumValues; ValueIndex++)
 		{
 			Data[ValueIndex] = DefTypeOb->GetItemLocation(RawData);
@@ -138,7 +138,7 @@ bool FEnvQueryInstance::PrepareContext(UClass* Context, TArray<FRotator>& Data)
 		const uint16 DefTypeValueSize = DefTypeOb->GetValueSize();
 		uint8* RawData = ContextData.RawData.GetData();
 
-		Data.Init(ContextData.NumValues);
+		Data.SetNumUninitialized(ContextData.NumValues);
 		for (int32 ValueIndex = 0; ValueIndex < ContextData.NumValues; ValueIndex++)
 		{
 			Data[ValueIndex] = DefTypeOb->GetItemRotation(RawData);
diff --git a/Engine/Source/Runtime/AIModule/Private/EnvironmentQuery/Items/EnvQueryItemType_Actor.cpp b/Engine/Source/Runtime/AIModule/Private/EnvironmentQuery/Items/EnvQueryItemType_Actor.cpp
--- a/Engine/Source/Runtime/AIModule/Private/EnvironmentQuery/Items/EnvQueryItemType_Actor.cpp
+++ b/Engine/Source/Runtime/AIModule/Private/EnvironmentQuery/Items/EnvQueryItemType_Actor.cpp
@@ -41,7 +41,7 @@ void UEnvQueryItemType_Actor::SetContextHelper(FEnvQueryContextData& ContextData
 {
 	ContextData.ValueType = UEnvQueryItemType_Actor::StaticClass();
 	ContextData.NumValues = 1;
-	ContextData.RawData.Init(sizeof(FWeakObjectPtr));
+	ContextData.RawData.SetNumUninitialized(sizeof(FWeakObjectPtr));
 
 	UEnvQueryItemType_Actor::SetValue(ContextData.RawData.GetData(), SingleActor);
 }
@@ -50,7 +50,7 @@ void UEnvQueryItemType_Actor::SetContextHelper(FEnvQueryContextData& ContextData
 {
 	ContextData.ValueType = UEnvQueryItemType_Actor::StaticClass();
 	ContextData.NumValues = MultipleActors.Num();
-	ContextData.RawData.Init(sizeof(FWeakObjectPtr) * MultipleActors.Num());
+	ContextData.RawData.SetNumUninitialized(sizeof(FWeakObjectPtr) * MultipleActors.Num());
 
 	uint8* RawData = (uint8*)ContextData.RawData.GetData();
 	for (int32 ActorIndex = 0; ActorIndex < MultipleActors.Num(); ActorIndex++)
@@ -64,7 +64,7 @@ void UEnvQueryItemType_Actor::SetContextHelper(FEnvQueryContextData& ContextData
 {
 	ContextData.ValueType = UEnvQueryItemType_Actor::StaticClass();
 	ContextData.NumValues = MultipleActors.Num();
-	ContextData.RawData.Init(sizeof(FWeakObjectPtr)* MultipleActors.Num());
+	ContextData.RawData.SetNumUninitialized(sizeof(FWeakObjectPtr)* MultipleActors.Num());
 
 	uint8* RawData = (uint8*)ContextData.RawData.GetData();
 	for (int32 ActorIndex = 0; ActorIndex < MultipleActors.Num(); ActorIndex++)
diff --git a/Engine/Source/Runtime/AIModule/Private/EnvironmentQuery/Items/EnvQueryItemType_Direction.cpp b/Engine/Source/Runtime/AIModule/Private/EnvironmentQuery/Items/EnvQueryItemType_Direction.cpp
--- a/Engine/Source/Runtime/AIModule/Private/EnvironmentQuery/Items/EnvQueryItemType_Direction.cpp
+++ b/Engine/Source/Runtime/AIModule/Private/EnvironmentQuery/Items/EnvQueryItemType_Direction.cpp
@@ -43,7 +43,7 @@ void UEnvQueryItemType_Direction::SetContextHelper(FEnvQueryContextData& Context
 {
 	ContextData.ValueType = UEnvQueryItemType_Direction::StaticClass();
 	ContextData.NumValues = 1;
-	ContextData.RawData.Init(sizeof(FVector));
+	ContextData.RawData.SetNumUninitialized(sizeof(FVector));
 
 	UEnvQueryItemType_Direction::SetValue((uint8*)ContextData.RawData.GetData(), SingleDirection);
 }
@@ -52,7 +52,7 @@ void UEnvQueryItemType_Direction::SetContextHelper(FEnvQueryContextData& Context
 {
 	ContextData.ValueType = UEnvQueryItemType_Direction::StaticClass();
 	ContextData.NumValues = 1;
-	ContextData.RawData.Init(sizeof(FVector));
+	ContextData.RawData.SetNumUninitialized(sizeof(FVector));
 
 	UEnvQueryItemType_Direction::SetValueRot((uint8*)ContextData.RawData.GetData(), SingleRotation);
 }
@@ -61,7 +61,7 @@ void UEnvQueryItemType_Direction::SetContextHelper(FEnvQueryContextData& Context
 {
 	ContextData.ValueType = UEnvQueryItemType_Direction::StaticClass();
 	ContextData.NumValues = MultipleDirections.Num();
-	ContextData.RawData.Init(sizeof(FVector) * MultipleDirections.Num());
+	ContextData.RawData.SetNumUninitialized(sizeof(FVector) * MultipleDirections.Num());
 
 	uint8* RawData = (uint8*)ContextData.RawData.GetData();
 	for (int32 DirectionIndex = 0; DirectionIndex < MultipleDirections.Num(); DirectionIndex++)
@@ -75,7 +75,7 @@ void UEnvQueryItemType_Direction::SetContextHelper(FEnvQueryContextData& Context
 {
 	ContextData.ValueType = UEnvQueryItemType_Direction::StaticClass();
 	ContextData.NumValues = MultipleRotations.Num();
-	ContextData.RawData.Init(sizeof(FVector) * MultipleRotations.Num());
+	ContextData.RawData.SetNumUninitialized(sizeof(FVector) * MultipleRotations.Num());
 
 	uint8* RawData = (uint8*)ContextData.RawData.GetData();
 	for (int32 RotationIndex = 0; RotationIndex < MultipleRotations.Num(); RotationIndex++)
diff --git a/Engine/Source/Runtime/AIModule/Private/EnvironmentQuery/Items/EnvQueryItemType_Point.cpp b/Engine/Source/Runtime/AIModule/Private/EnvironmentQuery/Items/EnvQueryItemType_Point.cpp
--- a/Engine/Source/Runtime/AIModule/Private/EnvironmentQuery/Items/EnvQueryItemType_Point.cpp
+++ b/Engine/Source/Runtime/AIModule/Private/EnvironmentQuery/Items/EnvQueryItemType_Point.cpp
@@ -42,7 +42,7 @@ void UEnvQueryItemType_Point::SetContextHelper(FEnvQueryContextData& ContextData
 {
 	ContextData.ValueType = UEnvQueryItemType_Point::StaticClass();
 	ContextData.NumValues = 1;
-	ContextData.RawData.Init(sizeof(FNavLocation));
+	ContextData.RawData.SetNumUninitialized(sizeof(FNavLocation));
 
 	UEnvQueryItemType_Point::SetValue((uint8*)ContextData.RawData.GetData(), SinglePoint);
 }
@@ -51,7 +51,7 @@ void UEnvQueryItemType_Point::SetContextHelper(FEnvQueryContextData& ContextData
 {
 	ContextData.ValueType = UEnvQueryItemType_Point::StaticClass();
 	ContextData.NumValues = MultiplePoints.Num();
-	ContextData.RawData.Init(sizeof(FNavLocation)* MultiplePoints.Num());
+	ContextData.RawData.SetNumUninitialized(sizeof(FNavLocation)* MultiplePoints.Num());
 
 	uint8* RawData = (uint8*)ContextData.RawData.GetData();
 	for (int32 PointIndex = 0; PointIndex < MultiplePoints.Num(); PointIndex++)
diff --git a/Engine/Source/Runtime/Core/Private/Mac/MacPlatformProcess.cpp b/Engine/Source/Runtime/Core/Private/Mac/MacPlatformProcess.cpp
--- a/Engine/Source/Runtime/Core/Private/Mac/MacPlatformProcess.cpp
+++ b/Engine/Source/Runtime/Core/Private/Mac/MacPlatformProcess.cpp
@@ -975,7 +975,7 @@ bool FMacPlatformProcess::ReadPipeToArray(void* ReadPipe, TArray<uint8>& Output)
 
 	if (ReadPipe)
 	{
-		Output.Init(READ_SIZE);
+		Output.SetNumUninitialized(READ_SIZE);
 		int32 BytesRead = 0;
 		BytesRead = read([(NSFileHandle*)ReadPipe fileDescriptor], Output.GetData(), READ_SIZE);
 		if (BytesRead > 0)
diff --git a/Engine/Source/Runtime/Core/Private/Serialization/Bitstreams.cpp b/Engine/Source/Runtime/Core/Private/Serialization/Bitstreams.cpp
--- a/Engine/Source/Runtime/Core/Private/Serialization/Bitstreams.cpp
+++ b/Engine/Source/Runtime/Core/Private/Serialization/Bitstreams.cpp
@@ -266,7 +266,7 @@ void FBitWriterMark::Copy( FBitWriter& Writer, TArray<uint8> &Buffer )
 	int32 Bytes = (Writer.Num - Num + 7) >> 3;
 	if( Bytes > 0 )
 	{
-		Buffer.Init(Bytes);		// This makes room but doesnt zero
+		Buffer.SetNumUninitialized(Bytes);		// This makes room but doesnt zero
 		Buffer[Bytes-1] = 0;	// Make sure the last byte is 0 out, because appBitsCpy wont touch the last bits
 		appBitsCpy(Buffer.GetData(), 0, Writer.Buffer.GetData(), Num, Writer.Num - Num);
 	}
@@ -464,7 +464,7 @@ void FBitReaderMark::Copy( FBitReader& Reader, TArray<uint8> &Buffer )
 	int32 Bytes = (Reader.Pos - Pos + 7) >> 3;
 	if( Bytes > 0 )
 	{
-		Buffer.Init(Bytes);		// This makes room but doesnt zero
+		Buffer.SetNumUninitialized(Bytes);		// This makes room but doesnt zero
 		Buffer[Bytes-1] = 0;	// Make sure the last byte is 0 out, because appBitsCpy wont touch the last bits
 		appBitsCpy(Buffer.GetData(), 0, Reader.Buffer.GetData(), Pos, Reader.Pos - Pos);
 	}
diff --git a/Engine/Source/Runtime/Core/Private/Windows/WindowsPlatformProcess.cpp b/Engine/Source/Runtime/Core/Private/Windows/WindowsPlatformProcess.cpp
--- a/Engine/Source/Runtime/Core/Private/Windows/WindowsPlatformProcess.cpp
+++ b/Engine/Source/Runtime/Core/Private/Windows/WindowsPlatformProcess.cpp
@@ -1091,7 +1091,7 @@ bool FWindowsPlatformProcess::ReadPipeToArray(void* ReadPipe, TArray<uint8> & Ou
 	uint32 BytesAvailable = 0;
 	if (::PeekNamedPipe(ReadPipe, NULL, 0, NULL, (::DWORD*)&BytesAvailable, NULL) && (BytesAvailable > 0))
 	{
-		Output.Init(BytesAvailable);
+		Output.SetNumUninitialized(BytesAvailable);
 		uint32 BytesRead = 0;
 		if (::ReadFile(ReadPipe, Output.GetData(), BytesAvailable, (::DWORD*)&BytesRead, NULL))
 		{
diff --git a/Engine/Source/Runtime/Core/Public/Containers/Array.h b/Engine/Source/Runtime/Core/Public/Containers/Array.h
--- a/Engine/Source/Runtime/Core/Public/Containers/Array.h
+++ b/Engine/Source/Runtime/Core/Public/Containers/Array.h
@@ -1844,6 +1844,7 @@ public:
 	 * @param Number The number of elements that the array should be able to
 	 *               contain after allocation.
 	 */
+	DEPRECATED(4.8, "Init is deprecated - please use SetNumUninitialized(Number) instead.")
 	void Init(int32 Number)
 	{
 		Empty(Number);
diff --git a/Engine/Source/Runtime/Engine/Private/AI/Navigation/NavigationModifier.cpp b/Engine/Source/Runtime/Engine/Private/AI/Navigation/NavigationModifier.cpp
--- a/Engine/Source/Runtime/Engine/Private/AI/Navigation/NavigationModifier.cpp
+++ b/Engine/Source/Runtime/Engine/Private/AI/Navigation/NavigationModifier.cpp
@@ -143,7 +143,7 @@ FAreaNavModifier::FAreaNavModifier(float Radius, float Height, const FTransform&
 	Radius *= FMath::Max(Scale3D.X, Scale3D.Y);
 	Height *= Scale3D.Z;
 
-	Points.Init(2);
+	Points.SetNumUninitialized(2);
 	Points[0] = LocalToWorld.GetLocation();
 	Points[1].X = Radius;
 	Points[1].Z = Height;
@@ -275,7 +275,7 @@ void FAreaNavModifier::SetBox(const FBox& Box, const FTransform& LocalToWorld)
 			Bounds += Corners[i];
 		}
 
-		Points.Init(2);
+		Points.SetNumUninitialized(2);
 		Points[0] = Bounds.GetCenter();
 		Points[1] = Bounds.GetExtent();
 		ShapeType = ENavigationShapeType::Box;
diff --git a/Engine/Source/Runtime/Engine/Private/AI/Navigation/RecastNavMeshGenerator.cpp b/Engine/Source/Runtime/Engine/Private/AI/Navigation/RecastNavMeshGenerator.cpp
--- a/Engine/Source/Runtime/Engine/Private/AI/Navigation/RecastNavMeshGenerator.cpp
+++ b/Engine/Source/Runtime/Engine/Private/AI/Navigation/RecastNavMeshGenerator.cpp
@@ -2122,8 +2122,8 @@ void FRecastTileGenerator::AppendGeometry(const TNavStatArray<uint8>& RawCollisi
 	const int32 NumCoords = CollisionCache.Header.NumVerts * 3;
 	const int32 NumIndices = CollisionCache.Header.NumFaces * 3;
 	
-	GeometryElement.GeomCoords.Init(NumCoords);
-	GeometryElement.GeomIndices.Init(NumIndices);
+	GeometryElement.GeomCoords.SetNumUninitialized(NumCoords);
+	GeometryElement.GeomIndices.SetNumUninitialized(NumIndices);
 
 	FMemory::Memcpy(GeometryElement.GeomCoords.GetData(), CollisionCache.Verts, sizeof(float) * NumCoords);
 	FMemory::Memcpy(GeometryElement.GeomIndices.GetData(), CollisionCache.Indices, sizeof(int32) * NumIndices);
@@ -2203,7 +2203,7 @@ static void RasterizeGeometry(
 	const TArray<float>& Coords, const TArray<int32>& Indices, const FTransform& LocalToWorld, FTileRasterizationContext& RasterContext)
 {
 	TArray<float> WorldRecastCoords;
-	WorldRecastCoords.Init(Coords.Num());
+	WorldRecastCoords.SetNumUninitialized(Coords.Num());
 	
 	FMatrix LocalToRecastWorld = LocalToWorld.ToMatrixWithScale()*Unreal2RecastMatrix();
 	// Convert geometry to recast world space
diff --git a/Engine/Source/Runtime/Engine/Private/Animation/AnimationRuntime.cpp b/Engine/Source/Runtime/Engine/Private/Animation/AnimationRuntime.cpp
--- a/Engine/Source/Runtime/Engine/Private/Animation/AnimationRuntime.cpp
+++ b/Engine/Source/Runtime/Engine/Private/Animation/AnimationRuntime.cpp
@@ -1563,7 +1563,7 @@ void FA2CSPose::LocalBlendCSBoneTransforms
 		BoneMask.AddZeroed( Bones.Num() );
 
 		TArray<struct FBoneTransform> LocalBoneTransforms;
-		LocalBoneTransforms.Init( BoneTransforms.Num() );
+		LocalBoneTransforms.SetNumUninitialized( BoneTransforms.Num() );
 
 		const TArray<FBoneIndexType> & RequiredBoneIndices = BoneContainer->GetBoneIndicesArray();
 		// First, convert BoneTransforms to local space for blending.
diff --git a/Engine/Source/Runtime/Engine/Private/ColorVertexBuffer.cpp b/Engine/Source/Runtime/Engine/Private/ColorVertexBuffer.cpp
--- a/Engine/Source/Runtime/Engine/Private/ColorVertexBuffer.cpp
+++ b/Engine/Source/Runtime/Engine/Private/ColorVertexBuffer.cpp
@@ -302,7 +302,7 @@ void FColorVertexBuffer::GetVertexColors( TArray<FColor>& OutColors )
 {
 	if( VertexData != NULL && NumVertices > 0 )
 	{
-		OutColors.Init( NumVertices );
+		OutColors.SetNumUninitialized( NumVertices );
 
 		FMemory::Memcpy( OutColors.GetData(), VertexData->GetDataPointer(), NumVertices * VertexData->GetStride() ) ;
 	}
diff --git a/Engine/Source/Runtime/Engine/Private/DestructibleFractureSettings.cpp b/Engine/Source/Runtime/Engine/Private/DestructibleFractureSettings.cpp
--- a/Engine/Source/Runtime/Engine/Private/DestructibleFractureSettings.cpp
+++ b/Engine/Source/Runtime/Engine/Private/DestructibleFractureSettings.cpp
@@ -70,7 +70,7 @@ static void BuildApexRenderMesh(NxRenderMeshAssetAuthoring& RenderMeshAssetAutho
 		SubmeshDesc.m_numVertexBuffers = 1;
 		SubmeshDesc.m_vertexBuffers = &vb;
 
-		PartIndices.Init(HMesh.partCount());
+		PartIndices.SetNumUninitialized(HMesh.partCount());
 		for(uint32 PartIndex = 0; PartIndex < HMesh.partCount(); ++PartIndex)
 		{
 			TArray<NxVertex> PartVertices;
@@ -104,7 +104,7 @@ static void BuildApexRenderMesh(NxRenderMeshAssetAuthoring& RenderMeshAssetAutho
 			if (PartVertices.Num() > 0)
 			{
 				TArray<PxU32> Map;
-				Map.Init(PartVertices.Num());
+				Map.SetNumUninitialized(PartVertices.Num());
 				const PxU32 ReducedPartVertexCount = RenderMeshAssetAuthor.createReductionMap(Map.GetData(), PartVertices.GetData(), NULL, (PxU32)PartVertices.Num(), PxVec3( 0.0001f ), 0.001f, 1.0f/256.01f);
 				const PxU32 VertexPartStart = (PxU32)Vertices.Num();
 				Resize(Vertices, NxVertex(), VertexPartStart + ReducedPartVertexCount);
diff --git a/Engine/Source/Runtime/Engine/Private/InstancedStaticMesh.cpp b/Engine/Source/Runtime/Engine/Private/InstancedStaticMesh.cpp
--- a/Engine/Source/Runtime/Engine/Private/InstancedStaticMesh.cpp
+++ b/Engine/Source/Runtime/Engine/Private/InstancedStaticMesh.cpp
@@ -961,7 +961,7 @@ void UInstancedStaticMeshComponent::CreateAllInstanceBodies()
 	if (BodySetup)
 	{
 	    int32 NumBodies = PerInstanceSMData.Num();
-	    InstanceBodies.Init(NumBodies);
+	    InstanceBodies.SetNumUninitialized(NumBodies);
     
 	    TArray<FTransform> Transforms;
 	    Transforms.Reserve(NumBodies);
diff --git a/Engine/Source/Runtime/Landscape/Private/LandscapeCollision.cpp b/Engine/Source/Runtime/Landscape/Private/LandscapeCollision.cpp
--- a/Engine/Source/Runtime/Landscape/Private/LandscapeCollision.cpp
+++ b/Engine/Source/Runtime/Landscape/Private/LandscapeCollision.cpp
@@ -492,7 +492,7 @@ bool ULandscapeHeightfieldCollisionComponent::CookCollisionData(const FName& For
 
 	if (Result)
 	{
-		OutCookedData.Init(OutData.Num());
+		OutCookedData.SetNumUninitialized(OutData.Num());
 		FMemory::Memcpy(OutCookedData.GetData(), OutData.GetData(), OutData.Num());
 
 		if (bShouldSaveCookedDataToDDC[CookedDataIndex])
@@ -574,7 +574,7 @@ bool ULandscapeMeshCollisionComponent::CookCollisionData(const FName& Format, bo
 	}
 
 	// Scale all verts into temporary vertex buffer.
-	Vertices.Init(NumVerts);
+	Vertices.SetNumUninitialized(NumVerts);
 	for (int32 i = 0; i < NumVerts; i++)
 	{
 		int32 X = i % CollisionSizeVerts;
@@ -583,10 +583,10 @@ bool ULandscapeMeshCollisionComponent::CookCollisionData(const FName& Format, bo
 	}
 
 	const int32 NumTris = FMath::Square(CollisionSizeQuads) * 2;
-	Indices.Init(NumTris);
+	Indices.SetNumUninitialized(NumTris);
 	if (DominantLayers)
 	{
-		MaterialIndices.Init(NumTris);
+		MaterialIndices.SetNumUninitialized(NumTris);
 	}
 
 	int32 TriangleIdx = 0;
@@ -680,7 +680,7 @@ bool ULandscapeMeshCollisionComponent::CookCollisionData(const FName& Format, bo
 
 	if (Result)
 	{
-		OutCookedData.Init(OutData.Num());
+		OutCookedData.SetNumUninitialized(OutData.Num());
 		FMemory::Memcpy(OutCookedData.GetData(), OutData.GetData(), OutData.Num());
 
 		if (bShouldSaveCookedDataToDDC[CookedDataIndex])
diff --git a/Engine/Source/Runtime/Landscape/Private/LandscapeGrass.cpp b/Engine/Source/Runtime/Landscape/Private/LandscapeGrass.cpp
--- a/Engine/Source/Runtime/Landscape/Private/LandscapeGrass.cpp
+++ b/Engine/Source/Runtime/Landscape/Private/LandscapeGrass.cpp
@@ -431,7 +431,7 @@ public:
 	void ApplyResults()
 	{
 		TArray<FColor> Samples;
-		Samples.Init(TargetSize.X*TargetSize.Y);
+		Samples.SetNumUninitialized(TargetSize.X*TargetSize.Y);
 
 		// Copy the contents of the remote texture to system memory
 		FReadSurfaceDataFlags ReadSurfaceDataFlags;
diff --git a/Engine/Source/Runtime/Networking/Public/Common/UdpSocketReceiver.h b/Engine/Source/Runtime/Networking/Public/Common/UdpSocketReceiver.h
--- a/Engine/Source/Runtime/Networking/Public/Common/UdpSocketReceiver.h
+++ b/Engine/Source/Runtime/Networking/Public/Common/UdpSocketReceiver.h
@@ -93,7 +93,7 @@ public:
 			while (Socket->HasPendingData(Size))
 			{
 				FArrayReaderPtr Reader = MakeShareable(new FArrayReader(true));
-				Reader->Init(FMath::Min(Size, 65507u));
+				Reader->SetNumUninitialized(FMath::Min(Size, 65507u));
 
 				int32 Read = 0;
 				Socket->RecvFrom(Reader->GetData(), Reader->Num(), Read, *Sender);
diff --git a/Engine/Source/Runtime/Renderer/Private/CompositionLighting/PostProcessDeferredDecals.cpp b/Engine/Source/Runtime/Renderer/Private/CompositionLighting/PostProcessDeferredDecals.cpp
--- a/Engine/Source/Runtime/Renderer/Private/CompositionLighting/PostProcessDeferredDecals.cpp
+++ b/Engine/Source/Runtime/Renderer/Private/CompositionLighting/PostProcessDeferredDecals.cpp
@@ -555,7 +555,7 @@ public:
 	{
 		const int32 NumVerts = 8;
 		TResourceArray<FVector4, VERTEXBUFFER_ALIGNMENT> Verts;
-		Verts.Init(NumVerts);
+		Verts.SetNumUninitialized(NumVerts);
 
 		for (uint32 Z = 0; Z < 2; Z++)
 		{
diff --git a/Engine/Source/Runtime/Renderer/Private/PostProcess/SceneFilterRendering.cpp b/Engine/Source/Runtime/Renderer/Private/PostProcess/SceneFilterRendering.cpp
--- a/Engine/Source/Runtime/Renderer/Private/PostProcess/SceneFilterRendering.cpp
+++ b/Engine/Source/Runtime/Renderer/Private/PostProcess/SceneFilterRendering.cpp
@@ -19,7 +19,7 @@ public:
 	void InitRHI() override
 	{
 		TResourceArray<FFilterVertex, VERTEXBUFFER_ALIGNMENT> Vertices;
-		Vertices.Init(6);
+		Vertices.SetNumUninitialized(6);
 
 		Vertices[0].Position = FVector4(1,  1,	0,	1);
 		Vertices[0].UV = FVector2D(1,	1);
