fixed behavior tree's search update for branch with loop decorator
#ue4, merged from dev branch

diff --git a/Engine/Source/Runtime/AIModule/Classes/BehaviorTree/BehaviorTreeTypes.h b/Engine/Source/Runtime/AIModule/Classes/BehaviorTree/BehaviorTreeTypes.h
--- a/Engine/Source/Runtime/AIModule/Classes/BehaviorTree/BehaviorTreeTypes.h
+++ b/Engine/Source/Runtime/AIModule/Classes/BehaviorTree/BehaviorTreeTypes.h
@@ -144,7 +144,6 @@ namespace EBTNodeUpdateMode
 	enum Type
 	{
 		Add,				// add node
-		AddForLowerPri,		// add node only when new task has lower priority
 		Remove,				// remove node
 	};
 }
diff --git a/Engine/Source/Runtime/AIModule/Private/BehaviorTree/BTCompositeNode.cpp b/Engine/Source/Runtime/AIModule/Private/BehaviorTree/BTCompositeNode.cpp
--- a/Engine/Source/Runtime/AIModule/Private/BehaviorTree/BTCompositeNode.cpp
+++ b/Engine/Source/Runtime/AIModule/Private/BehaviorTree/BTCompositeNode.cpp
@@ -211,7 +211,7 @@ void UBTCompositeNode::NotifyDecoratorsOnDeactivation(FBehaviorTreeSearchData& S
 			else if (DecoratorOb->GetFlowAbortMode() == EBTFlowAbortMode::LowerPriority)
 			{
 				// - observers with mode "Lower Priority" will try to reactivate themselves ("Both" is not removed on node activation)
-				SearchData.AddUniqueUpdate(FBehaviorTreeSearchUpdate(DecoratorOb, SearchData.OwnerComp.GetActiveInstanceIdx(), EBTNodeUpdateMode::AddForLowerPri));
+				SearchData.AddUniqueUpdate(FBehaviorTreeSearchUpdate(DecoratorOb, SearchData.OwnerComp.GetActiveInstanceIdx(), EBTNodeUpdateMode::Add));
 			}
 		}
 	}
diff --git a/Engine/Source/Runtime/AIModule/Private/BehaviorTree/BehaviorTreeComponent.cpp b/Engine/Source/Runtime/AIModule/Private/BehaviorTree/BehaviorTreeComponent.cpp
--- a/Engine/Source/Runtime/AIModule/Private/BehaviorTree/BehaviorTreeComponent.cpp
+++ b/Engine/Source/Runtime/AIModule/Private/BehaviorTree/BehaviorTreeComponent.cpp
@@ -852,8 +852,7 @@ void UBehaviorTreeComponent::ApplySearchUpdates(const TArray<FBehaviorTreeSearch
 		}
 
 		if ((UpdateInfo.Mode == EBTNodeUpdateMode::Remove && !bIsActive) ||
-			(UpdateInfo.Mode == EBTNodeUpdateMode::Add && bIsActive) ||
-			(UpdateInfo.Mode == EBTNodeUpdateMode::AddForLowerPri && (bIsActive || UpdateNode->GetExecutionIndex() > NewNodeExecutionIndex)) ||
+			(UpdateInfo.Mode == EBTNodeUpdateMode::Add && (bIsActive || UpdateNode->GetExecutionIndex() > NewNodeExecutionIndex)) ||
 			(UpdateInfo.bPostUpdate != bPostUpdate))
 		{
 			continue;
diff --git a/Engine/Source/Runtime/AIModule/Private/BehaviorTree/BehaviorTreeTypes.cpp b/Engine/Source/Runtime/AIModule/Private/BehaviorTree/BehaviorTreeTypes.cpp
--- a/Engine/Source/Runtime/AIModule/Private/BehaviorTree/BehaviorTreeTypes.cpp
+++ b/Engine/Source/Runtime/AIModule/Private/BehaviorTree/BehaviorTreeTypes.cpp
@@ -160,9 +160,7 @@ void FBehaviorTreeSearchData::AddUniqueUpdate(const FBehaviorTreeSearchUpdate& U
 		if (Info.AuxNode == UpdateInfo.AuxNode && Info.TaskNode == UpdateInfo.TaskNode)
 		{
 			// duplicate, skip
-			// special case: AddForLowerPri replacing Add
-			if ((Info.Mode == UpdateInfo.Mode) ||
-				(Info.Mode == EBTNodeUpdateMode::Add && UpdateInfo.Mode == EBTNodeUpdateMode::AddForLowerPri))
+			if (Info.Mode == UpdateInfo.Mode)
 			{
 				bSkipAdding = true;
 				break;
@@ -401,7 +399,7 @@ FString UBehaviorTreeTypes::DescribeTaskStatus(EBTTaskStatus::Type TaskStatus)
 
 FString UBehaviorTreeTypes::DescribeNodeUpdateMode(EBTNodeUpdateMode::Type UpdateMode)
 {
-	static FString UpdateModeDesc[] = { TEXT("Add"), TEXT("AddForLowerPri"), TEXT("Remove") };
+	static FString UpdateModeDesc[] = { TEXT("Add"), TEXT("Remove") };
 	return (UpdateMode < ARRAY_COUNT(UpdateModeDesc)) ? UpdateModeDesc[UpdateMode] : FString();
 }
 
diff --git a/Engine/Source/Runtime/AIModule/Private/BehaviorTree/Composites/BTComposite_SimpleParallel.cpp b/Engine/Source/Runtime/AIModule/Private/BehaviorTree/Composites/BTComposite_SimpleParallel.cpp
--- a/Engine/Source/Runtime/AIModule/Private/BehaviorTree/Composites/BTComposite_SimpleParallel.cpp
+++ b/Engine/Source/Runtime/AIModule/Private/BehaviorTree/Composites/BTComposite_SimpleParallel.cpp
@@ -119,12 +119,12 @@ void UBTComposite_SimpleParallel::NotifyNodeDeactivation(FBehaviorTreeSearchData
 	const FBTNodeIndex LastBackgroundIndex(ActiveInstanceIdx, GetLastExecutionIndex());
 	SearchData.OwnerComp.UnregisterAuxNodesUpTo(FirstBackgroundIndex);
 
-	// remove all pending updates "AddForLowerPri" from background tree 
+	// remove all pending updates "Add" from background tree 
 	// it doesn't make sense for decorators to reactivate themselves there
 	for (int32 Idx = SearchData.PendingUpdates.Num() - 1; Idx >= 0; Idx--)
 	{
 		const FBehaviorTreeSearchUpdate& UpdateInfo = SearchData.PendingUpdates[Idx];
-		if (UpdateInfo.Mode == EBTNodeUpdateMode::AddForLowerPri)
+		if (UpdateInfo.Mode == EBTNodeUpdateMode::Add)
 		{
 			const uint16 UpdateNodeIdx = UpdateInfo.AuxNode ? UpdateInfo.AuxNode->GetExecutionIndex() : UpdateInfo.TaskNode->GetExecutionIndex();
 			const FBTNodeIndex UpdateIdx(UpdateInfo.InstanceIndex, UpdateNodeIdx);
