[UEDOC-1420] Fixed comments so they'll show up properly.

diff --git a/Engine/Source/Runtime/AIModule/Classes/Navigation/PathFollowingComponent.h b/Engine/Source/Runtime/AIModule/Classes/Navigation/PathFollowingComponent.h
--- a/Engine/Source/Runtime/AIModule/Classes/Navigation/PathFollowingComponent.h
+++ b/Engine/Source/Runtime/AIModule/Classes/Navigation/PathFollowingComponent.h
@@ -21,10 +21,17 @@ namespace EPathFollowingStatus
 {
 	enum Type
 	{
-		Idle,				// no requests
-		Waiting,			// request with incomplete path, will start after UpdateMove()
-		Paused,				// request paused, will continue after ResumeMove()
-		Moving,				// following path
+		/** No requests */
+		Idle,
+
+		/** Request with incomplete path, will start after UpdateMove() */
+		Waiting,
+
+		/** Request paused, will continue after ResumeMove() */
+		Paused,
+
+		/** Following path */
+		Moving,
 	};
 }
 
@@ -33,12 +40,23 @@ namespace EPathFollowingResult
 {
 	enum Type
 	{
-		Success,			// reached destination
-		Blocked,			// movement was blocked
-		OffPath,			// agent is not on path
-		Aborted,			// aborted and stopped (failure)
-		Skipped,			// aborted and replaced with new request
-		Invalid,			// request was invalid
+		/** Reached destination */
+		Success,
+
+		/** Movement was blocked */
+		Blocked,
+
+		/** Agent is not on path */
+		OffPath,
+
+		/** Aborted and stopped (failure) */
+		Aborted,
+
+		/** Aborted and replaced with new request */
+		Skipped,
+
+		/** Request was invalid */
+		Invalid,
 	};
 }
 
@@ -82,8 +100,11 @@ namespace EPathFollowingMessage
 {
 	enum Type
 	{
-		NoPath,			// abort reason for AI message
-		OtherRequest,	// abort reason for AI message
+		/** Aborted because no path was found */
+		NoPath,
+
+		/** Aborted because another request came in */
+		OtherRequest,
 	};
 }
 
