#ue4 - Perf: Change all uses of "GetActorRotation().Vector()" to "GetActorForwardVector()" which is ~10x faster.

diff --git a/Engine/Source/Developer/GameplayDebugger/Private/GameplayDebuggingHUDComponent.cpp b/Engine/Source/Developer/GameplayDebugger/Private/GameplayDebuggingHUDComponent.cpp
--- a/Engine/Source/Developer/GameplayDebugger/Private/GameplayDebuggingHUDComponent.cpp
+++ b/Engine/Source/Developer/GameplayDebugger/Private/GameplayDebuggingHUDComponent.cpp
@@ -648,7 +648,7 @@ void AGameplayDebuggingHUDComponent::DrawPerception(APlayerController* PC, class
 				BTAI->GetAIPerceptionComponent()->DrawDebugInfo(DefaultContext.Canvas);
 
 				const FVector AILocation = MyPawn->GetActorLocation();
-				const FVector Facing = MyPawn->GetActorRotation().Vector();
+				const FVector Facing = MyPawn->GetActorForwardVector();
 
 				UAIPerceptionSystem* PerceptionSys = UAIPerceptionSystem::GetCurrent(this);
 				if (PerceptionSys)
diff --git a/Engine/Source/Runtime/AIModule/Private/BehaviorTree/Tasks/BTTask_RotateToFaceBBEntry.cpp b/Engine/Source/Runtime/AIModule/Private/BehaviorTree/Tasks/BTTask_RotateToFaceBBEntry.cpp
--- a/Engine/Source/Runtime/AIModule/Private/BehaviorTree/Tasks/BTTask_RotateToFaceBBEntry.cpp
+++ b/Engine/Source/Runtime/AIModule/Private/BehaviorTree/Tasks/BTTask_RotateToFaceBBEntry.cpp
@@ -65,7 +65,7 @@ EBTNodeResult::Type UBTTask_RotateToFaceBBEntry::ExecuteTask(UBehaviorTreeCompon
 
 		if (ActorValue != NULL)
 		{
-			const float AngleDifference = CalculateAngleDifferenceDot(Pawn->GetActorRotation().Vector()
+			const float AngleDifference = CalculateAngleDifferenceDot(Pawn->GetActorForwardVector()
 				, (ActorValue->GetActorLocation() - PawnLocation).GetSafeNormal2D());
 			
 			if (AngleDifference >= PrecisionDot)
@@ -87,7 +87,7 @@ EBTNodeResult::Type UBTTask_RotateToFaceBBEntry::ExecuteTask(UBehaviorTreeCompon
 		
 		if (FAISystem::IsValidLocation(KeyValue))
 		{
-			const float AngleDifference = CalculateAngleDifferenceDot(Pawn->GetActorRotation().Vector()
+			const float AngleDifference = CalculateAngleDifferenceDot(Pawn->GetActorForwardVector()
 				, (KeyValue - PawnLocation).GetSafeNormal2D());
 
 			if (AngleDifference >= PrecisionDot)
@@ -109,7 +109,7 @@ EBTNodeResult::Type UBTTask_RotateToFaceBBEntry::ExecuteTask(UBehaviorTreeCompon
 		if (FAISystem::IsValidRotation(KeyValue))
 		{
 			const FVector DirectionVector = KeyValue.Vector();
-			const float AngleDifference = CalculateAngleDifferenceDot(Pawn->GetActorRotation().Vector(), DirectionVector);
+			const float AngleDifference = CalculateAngleDifferenceDot(Pawn->GetActorForwardVector(), DirectionVector);
 
 			if (AngleDifference >= PrecisionDot)
 			{
@@ -139,7 +139,7 @@ void UBTTask_RotateToFaceBBEntry::TickTask(UBehaviorTreeComponent& OwnerComp, ui
 	}
 	else
 	{
-		const FVector PawnDirection = AIController->GetPawn()->GetActorRotation().Vector();				
+		const FVector PawnDirection = AIController->GetPawn()->GetActorForwardVector();				
 		const FVector FocalPoint = AIController->GetFocalPointForPriority(EAIFocusPriority::Gameplay);
 
 		if (CalculateAngleDifferenceDot(PawnDirection, (FocalPoint - AIController->GetPawn()->GetActorLocation()).GetSafeNormal2D()) >= PrecisionDot)
@@ -192,7 +192,7 @@ void UBTTask_RotateToFaceBBEntry::DescribeRuntimeValues(const UBehaviorTreeCompo
 
 	if (AIController != NULL && AIController->GetPawn() != NULL)
 	{
-		const FVector PawnDirection = AIController->GetPawn()->GetActorRotation().Vector();
+		const FVector PawnDirection = AIController->GetPawn()->GetActorForwardVector();
 		const FVector FocalPoint = AIController->GetFocalPointForPriority(EAIFocusPriority::Gameplay);
 
 		if (FocalPoint != FAISystem::InvalidLocation)
diff --git a/Engine/Source/Runtime/AIModule/Private/EnvironmentQuery/Generators/EnvQueryGenerator_OnCircle.cpp b/Engine/Source/Runtime/AIModule/Private/EnvironmentQuery/Generators/EnvQueryGenerator_OnCircle.cpp
--- a/Engine/Source/Runtime/AIModule/Private/EnvironmentQuery/Generators/EnvQueryGenerator_OnCircle.cpp
+++ b/Engine/Source/Runtime/AIModule/Private/EnvironmentQuery/Generators/EnvQueryGenerator_OnCircle.cpp
@@ -135,7 +135,7 @@ FVector UEnvQueryGenerator_OnCircle::CalcDirection(FEnvQueryInstance& QueryInsta
 	if (bDefineArc)
 	{
 		// By default, use Querier rotation for arc direction.
-		Direction = Querier->GetActorRotation().Vector();
+		Direction = Querier->GetActorForwardVector();
 		if (ArcDirection.DirMode == EEnvDirection::TwoPoints)
 		{
 			TArray<FVector> Start;
diff --git a/Engine/Source/Runtime/AIModule/Private/Navigation/CrowdFollowingComponent.cpp b/Engine/Source/Runtime/AIModule/Private/Navigation/CrowdFollowingComponent.cpp
--- a/Engine/Source/Runtime/AIModule/Private/Navigation/CrowdFollowingComponent.cpp
+++ b/Engine/Source/Runtime/AIModule/Private/Navigation/CrowdFollowingComponent.cpp
@@ -735,7 +735,7 @@ FVector UCrowdFollowingComponent::GetMoveFocus(bool bAllowStrafe) const
 		// if we're not moving, falling, or don't have a crowd agent move direction, set our focus to ahead of the rotation of our owner to keep the same rotation,
 		// otherwise use the Crowd Agent Move Direction to move in the direction we're supposed to be going
 		const FVector ForwardDir = MovementComp->GetOwner() && ((Status != EPathFollowingStatus::Moving) || (CharacterMovement && (CharacterMovement->MovementMode == MOVE_Falling)) || CrowdAgentMoveDirection.IsNearlyZero()) ?
-			MovementComp->GetOwner()->GetActorRotation().Vector() :
+			MovementComp->GetOwner()->GetActorForwardVector() :
 			CrowdAgentMoveDirection;
 
 		return AgentLoc + ForwardDir * 100.0f;
diff --git a/Engine/Source/Runtime/Engine/Private/Components/CharacterMovementComponent.cpp b/Engine/Source/Runtime/Engine/Private/Components/CharacterMovementComponent.cpp
--- a/Engine/Source/Runtime/Engine/Private/Components/CharacterMovementComponent.cpp
+++ b/Engine/Source/Runtime/Engine/Private/Components/CharacterMovementComponent.cpp
@@ -1544,7 +1544,7 @@ void UCharacterMovementComponent::PerformMovement(float DeltaSeconds)
 				// Convert Local Space Root Motion to world space. Do it right before used by physics to make sure we use up to date transforms, as translation is relative to rotation.
 				RootMotionParams.Set( SkelMeshComp->ConvertLocalRootMotionToWorld(RootMotionParams.RootMotionTransform) );
 				UE_LOG(LogRootMotion, Log,  TEXT("PerformMovement WorldSpaceRootMotion Translation: %s, Rotation: %s, Actor Facing: %s"),
-					*RootMotionParams.RootMotionTransform.GetTranslation().ToCompactString(), *RootMotionParams.RootMotionTransform.GetRotation().Rotator().ToCompactString(), *CharacterOwner->GetActorRotation().Vector().ToCompactString());
+					*RootMotionParams.RootMotionTransform.GetTranslation().ToCompactString(), *RootMotionParams.RootMotionTransform.GetRotation().Rotator().ToCompactString(), *CharacterOwner->GetActorForwardVector().ToCompactString());
 			}
 
 			// Then turn root motion to velocity to be used by various physics modes.
@@ -2257,7 +2257,7 @@ void UCharacterMovementComponent::CalcVelocity(float DeltaTime, float Friction,
 		}
 		else 
 		{
-			Acceleration = MaxAccel * (Velocity.SizeSquared() < SMALL_NUMBER ? CharacterOwner->GetActorRotation().Vector() : GetSafeNormalPrecise(Velocity));
+			Acceleration = MaxAccel * (Velocity.SizeSquared() < SMALL_NUMBER ? CharacterOwner->GetActorForwardVector() : GetSafeNormalPrecise(Velocity));
 		}
 
 		AnalogInputModifier = 1.f;
diff --git a/Engine/Source/Runtime/Engine/Private/UnrealEngine.cpp b/Engine/Source/Runtime/Engine/Private/UnrealEngine.cpp
--- a/Engine/Source/Runtime/Engine/Private/UnrealEngine.cpp
+++ b/Engine/Source/Runtime/Engine/Private/UnrealEngine.cpp
@@ -2742,7 +2742,7 @@ bool UEngine::HandleMergeMeshCommand( const TCHAR* Cmd, FOutputDevice& Ar, UWorl
 			return true;
 		}
 
-		FVector SpawnLocation = PlayerPawn->GetActorLocation() + PlayerPawn->GetActorRotation().Vector()*50.f;
+		FVector SpawnLocation = PlayerPawn->GetActorLocation() + PlayerPawn->GetActorForwardVector()*50.f;
 
 		// set the new composite mesh in the existing skeletal mesh component
 		ASkeletalMeshActor* const SMA = PlayerPawn->GetWorld()->SpawnActor<ASkeletalMeshActor>( SpawnLocation, PlayerPawn->GetActorRotation()*-1);
diff --git a/Engine/Source/Runtime/GameplayAbilities/Private/AbilitySystemBlueprintLibrary.cpp b/Engine/Source/Runtime/GameplayAbilities/Private/AbilitySystemBlueprintLibrary.cpp
--- a/Engine/Source/Runtime/GameplayAbilities/Private/AbilitySystemBlueprintLibrary.cpp
+++ b/Engine/Source/Runtime/GameplayAbilities/Private/AbilitySystemBlueprintLibrary.cpp
@@ -452,7 +452,7 @@ bool UAbilitySystemBlueprintLibrary::GetGameplayCueEndLocationAndNormal(AActor*
 	else if(TargetActor)
 	{
 		Location = TargetActor->GetActorLocation();
-		Normal = TargetActor->GetActorRotation().Vector();
+		Normal = TargetActor->GetActorForwardVector();
 		return true;
 	}
 
