#ue4 - Perf: Convert some uses of GetActorRotation() to GetActorQuat() where possible. Use FQuatRotationMatrix rather than FRotationMatrix where possible.

- FQuatRotationMatrix(GetActorQuat()) is ~4x faster than FRotationMatrix(GetActorRotation()), due to latter form's FQuat->FRotator conversion and different matrix construction using Sin/Cos.

diff --git a/Engine/Source/Editor/GeometryMode/Private/GeometryEdMode.cpp b/Engine/Source/Editor/GeometryMode/Private/GeometryEdMode.cpp
--- a/Engine/Source/Editor/GeometryMode/Private/GeometryEdMode.cpp
+++ b/Engine/Source/Editor/GeometryMode/Private/GeometryEdMode.cpp
@@ -84,7 +84,7 @@ bool FEdModeGeometry::GetCustomDrawingCoordinateSystem( FMatrix& InMatrix, void*
 		FGeomObject* GeomObject = GeomBase->GetParentObject();
 		check(GeomObject != nullptr);
 		ABrush* Brush = GeomObject->GetActualBrush();
-		InMatrix = FRotationMatrix(GeomBase->GetNormal().Rotation()) * FRotationMatrix(Brush->GetActorRotation());
+		InMatrix = FRotationMatrix(GeomBase->GetNormal().Rotation()) * FQuatRotationMatrix(Brush->GetActorQuat());
 	}
 	else
 	{
@@ -103,7 +103,7 @@ bool FEdModeGeometry::GetCustomDrawingCoordinateSystem( FMatrix& InMatrix, void*
 				FGeomObject* GeomObject = GeomBase->GetParentObject();
 				check(GeomObject != nullptr);
 				ABrush* Brush = GeomObject->GetActualBrush();
-				InMatrix = FRotationMatrix( go->SelectionOrder[ go->SelectionOrder.Num()-1 ]->GetWidgetRotation() ) * FRotationMatrix(Brush->GetActorRotation());
+				InMatrix = FRotationMatrix( go->SelectionOrder[ go->SelectionOrder.Num()-1 ]->GetWidgetRotation() ) * FQuatRotationMatrix(Brush->GetActorQuat());
 				return 1;
 			}
 		}
diff --git a/Engine/Source/Editor/Kismet/Private/SCSEditorViewportClient.cpp b/Engine/Source/Editor/Kismet/Private/SCSEditorViewportClient.cpp
--- a/Engine/Source/Editor/Kismet/Private/SCSEditorViewportClient.cpp
+++ b/Engine/Source/Editor/Kismet/Private/SCSEditorViewportClient.cpp
@@ -638,7 +638,7 @@ FMatrix FSCSEditorViewportClient::GetWidgetCoordSystem() const
 					}					
 					else
 					{
-						Matrix = FRotationMatrix( SceneComp->GetComponentRotation() );
+						Matrix = FQuatRotationMatrix( SceneComp->GetComponentQuat() );
 					}
 				}
 			}
diff --git a/Engine/Source/Editor/LandscapeEditor/Private/LandscapeEdMode.cpp b/Engine/Source/Editor/LandscapeEditor/Private/LandscapeEdMode.cpp
--- a/Engine/Source/Editor/LandscapeEditor/Private/LandscapeEdMode.cpp
+++ b/Engine/Source/Editor/LandscapeEditor/Private/LandscapeEdMode.cpp
@@ -2121,7 +2121,7 @@ FVector FEdModeLandscape::GetWidgetLocation() const
 			// Apply Landscape transformation when it is available
 			ULandscapeInfo* Info = CurrentGizmoActor->TargetLandscapeInfo;
 			return CurrentGizmoActor->GetActorLocation()
-				+ FRotationMatrix(Info->GetLandscapeProxy()->GetActorRotation()).TransformPosition(FVector(0, 0, CurrentGizmoActor->GetLength()));
+				+ FQuatRotationMatrix(Info->GetLandscapeProxy()->GetActorQuat()).TransformPosition(FVector(0, 0, CurrentGizmoActor->GetLength()));
 		}
 		return CurrentGizmoActor->GetActorLocation();
 	}
diff --git a/Engine/Source/Editor/UnrealEd/Private/BSPOps.cpp b/Engine/Source/Editor/UnrealEd/Private/BSPOps.cpp
--- a/Engine/Source/Editor/UnrealEd/Private/BSPOps.cpp
+++ b/Engine/Source/Editor/UnrealEd/Private/BSPOps.cpp
@@ -1224,15 +1224,16 @@ void FBSPOps::RotateBrushVerts(ABrush* Brush, const FRotator& Rotation, bool bCl
 			FPoly* Poly = &(Brush->GetBrushComponent()->Brush->Polys->Element[poly]);
 
 			// Rotate the vertices.
+			const FRotationMatrix RotMatrix( Rotation );
 			for( int32 vertex = 0 ; vertex < Poly->Vertices.Num() ; vertex++ )
 			{
-				Poly->Vertices[vertex] = Brush->GetPrePivot() + FRotationMatrix( Rotation ).TransformVector( Poly->Vertices[vertex] - Brush->GetPrePivot() );
+				Poly->Vertices[vertex] = Brush->GetPrePivot() + RotMatrix.TransformVector( Poly->Vertices[vertex] - Brush->GetPrePivot() );
 			}
-			Poly->Base = Brush->GetPrePivot() + FRotationMatrix( Rotation ).TransformVector( Poly->Base - Brush->GetPrePivot() );
+			Poly->Base = Brush->GetPrePivot() + RotMatrix.TransformVector( Poly->Base - Brush->GetPrePivot() );
 
 			// Rotate the texture vectors.
-			Poly->TextureU = FRotationMatrix( Rotation ).TransformVector( Poly->TextureU );
-			Poly->TextureV = FRotationMatrix( Rotation ).TransformVector( Poly->TextureV );
+			Poly->TextureU = RotMatrix.TransformVector( Poly->TextureU );
+			Poly->TextureV = RotMatrix.TransformVector( Poly->TextureV );
 
 			// Recalc the normal for the poly.
 			Poly->Normal = FVector::ZeroVector;
diff --git a/Engine/Source/Editor/UnrealEd/Private/EditorModeManager.cpp b/Engine/Source/Editor/UnrealEd/Private/EditorModeManager.cpp
--- a/Engine/Source/Editor/UnrealEd/Private/EditorModeManager.cpp
+++ b/Engine/Source/Editor/UnrealEd/Private/EditorModeManager.cpp
@@ -404,7 +404,7 @@ FMatrix FEditorModeTools::GetCustomDrawingCoordinateSystem()
 				// Coordinate system needs to come from the last actor selected
 				if (Num > 0)
 				{
-					Matrix = FRotationMatrix(GetSelectedActors()->GetBottom<AActor>()->GetActorRotation());
+					Matrix = FQuatRotationMatrix(GetSelectedActors()->GetBottom<AActor>()->GetActorQuat());
 				}
 			}
 
diff --git a/Engine/Source/Runtime/AIModule/Private/AIController.cpp b/Engine/Source/Runtime/AIModule/Private/AIController.cpp
--- a/Engine/Source/Runtime/AIModule/Private/AIController.cpp
+++ b/Engine/Source/Runtime/AIModule/Private/AIController.cpp
@@ -654,7 +654,7 @@ bool AAIController::PreparePathfinding(const FAIMoveRequest& MoveRequest, FPathF
 			if (NavGoal)
 			{
 				const FVector Offset = NavGoal->GetMoveGoalOffset(this);
-				GoalLocation = FRotationTranslationMatrix(MoveRequest.GetGoalActor()->GetActorRotation(), NavGoal->GetNavAgentLocation()).TransformPosition(Offset);
+				GoalLocation = FQuatRotationTranslationMatrix(MoveRequest.GetGoalActor()->GetActorQuat(), NavGoal->GetNavAgentLocation()).TransformPosition(Offset);
 			}
 			else
 			{
diff --git a/Engine/Source/Runtime/AIModule/Private/Navigation/PathFollowingComponent.cpp b/Engine/Source/Runtime/AIModule/Private/Navigation/PathFollowingComponent.cpp
--- a/Engine/Source/Runtime/AIModule/Private/Navigation/PathFollowingComponent.cpp
+++ b/Engine/Source/Runtime/AIModule/Private/Navigation/PathFollowingComponent.cpp
@@ -791,7 +791,7 @@ void UPathFollowingComponent::UpdatePathSegment()
 			{
 				const FVector AgentLocation = DestinationAgent ? DestinationAgent->GetNavAgentLocation() : DestinationActor->GetActorLocation();
 				// note that the condition below requires GoalLocation to be in world space.
-				const FVector GoalLocation = FRotationTranslationMatrix(DestinationActor->GetActorRotation(), AgentLocation).TransformPosition(MoveOffset);
+				const FVector GoalLocation = FQuatRotationTranslationMatrix(DestinationActor->GetActorQuat(), AgentLocation).TransformPosition(MoveOffset);
 				FVector HitLocation;
 
 				if (MyNavData == nullptr //|| MyNavData->DoesNodeContainLocation(Path->GetPathPoints().Last().NodeRef, GoalLocation))
@@ -889,7 +889,7 @@ bool UPathFollowingComponent::HasReached(const AActor& TestGoal, float InAccepta
 	{
 		const FVector GoalMoveOffset = NavAgent->GetMoveGoalOffset(GetOwner());
 		NavAgent->GetMoveGoalReachTest(GetOwner(), GoalMoveOffset, GoalOffset, GoalRadius, GoalHalfHeight);
-		TestPoint = FRotationTranslationMatrix(TestGoal.GetActorRotation(), NavAgent->GetNavAgentLocation()).TransformPosition(GoalOffset);
+		TestPoint = FQuatRotationTranslationMatrix(TestGoal.GetActorQuat(), NavAgent->GetNavAgentLocation()).TransformPosition(GoalOffset);
 	}
 
 	const FVector CurrentLocation = MovementComp ? MovementComp->GetActorFeetLocation() : FVector::ZeroVector;
@@ -911,7 +911,7 @@ bool UPathFollowingComponent::HasReachedDestination(const FVector& CurrentLocati
 			FVector GoalOffset;
 			DestinationAgent->GetMoveGoalReachTest(GetOwner(), MoveOffset, GoalOffset, GoalRadius, GoalHalfHeight);
 
-			GoalLocation = FRotationTranslationMatrix(DestinationActor->GetActorRotation(), DestinationAgent->GetNavAgentLocation()).TransformPosition(GoalOffset);
+			GoalLocation = FQuatRotationTranslationMatrix(DestinationActor->GetActorQuat(), DestinationAgent->GetNavAgentLocation()).TransformPosition(GoalOffset);
 		}
 		else
 		{
@@ -1018,7 +1018,7 @@ void UPathFollowingComponent::DebugReachTest(float& CurrentDot, float& CurrentDi
 				FVector GoalOffset;
 				DestinationAgent->GetMoveGoalReachTest(GetOwner(), MoveOffset, GoalOffset, GoalRadius, GoalHalfHeight);
 
-				GoalLocation = FRotationTranslationMatrix(DestinationActor->GetActorRotation(), DestinationAgent->GetNavAgentLocation()).TransformPosition(GoalOffset);
+				GoalLocation = FQuatRotationTranslationMatrix(DestinationActor->GetActorQuat(), DestinationAgent->GetNavAgentLocation()).TransformPosition(GoalOffset);
 			}
 			else
 			{
diff --git a/Engine/Source/Runtime/Engine/Private/Animation/AnimNode_Trail.cpp b/Engine/Source/Runtime/Engine/Private/Animation/AnimNode_Trail.cpp
--- a/Engine/Source/Runtime/Engine/Private/Animation/AnimNode_Trail.cpp
+++ b/Engine/Source/Runtime/Engine/Private/Animation/AnimNode_Trail.cpp
@@ -104,7 +104,7 @@ void FAnimNode_Trail::EvaluateBoneTransforms(USkeletalMeshComponent* SkelComp, c
 
 		if (bActorSpaceFakeVel && SkelComp->GetOwner())
 		{
-			const FTransform BoneToWorld(SkelComp->GetOwner()->GetActorRotation(), SkelComp->GetOwner()->GetActorLocation());
+			const FTransform BoneToWorld(SkelComp->GetOwner()->GetActorQuat(), SkelComp->GetOwner()->GetActorLocation());
 			FakeMovement = BoneToWorld.TransformVector(FakeMovement);
 		}
 
diff --git a/Engine/Source/Runtime/Engine/Private/Character.cpp b/Engine/Source/Runtime/Engine/Private/Character.cpp
--- a/Engine/Source/Runtime/Engine/Private/Character.cpp
+++ b/Engine/Source/Runtime/Engine/Private/Character.cpp
@@ -940,7 +940,7 @@ void ACharacter::OnRep_ReplicatedBasedMovement()
 		if (ReplicatedBasedMovement.HasRelativeRotation())
 		{
 			// Relative location, relative rotation
-			FRotator NewRotation = (FRotationMatrix(ReplicatedBasedMovement.Rotation) * FQuatRotationTranslationMatrix(CharacterMovement->OldBaseQuat, FVector::ZeroVector)).Rotator();
+			FRotator NewRotation = (FRotationMatrix(ReplicatedBasedMovement.Rotation) * FQuatRotationMatrix(CharacterMovement->OldBaseQuat)).Rotator();
 			
 			// TODO: need a better way to not assume we only use Yaw.
 			NewRotation.Pitch = 0.f;
diff --git a/Engine/Source/Runtime/Engine/Private/CheatManager.cpp b/Engine/Source/Runtime/Engine/Private/CheatManager.cpp
--- a/Engine/Source/Runtime/Engine/Private/CheatManager.cpp
+++ b/Engine/Source/Runtime/Engine/Private/CheatManager.cpp
@@ -202,7 +202,7 @@ void UCheatManager::DamageTarget(float DamageAmount)
 	{
 		check(Hit.GetActor() != NULL);
 		FVector ActorForward, ActorSide, ActorUp;
-		FRotationMatrix(Hit.GetActor()->GetActorRotation()).GetScaledAxes(ActorForward, ActorSide, ActorUp);
+		FQuatRotationMatrix(Hit.GetActor()->GetActorQuat()).GetScaledAxes(ActorForward, ActorSide, ActorUp);
 
 		FPointDamageEvent DamageEvent(DamageAmount, Hit, -ActorForward, UDamageType::StaticClass());
 		Hit.GetActor()->TakeDamage(DamageAmount, DamageEvent, MyPC, MyPC->GetPawn());
diff --git a/Engine/Source/Runtime/Engine/Private/Components/CharacterMovementComponent.cpp b/Engine/Source/Runtime/Engine/Private/Components/CharacterMovementComponent.cpp
--- a/Engine/Source/Runtime/Engine/Private/Components/CharacterMovementComponent.cpp
+++ b/Engine/Source/Runtime/Engine/Private/Components/CharacterMovementComponent.cpp
@@ -1680,7 +1680,7 @@ void UCharacterMovementComponent::SaveBaseLocation()
 		else
 		{
 			// Relative rotation
-			const FRotator RelativeRotation = (FRotationMatrix(CharacterOwner->GetActorRotation()) * FQuatRotationTranslationMatrix(OldBaseQuat, FVector::ZeroVector).GetTransposed()).Rotator();
+			const FRotator RelativeRotation = (FQuatRotationMatrix(CharacterOwner->GetActorQuat()) * FQuatRotationMatrix(OldBaseQuat).GetTransposed()).Rotator();
 			CharacterOwner->SaveRelativeBasedMovement(RelativeLocation, RelativeRotation, true);	
 		}
 	}
diff --git a/Engine/Source/Runtime/Engine/Private/Particles/ParticleModules_Velocity.cpp b/Engine/Source/Runtime/Engine/Private/Particles/ParticleModules_Velocity.cpp
--- a/Engine/Source/Runtime/Engine/Private/Particles/ParticleModules_Velocity.cpp
+++ b/Engine/Source/Runtime/Engine/Private/Particles/ParticleModules_Velocity.cpp
@@ -538,7 +538,7 @@ void UParticleModuleVelocityCone::Render3DPreview(FParticleEmitterInstance* Owne
 	UParticleLODLevel* LODLevel	= Owner->SpriteTemplate->GetCurrentLODLevel(Owner);
 	check(LODLevel);
 	FVector OwnerScale(1.0f);
-	FRotationMatrix OwnerRotation(FRotator::ZeroRotator);
+	FMatrix OwnerRotation(FMatrix::Identity);
 	FVector LocalToWorldOrigin(0.0f);
 	FMatrix LocalToWorld(FMatrix::Identity);
 	if (Owner && Owner->Component)
@@ -551,7 +551,7 @@ void UParticleModuleVelocityCone::Render3DPreview(FParticleEmitterInstance* Owne
 				OwnerScale = Owner->Component->ComponentToWorld.GetScale3D();
 			}
 
-			OwnerRotation = FRotationMatrix(Actor->GetActorRotation());
+			OwnerRotation = FQuatRotationMatrix(Actor->GetActorQuat());
 		}
 	  LocalToWorldOrigin = Owner->Component->ComponentToWorld.GetLocation();
 	  LocalToWorld = Owner->Component->ComponentToWorld.ToMatrixWithScale().RemoveTranslation();
diff --git a/Engine/Source/Runtime/Engine/Private/PlayerCameraManager.cpp b/Engine/Source/Runtime/Engine/Private/PlayerCameraManager.cpp
--- a/Engine/Source/Runtime/Engine/Private/PlayerCameraManager.cpp
+++ b/Engine/Source/Runtime/Engine/Private/PlayerCameraManager.cpp
@@ -532,7 +532,7 @@ void APlayerCameraManager::UpdateViewTarget(FTViewTarget& OutVT, float DeltaTime
 //			APawn* TPawn = Cast<APawn>(OutVT.Target);
 // 			if ((TPawn != NULL) && (TPawn->Mesh != NULL))
 // 			{
-// 				Loc += FRotationMatrix(OutVT.Target->GetActorRotation()).TransformVector(TPawn->Mesh->RelativeLocation - GetDefault<APawn>(TPawn->GetClass())->Mesh->RelativeLocation);
+// 				Loc += FQuatRotationMatrix(OutVT.Target->GetActorQuat()).TransformVector(TPawn->Mesh->RelativeLocation - GetDefault<APawn>(TPawn->GetClass())->Mesh->RelativeLocation);
 // 			}
 
 			//OutVT.Target.GetActorEyesViewPoint(Loc, Rot);
