fixed crash on using EQS asset with broken generator data
#ue4

diff --git a/Engine/Source/Runtime/AIModule/Private/EnvironmentQuery/EnvQueryManager.cpp b/Engine/Source/Runtime/AIModule/Private/EnvironmentQuery/EnvQueryManager.cpp
--- a/Engine/Source/Runtime/AIModule/Private/EnvironmentQuery/EnvQueryManager.cpp
+++ b/Engine/Source/Runtime/AIModule/Private/EnvironmentQuery/EnvQueryManager.cpp
@@ -526,8 +526,11 @@ TSharedPtr<FEnvQueryInstance> UEnvQueryManager::CreateQueryInstance(const UEnvQu
 
 			CreateOptionInstance(LocalOption, SortedTests, *InstanceTemplate);
 		}
+	}
 
-		check(InstanceTemplate->Options.Num() != 0);
+	if (InstanceTemplate->Options.Num() == 0)
+	{
+		return nullptr;
 	}
 
 	// create new instance
