added support for physx based projection in EQS generators
#ue4

diff --git a/Engine/Source/Runtime/AIModule/Private/EnvironmentQuery/EnvQueryTraceHelpers.cpp b/Engine/Source/Runtime/AIModule/Private/EnvironmentQuery/EnvQueryTraceHelpers.cpp
new file mode 100644
--- /dev/null
+++ b/Engine/Source/Runtime/AIModule/Private/EnvironmentQuery/EnvQueryTraceHelpers.cpp
@@ -0,0 +1,140 @@
+// Copyright 1998-2015 Epic Games, Inc. All Rights Reserved.
+
+#include "AIModulePrivate.h"
+#include "EnvQueryTraceHelpers.h"
+
+void FEQSHelpers::RunNavRaycasts(ANavigationData* NavData, const FEnvTraceData& TraceData, const FVector& SourcePt, TArray<FNavLocation>& Points, ETraceMode TraceMode)
+{
+	TSharedPtr<const FNavigationQueryFilter> NavigationFilter = UNavigationQueryFilter::GetQueryFilter(NavData, TraceData.NavigationFilter);
+
+	TArray<FNavigationRaycastWork> RaycastWorkload;
+	RaycastWorkload.Reserve(Points.Num());
+
+	for (const auto& ItemLocation : Points)
+	{
+		RaycastWorkload.Add(FNavigationRaycastWork(SourcePt, ItemLocation.Location));
+	}
+
+	NavData->BatchRaycast(RaycastWorkload, NavigationFilter);
+
+	for (int32 Idx = 0; Idx < Points.Num(); Idx++)
+	{
+		Points[Idx] = RaycastWorkload[Idx].HitLocation;
+	}
+
+	if (TraceMode == ETraceMode::Discard)
+	{
+		for (int32 Idx = Points.Num() - 1; Idx >= 0; Idx++)
+		{
+			if (!RaycastWorkload[Idx].bDidHit)
+			{
+				Points.RemoveAt(Idx, 1, false);
+			}
+		}
+	}
+}
+
+void FEQSHelpers::RunNavProjection(ANavigationData* NavData, const FEnvTraceData& TraceData, TArray<FNavLocation>& Points, ETraceMode TraceMode)
+{
+	TSharedPtr<const FNavigationQueryFilter> NavigationFilter = UNavigationQueryFilter::GetQueryFilter(NavData, TraceData.NavigationFilter);
+	TArray<FNavigationProjectionWork> Workload;
+	Workload.Reserve(Points.Num());
+
+	if (TraceData.ProjectDown == TraceData.ProjectUp)
+	{
+		for (const auto& Point : Points)
+		{
+			Workload.Add(FNavigationProjectionWork(Point.Location));
+		}
+	}
+	else
+	{
+		const FVector VerticalOffset = FVector(0, 0, (TraceData.ProjectUp - TraceData.ProjectDown) / 2);
+		for (const auto& Point : Points)
+		{
+			Workload.Add(FNavigationProjectionWork(Point.Location + VerticalOffset));
+		}
+	}
+
+	const FVector ProjectionExtent(TraceData.ExtentX, TraceData.ExtentX, (TraceData.ProjectDown + TraceData.ProjectUp) / 2);
+	NavData->BatchProjectPoints(Workload, ProjectionExtent, NavigationFilter);
+
+	for (int32 Idx = Workload.Num() - 1; Idx >= 0; Idx--)
+	{
+		if (Workload[Idx].bResult)
+		{
+			Points[Idx] = Workload[Idx].OutLocation;
+			Points[Idx].Location.Z += TraceData.PostProjectionVerticalOffset;
+		}
+		else if (TraceMode == ETraceMode::Discard)
+		{
+			Points.RemoveAt(Idx, 1, false);
+		}
+	}
+}
+
+void FEQSHelpers::RunPhysRaycasts(UWorld* World, const FEnvTraceData& TraceData, const FVector& SourcePt, TArray<FNavLocation>& Points, ETraceMode TraceMode)
+{
+	ECollisionChannel TraceCollisionChannel = UEngineTypes::ConvertToCollisionChannel(TraceData.TraceChannel);
+	FVector TraceExtent(TraceData.ExtentX, TraceData.ExtentY, TraceData.ExtentZ);
+
+	FCollisionQueryParams TraceParams(TEXT("EnvQueryTrace"), TraceData.bTraceComplex);
+	TraceParams.bTraceAsyncScene = true;
+
+	FBatchTrace BatchOb(World, TraceCollisionChannel, TraceParams, TraceExtent, TraceMode);
+
+	switch (TraceData.TraceShape)
+	{
+	case EEnvTraceShape::Line:
+		BatchOb.DoSingleSourceMultiDestinations<EEnvTraceShape::Line>(SourcePt, Points);
+		break;
+
+	case EEnvTraceShape::Sphere:
+		BatchOb.DoSingleSourceMultiDestinations<EEnvTraceShape::Sphere>(SourcePt, Points);
+		break;
+
+	case EEnvTraceShape::Capsule:
+		BatchOb.DoSingleSourceMultiDestinations<EEnvTraceShape::Capsule>(SourcePt, Points);
+		break;
+
+	case EEnvTraceShape::Box:
+		BatchOb.DoSingleSourceMultiDestinations<EEnvTraceShape::Box>(SourcePt, Points);
+		break;
+
+	default:
+		break;
+	}
+}
+
+void FEQSHelpers::RunPhysProjection(UWorld* World, const FEnvTraceData& TraceData, TArray<FNavLocation>& Points, ETraceMode TraceMode)
+{
+	ECollisionChannel TraceCollisionChannel = UEngineTypes::ConvertToCollisionChannel(TraceData.TraceChannel);
+	FVector TraceExtent(TraceData.ExtentX, TraceData.ExtentY, TraceData.ExtentZ);
+
+	FCollisionQueryParams TraceParams(TEXT("EnvQueryTrace"), TraceData.bTraceComplex);
+	TraceParams.bTraceAsyncScene = true;
+
+	FBatchTrace BatchOb(World, TraceCollisionChannel, TraceParams, TraceExtent, TraceMode);
+
+	switch (TraceData.TraceShape)
+	{
+	case EEnvTraceShape::Line:
+		BatchOb.DoProject<EEnvTraceShape::Line>(Points, TraceData.ProjectUp, -TraceData.ProjectDown, TraceData.PostProjectionVerticalOffset);
+		break;
+
+	case EEnvTraceShape::Sphere:
+		BatchOb.DoProject<EEnvTraceShape::Sphere>(Points, TraceData.ProjectUp, -TraceData.ProjectDown, TraceData.PostProjectionVerticalOffset);
+		break;
+
+	case EEnvTraceShape::Capsule:
+		BatchOb.DoProject<EEnvTraceShape::Capsule>(Points, TraceData.ProjectUp, -TraceData.ProjectDown, TraceData.PostProjectionVerticalOffset);
+		break;
+
+	case EEnvTraceShape::Box:
+		BatchOb.DoProject<EEnvTraceShape::Box>(Points, TraceData.ProjectUp, -TraceData.ProjectDown, TraceData.PostProjectionVerticalOffset);
+		break;
+
+	default:
+		break;
+	}
+}
diff --git a/Engine/Source/Runtime/AIModule/Private/EnvironmentQuery/EnvQueryTraceHelpers.h b/Engine/Source/Runtime/AIModule/Private/EnvironmentQuery/EnvQueryTraceHelpers.h
new file mode 100644
--- /dev/null
+++ b/Engine/Source/Runtime/AIModule/Private/EnvironmentQuery/EnvQueryTraceHelpers.h
@@ -0,0 +1,221 @@
+// Copyright 1998-2015 Epic Games, Inc. All Rights Reserved.
+
+#pragma once
+
+namespace FEQSHelpers
+{
+	enum class ETraceMode : uint8
+	{
+		Keep,
+		Discard,
+	};
+
+	struct FBatchTrace
+	{
+		UWorld* World;
+		const FVector Extent;
+		const FCollisionQueryParams Params;
+		enum ECollisionChannel Channel;
+		ETraceMode TraceMode;
+
+		FBatchTrace(UWorld* InWorld, enum ECollisionChannel InChannel, const FCollisionQueryParams& InParams,
+			const FVector& InExtent, ETraceMode InTraceMode)
+			: World(InWorld), Extent(InExtent), Params(InParams), Channel(InChannel), TraceMode(InTraceMode)
+		{
+
+		}
+
+		FORCEINLINE_DEBUGGABLE bool RunLineTrace(const FVector& StartPos, const FVector& EndPos, FVector& HitPos)
+		{
+			FHitResult OutHit;
+			const bool bHit = World->LineTraceSingleByChannel(OutHit, StartPos, EndPos, Channel, Params);
+			HitPos = OutHit.ImpactPoint;
+			return bHit;
+		}
+
+		FORCEINLINE_DEBUGGABLE bool RunSphereTrace(const FVector& StartPos, const FVector& EndPos, FVector& HitPos)
+		{
+			FHitResult OutHit;
+			const bool bHit = World->SweepSingleByChannel(OutHit, StartPos, EndPos, FQuat::Identity, Channel, FCollisionShape::MakeSphere(Extent.X), Params);
+			HitPos = OutHit.ImpactPoint;
+			return bHit;
+		}
+
+		FORCEINLINE_DEBUGGABLE bool RunCapsuleTrace(const FVector& StartPos, const FVector& EndPos, FVector& HitPos)
+		{
+			FHitResult OutHit;
+			const bool bHit = World->SweepSingleByChannel(OutHit, StartPos, EndPos, FQuat::Identity, Channel, FCollisionShape::MakeCapsule(Extent.X, Extent.Z), Params);
+			HitPos = OutHit.ImpactPoint;
+			return bHit;
+		}
+
+		FORCEINLINE_DEBUGGABLE bool RunBoxTrace(const FVector& StartPos, const FVector& EndPos, FVector& HitPos)
+		{
+			FHitResult OutHit;
+			const bool bHit = World->SweepSingleByChannel(OutHit, StartPos, EndPos, FQuat((EndPos - StartPos).Rotation()), Channel, FCollisionShape::MakeBox(Extent), Params);
+			HitPos = OutHit.ImpactPoint;
+			return bHit;
+		}
+
+		template<EEnvTraceShape::Type TraceType>
+		void DoSingleSourceMultiDestinations(const FVector& Source, TArray<FNavLocation>& Points)
+		{
+			UE_LOG(LogEQS, Error, TEXT("FBatchTrace::DoSingleSourceMultiDestinations called with unhandled trace type: %d"), int32(TraceType));
+		}
+
+		template<EEnvTraceShape::Type TraceType>
+		void DoProject(TArray<FNavLocation>& Points, float StartOffsetZ, float EndOffsetZ, float HitOffsetZ)
+		{
+			UE_LOG(LogEQS, Error, TEXT("FBatchTrace::DoSingleSourceMultiDestinations called with unhandled trace type: %d"), int32(TraceType));
+		}
+	};
+
+	template<>
+	void FBatchTrace::DoSingleSourceMultiDestinations<EEnvTraceShape::Line>(const FVector& Source, TArray<FNavLocation>& Points)
+	{
+		FVector HitPos(FVector::ZeroVector);
+		for (int32 Idx = Points.Num() - 1; Idx >= 0; Idx--)
+		{
+			const bool bHit = RunLineTrace(Source, Points[Idx].Location, HitPos);
+			if (bHit)
+			{
+				Points[Idx] = HitPos;
+			}
+			else if (TraceMode == ETraceMode::Discard)
+			{
+				Points.RemoveAt(Idx, 1, false);
+			}
+		}
+	}
+
+	template<>
+	void FBatchTrace::DoSingleSourceMultiDestinations<EEnvTraceShape::Box>(const FVector& Source, TArray<FNavLocation>& Points)
+	{
+		FVector HitPos(FVector::ZeroVector);
+		for (int32 Idx = Points.Num() - 1; Idx >= 0; Idx--)
+		{
+			const bool bHit = RunBoxTrace(Source, Points[Idx].Location, HitPos);
+			if (bHit)
+			{
+				Points[Idx] = HitPos;
+			}
+			else if (TraceMode == ETraceMode::Discard)
+			{
+				Points.RemoveAt(Idx, 1, false);
+			}
+		}
+	}
+
+	template<>
+	void FBatchTrace::DoSingleSourceMultiDestinations<EEnvTraceShape::Sphere>(const FVector& Source, TArray<FNavLocation>& Points)
+	{
+		FVector HitPos(FVector::ZeroVector);
+		for (int32 Idx = Points.Num() - 1; Idx >= 0; Idx--)
+		{
+			const bool bHit = RunSphereTrace(Source, Points[Idx].Location, HitPos);
+			if (bHit)
+			{
+				Points[Idx] = HitPos;
+			}
+			else if (TraceMode == ETraceMode::Discard)
+			{
+				Points.RemoveAt(Idx, 1, false);
+			}
+		}
+	}
+
+	template<>
+	void FBatchTrace::DoSingleSourceMultiDestinations<EEnvTraceShape::Capsule>(const FVector& Source, TArray<FNavLocation>& Points)
+	{
+		FVector HitPos(FVector::ZeroVector);
+		for (int32 Idx = Points.Num() - 1; Idx >= 0; Idx--)
+		{
+			const bool bHit = RunCapsuleTrace(Source, Points[Idx].Location, HitPos);
+			if (bHit)
+			{
+				Points[Idx] = HitPos;
+			}
+			else if (TraceMode == ETraceMode::Discard)
+			{
+				Points.RemoveAt(Idx, 1, false);
+			}
+		}
+	}
+
+	template<>
+	void FBatchTrace::DoProject<EEnvTraceShape::Line>(TArray<FNavLocation>& Points, float StartOffsetZ, float EndOffsetZ, float HitOffsetZ)
+	{
+		FVector HitPos(FVector::ZeroVector);
+		for (int32 Idx = Points.Num() - 1; Idx >= 0; Idx--)
+		{
+			const bool bHit = RunLineTrace(Points[Idx].Location + FVector(0, 0, StartOffsetZ), Points[Idx].Location + FVector(0, 0, EndOffsetZ), HitPos);
+			if (bHit)
+			{
+				Points[Idx] = HitPos + FVector(0, 0, HitOffsetZ);
+			}
+			else if (TraceMode == ETraceMode::Discard)
+			{
+				Points.RemoveAt(Idx, 1, false);
+			}
+		}
+	}
+
+	template<>
+	void FBatchTrace::DoProject<EEnvTraceShape::Box>(TArray<FNavLocation>& Points, float StartOffsetZ, float EndOffsetZ, float HitOffsetZ)
+	{
+		FVector HitPos(FVector::ZeroVector);
+		for (int32 Idx = Points.Num() - 1; Idx >= 0; Idx--)
+		{
+			const bool bHit = RunBoxTrace(Points[Idx].Location + FVector(0, 0, StartOffsetZ), Points[Idx].Location + FVector(0, 0, EndOffsetZ), HitPos);
+			if (bHit)
+			{
+				Points[Idx] = HitPos + FVector(0, 0, HitOffsetZ);
+			}
+			else if (TraceMode == ETraceMode::Discard)
+			{
+				Points.RemoveAt(Idx, 1, false);
+			}
+		}
+	}
+
+	template<>
+	void FBatchTrace::DoProject<EEnvTraceShape::Sphere>(TArray<FNavLocation>& Points, float StartOffsetZ, float EndOffsetZ, float HitOffsetZ)
+	{
+		FVector HitPos(FVector::ZeroVector);
+		for (int32 Idx = Points.Num() - 1; Idx >= 0; Idx--)
+		{
+			const bool bHit = RunSphereTrace(Points[Idx].Location + FVector(0, 0, StartOffsetZ), Points[Idx].Location + FVector(0, 0, EndOffsetZ), HitPos);
+			if (bHit)
+			{
+				Points[Idx] = HitPos + FVector(0, 0, HitOffsetZ);
+			}
+			else if (TraceMode == ETraceMode::Discard)
+			{
+				Points.RemoveAt(Idx, 1, false);
+			}
+		}
+	}
+
+	template<>
+	void FBatchTrace::DoProject<EEnvTraceShape::Capsule>(TArray<FNavLocation>& Points, float StartOffsetZ, float EndOffsetZ, float HitOffsetZ)
+	{
+		FVector HitPos(FVector::ZeroVector);
+		for (int32 Idx = Points.Num() - 1; Idx >= 0; Idx--)
+		{
+			const bool bHit = RunCapsuleTrace(Points[Idx].Location + FVector(0, 0, StartOffsetZ), Points[Idx].Location + FVector(0, 0, EndOffsetZ), HitPos);
+			if (bHit)
+			{
+				Points[Idx] = HitPos + FVector(0, 0, HitOffsetZ);
+			}
+			else if (TraceMode == ETraceMode::Discard)
+			{
+				Points.RemoveAt(Idx, 1, false);
+			}
+		}
+	}
+
+	void RunNavRaycasts(ANavigationData* NavData, const FEnvTraceData& TraceData, const FVector& SourcePt, TArray<FNavLocation>& Points, ETraceMode TraceMode = ETraceMode::Keep);
+	void RunNavProjection(ANavigationData* NavData, const FEnvTraceData& TraceData, TArray<FNavLocation>& Points, ETraceMode TraceMode = ETraceMode::Discard);
+	void RunPhysRaycasts(UWorld* World, const FEnvTraceData& TraceData, const FVector& SourcePt, TArray<FNavLocation>& Points, ETraceMode TraceMode = ETraceMode::Keep);
+	void RunPhysProjection(UWorld* World, const FEnvTraceData& TraceData, TArray<FNavLocation>& Points, ETraceMode TraceMode = ETraceMode::Discard);
+}
diff --git a/Engine/Source/Runtime/AIModule/Private/EnvironmentQuery/Generators/EnvQueryGenerator_OnCircle.cpp b/Engine/Source/Runtime/AIModule/Private/EnvironmentQuery/Generators/EnvQueryGenerator_OnCircle.cpp
--- a/Engine/Source/Runtime/AIModule/Private/EnvironmentQuery/Generators/EnvQueryGenerator_OnCircle.cpp
+++ b/Engine/Source/Runtime/AIModule/Private/EnvironmentQuery/Generators/EnvQueryGenerator_OnCircle.cpp
@@ -4,93 +4,10 @@
 #include "EnvironmentQuery/Contexts/EnvQueryContext_Querier.h"
 #include "EnvironmentQuery/Items/EnvQueryItemType_Point.h"
 #include "EnvironmentQuery/Generators/EnvQueryGenerator_OnCircle.h"
+#include "EnvironmentQuery/EnvQueryTraceHelpers.h"
 
 #define LOCTEXT_NAMESPACE "EnvQueryGenerator"
 
-struct FBatchTracingHelper
-{
-	UWorld* World;
-	enum ECollisionChannel Channel;
-	const FCollisionQueryParams Params;
-	const FVector Extent;
-
-	FBatchTracingHelper(UWorld* InWorld, enum ECollisionChannel InChannel, const FCollisionQueryParams& InParams, const FVector& InExtent)
-		: World(InWorld), Channel(InChannel), Params(InParams), Extent(InExtent)
-	{
-
-	}
-
-	FORCEINLINE_DEBUGGABLE FVector RunLineTrace(const FVector& StartPos, const FVector& EndPos)
-	{
-		FHitResult OutHit;
-		const bool bHit = World->LineTraceSingleByChannel(OutHit, StartPos, EndPos, Channel, Params);
-		return bHit ? OutHit.ImpactPoint : EndPos;
-	}
-
-	FORCEINLINE_DEBUGGABLE FVector RunSphereTrace(const FVector& StartPos, const FVector& EndPos)
-	{
-		FHitResult OutHit;
-		const bool bHit = World->SweepSingleByChannel(OutHit, StartPos, EndPos, FQuat::Identity, Channel, FCollisionShape::MakeSphere(Extent.X), Params);
-		return bHit ? OutHit.ImpactPoint : EndPos;
-	}
-
-	FORCEINLINE_DEBUGGABLE FVector RunCapsuleTrace(const FVector& StartPos, const FVector& EndPos)
-	{
-		FHitResult OutHit;
-		const bool bHit = World->SweepSingleByChannel(OutHit, StartPos, EndPos, FQuat::Identity, Channel, FCollisionShape::MakeCapsule(Extent.X, Extent.Z), Params);
-		return bHit ? OutHit.ImpactPoint : EndPos;
-	}
-
-	FORCEINLINE_DEBUGGABLE FVector RunBoxTrace(const FVector& StartPos, const FVector& EndPos)
-	{
-		FHitResult OutHit;
-		const bool bHit = World->SweepSingleByChannel(OutHit, StartPos, EndPos, FQuat((EndPos - StartPos).Rotation()), Channel, FCollisionShape::MakeBox(Extent), Params);
-		return bHit ? OutHit.ImpactPoint : EndPos;
-	}
-
-	template<EEnvTraceShape::Type TraceType>
-	void DoSingleSourceMultiDestinations(const FVector& Source, TArray<FNavLocation>& DestinationAndClampedResult)
-	{
-		UE_LOG(LogEQS, Error, TEXT("FBatchTracingHelper::DoSingleSourceMultiDestinations called with unhandled trace type: %d"), int32(TraceType));
-	}
-};
-
-template<>
-void FBatchTracingHelper::DoSingleSourceMultiDestinations<EEnvTraceShape::Line>(const FVector& Source, TArray<FNavLocation>& DestinationAndClampedResult)
-{
-	for (auto& OutLocation : DestinationAndClampedResult)
-	{
-		OutLocation = RunLineTrace(Source, OutLocation.Location);
-	}
-}
-
-template<>
-void FBatchTracingHelper::DoSingleSourceMultiDestinations<EEnvTraceShape::Box>(const FVector& Source, TArray<FNavLocation>& DestinationAndClampedResult)
-{
-	for (auto& OutLocation : DestinationAndClampedResult)
-	{
-		OutLocation = RunBoxTrace(Source, OutLocation.Location);
-	}
-}
-
-template<>
-void FBatchTracingHelper::DoSingleSourceMultiDestinations<EEnvTraceShape::Sphere>(const FVector& Source, TArray<FNavLocation>& DestinationAndClampedResult)
-{
-	for (auto& OutLocation : DestinationAndClampedResult)
-	{
-		OutLocation = RunSphereTrace(Source, OutLocation.Location);
-	}
-}
-
-template<>
-void FBatchTracingHelper::DoSingleSourceMultiDestinations<EEnvTraceShape::Capsule>(const FVector& Source, TArray<FNavLocation>& DestinationAndClampedResult)
-{
-	for (auto& OutLocation : DestinationAndClampedResult)
-	{
-		OutLocation = RunCapsuleTrace(Source, OutLocation.Location);
-	}
-}
-
 //----------------------------------------------------------------------//
 // UEnvQueryGenerator_OnCircle
 //----------------------------------------------------------------------//
@@ -244,59 +161,13 @@ void UEnvQueryGenerator_OnCircle::GenerateItemsForCircle(uint8* ContextRawData,
 #endif
 			if (NavData)
 			{
-				TSharedPtr<const FNavigationQueryFilter> NavigationFilter = UNavigationQueryFilter::GetQueryFilter(NavData, TraceData.NavigationFilter);
-
-				TArray<FNavigationRaycastWork> RaycastWorkload;
-				RaycastWorkload.Reserve(ItemCandidates.Num());
-
-				for (const auto& ItemLocation : ItemCandidates)
-				{
-					RaycastWorkload.Add(FNavigationRaycastWork(CenterLocation, ItemLocation.Location));
-				}
-
-				NavData->BatchRaycast(RaycastWorkload, NavigationFilter);
-
-				for (int32 ItemIndex = 0; ItemIndex < ItemCandidates.Num(); ++ItemIndex)
-				{
-					ItemCandidates[ItemIndex] = RaycastWorkload[ItemIndex].HitLocation;
-				}
+				FEQSHelpers::RunNavRaycasts(NavData, TraceData, CenterLocation, ItemCandidates);
 			}
 		}
 			break;
 
 		case EEnvQueryTrace::Geometry:
-		{
-			ECollisionChannel TraceCollisionChannel = UEngineTypes::ConvertToCollisionChannel(TraceData.TraceChannel);
-			FVector TraceExtent(TraceData.ExtentX, TraceData.ExtentY, TraceData.ExtentZ);
-
-			FCollisionQueryParams TraceParams(TEXT("EnvQueryTrace"), TraceData.bTraceComplex);
-			TraceParams.bTraceAsyncScene = true;
-
-			FBatchTracingHelper TracingHelper(OutQueryInstance.World, TraceCollisionChannel, TraceParams, TraceExtent);
-
-			switch (TraceData.TraceShape)
-			{
-				case EEnvTraceShape::Line:		
-					TracingHelper.DoSingleSourceMultiDestinations<EEnvTraceShape::Line>(CenterLocation, ItemCandidates);
-					break;
-
-				case EEnvTraceShape::Sphere:	
-					TracingHelper.DoSingleSourceMultiDestinations<EEnvTraceShape::Sphere>(CenterLocation, ItemCandidates);
-					break;
-
-				case EEnvTraceShape::Capsule:
-					TracingHelper.DoSingleSourceMultiDestinations<EEnvTraceShape::Capsule>(CenterLocation, ItemCandidates);
-					break;
-
-				case EEnvTraceShape::Box:
-					TracingHelper.DoSingleSourceMultiDestinations<EEnvTraceShape::Box>(CenterLocation, ItemCandidates);
-					break;
-
-				default:
-					UE_VLOG(Cast<AActor>(OutQueryInstance.Owner.Get()), LogEQS, Warning, TEXT("UEnvQueryGenerator_OnCircle::CalcDirection failed to calc direction in %s. Using querier facing."), *OutQueryInstance.QueryName);
-					break;
-			}
-		}
+			FEQSHelpers::RunPhysRaycasts(OutQueryInstance.World, TraceData, CenterLocation, ItemCandidates);
 			break;
 
 		case EEnvQueryTrace::None:
diff --git a/Engine/Source/Runtime/AIModule/Private/EnvironmentQuery/Generators/EnvQueryGenerator_ProjectedPoints.cpp b/Engine/Source/Runtime/AIModule/Private/EnvironmentQuery/Generators/EnvQueryGenerator_ProjectedPoints.cpp
--- a/Engine/Source/Runtime/AIModule/Private/EnvironmentQuery/Generators/EnvQueryGenerator_ProjectedPoints.cpp
+++ b/Engine/Source/Runtime/AIModule/Private/EnvironmentQuery/Generators/EnvQueryGenerator_ProjectedPoints.cpp
@@ -4,10 +4,10 @@
 #include "AI/Navigation/RecastNavMesh.h"
 #include "EnvironmentQuery/Items/EnvQueryItemType_Point.h"
 #include "EnvironmentQuery/Generators/EnvQueryGenerator_ProjectedPoints.h"
+#include "EnvironmentQuery/EnvQueryTraceHelpers.h"
 
 UEnvQueryGenerator_ProjectedPoints::UEnvQueryGenerator_ProjectedPoints(const FObjectInitializer& ObjectInitializer) : Super(ObjectInitializer)
 {
-	ProjectionData.SetNavmeshOnly();
 	ProjectionData.bCanProjectDown = true;
 	ProjectionData.bCanDisableTrace = true;
 	ProjectionData.ExtentX = 0.0f;
@@ -17,9 +17,8 @@ UEnvQueryGenerator_ProjectedPoints::UEnvQueryGenerator_ProjectedPoints(const FOb
 
 void UEnvQueryGenerator_ProjectedPoints::ProjectAndFilterNavPoints(TArray<FNavLocation>& Points, FEnvQueryInstance& QueryInstance) const
 {
-	const bool bProjectToNavigation = (ProjectionData.TraceMode == EEnvQueryTrace::Navigation);
 	ANavigationData* NavData = nullptr;
-	if (bProjectToNavigation)
+	if (ProjectionData.TraceMode != EEnvQueryTrace::None)
 	{
 #if WITH_RECAST
 		NavData = (ANavigationData*)FEQSHelpers::FindNavMeshForQuery(QueryInstance);
@@ -28,41 +27,13 @@ void UEnvQueryGenerator_ProjectedPoints::ProjectAndFilterNavPoints(TArray<FNavLo
 
 	if (NavData)
 	{
-		TSharedPtr<const FNavigationQueryFilter> NavigationFilter = UNavigationQueryFilter::GetQueryFilter(NavData, ProjectionData.NavigationFilter);	
-		TArray<FNavigationProjectionWork> Workload;
-		Workload.Reserve(Points.Num());
-
-		if (ProjectionData.ProjectDown == ProjectionData.ProjectUp)
-		{
-			for (const auto& Point : Points)
-			{
-				Workload.Add(FNavigationProjectionWork(Point.Location));
-			}
-		}
-		else
-		{
-			const FVector VerticalOffset = FVector(0, 0, (ProjectionData.ProjectUp - ProjectionData.ProjectDown) / 2);
-			for (const auto& Point : Points)
-			{
-				Workload.Add(FNavigationProjectionWork(Point.Location + VerticalOffset));
-			}
-		}
-
-		const FVector ProjectionExtent(ProjectionData.ExtentX, ProjectionData.ExtentX, (ProjectionData.ProjectDown + ProjectionData.ProjectUp) / 2);
-		NavData->BatchProjectPoints(Workload, ProjectionExtent, NavigationFilter);
+		FEQSHelpers::ETraceMode Mode = (ProjectionData.TraceMode == EEnvQueryTrace::Navigation) ? FEQSHelpers::ETraceMode::Discard : FEQSHelpers::ETraceMode::Keep;
+		FEQSHelpers::RunNavProjection(NavData, ProjectionData, Points, Mode);
+	}
 
-		for (int32 Idx = Workload.Num() - 1; Idx >= 0; Idx--)
-		{
-			if (Workload[Idx].bResult)
-			{
-				Points[Idx] = Workload[Idx].OutLocation;
-				Points[Idx].Location.Z += ProjectionData.PostProjectionVerticalOffset;
-			}
-			else
-			{
-				Points.RemoveAt(Idx);
-			}
-		}
+	if (ProjectionData.TraceMode == EEnvQueryTrace::Geometry)
+	{
+		FEQSHelpers::RunPhysProjection(QueryInstance.World, ProjectionData, Points);
 	}
 }
 
