fixed crash in crowd movement when AI is using direct path to location
#ue4

diff --git a/Engine/Source/Runtime/AIModule/Private/Navigation/CrowdFollowingComponent.cpp b/Engine/Source/Runtime/AIModule/Private/Navigation/CrowdFollowingComponent.cpp
--- a/Engine/Source/Runtime/AIModule/Private/Navigation/CrowdFollowingComponent.cpp
+++ b/Engine/Source/Runtime/AIModule/Private/Navigation/CrowdFollowingComponent.cpp
@@ -702,7 +702,7 @@ void UCrowdFollowingComponent::FollowPathSegment(float DeltaTime)
 
 	if (bUpdateDirectMoveVelocity)
 	{
-		const FVector CurrentTargetPt = DestinationActor->GetActorLocation();
+		const FVector CurrentTargetPt = DestinationActor.IsValid() ? DestinationActor->GetActorLocation() : GetCurrentTargetLocation();
 		const FVector AgentLoc = GetCrowdAgentLocation();
 		const FVector NewDirection = (CurrentTargetPt - AgentLoc).GetSafeNormal();
 
