exposed EQS test's filtering and scoring options for multiple contexts
#ue4 UE-12158

diff --git a/Engine/Source/Runtime/AIModule/Classes/EnvironmentQuery/EnvQueryTest.h b/Engine/Source/Runtime/AIModule/Classes/EnvironmentQuery/EnvQueryTest.h
--- a/Engine/Source/Runtime/AIModule/Classes/EnvironmentQuery/EnvQueryTest.h
+++ b/Engine/Source/Runtime/AIModule/Classes/EnvironmentQuery/EnvQueryTest.h
@@ -48,6 +48,14 @@ class AIMODULE_API UEnvQueryTest : public UEnvQueryNode
 	UPROPERTY(EditDefaultsOnly, Category=Test)
 	TEnumAsByte<EEnvTestPurpose::Type> TestPurpose;
 	
+	/** Determines filtering operator when context returns multiple items */
+	UPROPERTY(EditDefaultsOnly, Category=Filter, AdvancedDisplay)
+	TEnumAsByte<EEnvTestFilterOperator::Type> MultipleContextFilterOp;
+
+	/** Determines scoring operator when context returns multiple items */
+	UPROPERTY(EditDefaultsOnly, Category = Score, AdvancedDisplay)
+	TEnumAsByte<EEnvTestScoreOperator::Type> MultipleContextScoreOp;
+
 	/** Does this test filter out results that are below a lower limit, above an upper limit, or both?  Or does it just look for a matching value? */
 	UPROPERTY(EditDefaultsOnly, Category=Filter)
 	TEnumAsByte<EEnvTestFilterType::Type> FilterType;
@@ -67,6 +75,10 @@ class AIMODULE_API UEnvQueryTest : public UEnvQueryNode
 	/** Cost of test */
 	TEnumAsByte<EEnvTestCost::Type> Cost;
 
+	/** The shape of the curve equation to apply to the normalized score before multiplying by factor. */
+	UPROPERTY(EditDefaultsOnly, Category=Score)
+	TEnumAsByte<EEnvTestScoreEquation::Type> ScoringEquation;
+
 	/** How should the lower bound for normalization of the raw test value before applying the scoring formula be determined?
 	    Should it use the lowest value found (tested), the lower threshold for filtering, or a separate specified normalization minimum? */
 	UPROPERTY(EditDefaultsOnly, Category=Score)
@@ -85,10 +97,6 @@ class AIMODULE_API UEnvQueryTest : public UEnvQueryNode
 	UPROPERTY(EditDefaultsOnly, Category=Score)
 	FAIDataProviderFloatValue ScoreClampMax;
 
-	/** The shape of the curve equation to apply to the normalized score before multiplying by factor. */
-	UPROPERTY(EditDefaultsOnly, Category=Score)
-	TEnumAsByte<EEnvTestScoreEquation::Type> ScoringEquation;
-
 	/** The weight (factor) by which to multiply the normalized score after the scoring equation is applied. */
 	UPROPERTY(EditDefaultsOnly, Category=Score, Meta=(ClampMin="0.001", UIMin="0.001"))
 	FAIDataProviderFloatValue ScoringFactor;
diff --git a/Engine/Source/Runtime/AIModule/Classes/EnvironmentQuery/EnvQueryTypes.h b/Engine/Source/Runtime/AIModule/Classes/EnvironmentQuery/EnvQueryTypes.h
--- a/Engine/Source/Runtime/AIModule/Classes/EnvironmentQuery/EnvQueryTypes.h
+++ b/Engine/Source/Runtime/AIModule/Classes/EnvironmentQuery/EnvQueryTypes.h
@@ -106,6 +106,27 @@ namespace EEnvTestCost
 	};
 }
 
+UENUM()
+namespace EEnvTestFilterOperator
+{
+	enum Type
+	{
+		AllPass			UMETA(Tooltip = "All contexts must pass condition"),
+		AnyPass			UMETA(Tooltip = "At least one context must pass condition"),
+	};
+}
+
+UENUM()
+namespace EEnvTestScoreOperator
+{
+	enum Type
+	{
+		AverageScore	UMETA(Tooltip = "Use average score from all contexts"),
+		MinScore		UMETA(Tooltip = "Use minimum score from all contexts"),
+		MaxScore		UMETA(Tooltip = "Use maximum score from all contexts"),
+	};
+}
+
 UENUM()
 namespace EEnvQueryStatus
 {
@@ -671,16 +692,19 @@ public:
 #if !NO_LOGGING
 	void Log(const FString Msg) const;
 #endif // #if !NO_LOGGING
-	
+
 #if USE_EQS_DEBUGGER
-#	define  UE_EQS_DBGMSG(Format, ...) \
-					Instance->ItemDetails[CurrentItem].FailedDescription = FString::Printf(Format, ##__VA_ARGS__)
+#	define  UE_EQS_DBGMSG(Condition, Format, ...) \
+					if (Condition) \
+					{ \
+						Instance->ItemDetails[CurrentItem].FailedDescription = FString::Printf(Format, ##__VA_ARGS__); \
+					}
 
 #	define UE_EQS_LOG(CategoryName, Verbosity, Format, ...) \
 					UE_LOG(CategoryName, Verbosity, Format, ##__VA_ARGS__); \
-					UE_EQS_DBGMSG(Format, ##__VA_ARGS__); 
+					UE_EQS_DBGMSG(true, Format, ##__VA_ARGS__); 
 #else
-#	define UE_EQS_DBGMSG(Format, ...)
+#	define UE_EQS_DBGMSG(Condition, Format, ...)
 #	define UE_EQS_LOG(CategoryName, Verbosity, Format, ...) UE_LOG(CategoryName, Verbosity, Format, ##__VA_ARGS__); 
 #endif
 
@@ -703,27 +727,18 @@ public:
 				switch (FilterType)
 				{
 					case EEnvTestFilterType::Maximum:
-						if (Score > Max)
-						{
-							UE_EQS_DBGMSG(TEXT("Value %f is above maximum value set to %f"), Score, Max);
-							bPassedTest = false;
-						}
+						bPassedTest = (Score <= Max);
+						UE_EQS_DBGMSG(!bPassedTest, TEXT("Value %f is above maximum value set to %f"), Score, Max);
 						break;
 
 					case EEnvTestFilterType::Minimum:
-						if (Score < Min)
-						{
-							UE_EQS_DBGMSG(TEXT("Value %f is below minimum value set to %f"), Score, Min);
-							bPassedTest = false;
-						}
+						bPassedTest = (Score >= Min);
+						UE_EQS_DBGMSG(!bPassedTest, TEXT("Value %f is below minimum value set to %f"), Score, Min);
 						break;
 
 					case EEnvTestFilterType::Range:
-						if ((Score < Min) || (Score > Max))
-						{
-							UE_EQS_DBGMSG(TEXT("Value %f is out of range set to (%f, %f)"), Score, Min, Max);
-							bPassedTest = false;
-						}
+						bPassedTest = (Score >= Min) && (Score <= Max);
+						UE_EQS_DBGMSG(!bPassedTest, TEXT("Value %f is out of range set to (%f, %f)"), Score, Min, Max);
 						break;
 
 					case EEnvTestFilterType::Match:
@@ -740,15 +755,11 @@ public:
 
 			if (bPassedTest)
 			{
-				// If we passed the test, either we really did, or we're only scoring, so we can't truly "fail".	
-				ItemScore += Score;
-				NumPartialScores++;
-			}
-			else
-			{
-				// We are ONLY filtering, and we failed
-				bPassed = false;
+				SetScoreInternal(Score);
+				NumPassedForItem++;
 			}
+
+			NumTestsForItem++;
 		}
 
 		void SetScore(EEnvTestPurpose::Type TestPurpose, EEnvTestFilterType::Type FilterType, bool bScore, bool bExpected)
@@ -758,12 +769,7 @@ public:
 			{
 				case EEnvTestFilterType::Match:
 					bPassedTest = (bScore == bExpected);
-#if USE_EQS_DEBUGGER
-					if (!bPassedTest)
-					{
-						UE_EQS_DBGMSG(TEXT("Boolean score don't mach (expected %s and got %s)"), bExpected ? TEXT("TRUE") : TEXT("FALSE"), bScore ? TEXT("TRUE") : TEXT("FALSE"));
-					}
-#endif
+					UE_EQS_DBGMSG(!bPassedTest, TEXT("Boolean score don't mach (expected %s and got %s)"), bExpected ? TEXT("TRUE") : TEXT("FALSE"), bScore ? TEXT("TRUE") : TEXT("FALSE"));
 					break;
 
 				case EEnvTestFilterType::Maximum:
@@ -787,23 +793,13 @@ public:
 					break;
 			}
 
-			if (!bPassedTest)
-			{
-				if (TestPurpose == EEnvTestPurpose::Score)
-				{
-					bSkipped = true;
-					NumPartialScores++;
-				}
-				else // We are filtering!
-				{
-					bPassed = false;
-				}
-			}
-			else
+			if (bPassedTest)
 			{
-				ItemScore += 1.0f;
-				NumPartialScores++;
+				SetScoreInternal(1.0f);
+				NumPassedForItem++;
 			}
+
+			NumTestsForItem++;
 		}
 
 		uint8* GetItemData()
@@ -818,7 +814,7 @@ public:
 
 		void SkipItem()
 		{
-			bSkipped++;
+			bSkipped = true;
 		}
 
 		void IgnoreTimeLimit()
@@ -849,16 +845,21 @@ public:
 	protected:
 
 		FEnvQueryInstance* Instance;
-		int32 CurrentItem;
-		int32 NumPartialScores;
 		double Deadline;
 		float ItemScore;
-		uint32 bPassed : 1;
-		uint32 bSkipped : 1;
+		int32 CurrentItem;
+		int16 NumPassedForItem;
+		int16 NumTestsForItem;
+		uint8 CachedFilterOp;
+		uint8 CachedScoreOp;
+		uint8 bPassed : 1;
+		uint8 bSkipped : 1;
+		uint8 bIsFiltering : 1;
 
 		void InitItemScore()
 		{
-			NumPartialScores = 0;
+			NumPassedForItem = 0;
+			NumTestsForItem = 0;
 			ItemScore = 0.0f;
 			bPassed = true;
 			bSkipped = false;
@@ -872,6 +873,46 @@ public:
 			for (CurrentItem++; CurrentItem < Instance->Items.Num() && !Instance->Items[CurrentItem].IsValid(); CurrentItem++)
 				;
 		}
+
+		FORCEINLINE void SetScoreInternal(float Score)
+		{
+			switch (CachedScoreOp)
+			{
+			case EEnvTestScoreOperator::AverageScore:
+				ItemScore += Score;
+				break;
+
+			case EEnvTestScoreOperator::MinScore:
+				if (!NumPassedForItem || ItemScore > Score)
+				{
+					ItemScore = Score;
+				}
+				break;
+
+			case EEnvTestScoreOperator::MaxScore:
+				if (!NumPassedForItem || ItemScore < Score)
+				{
+					ItemScore = Score;
+				}
+				break;
+			}
+		}
+
+		FORCEINLINE void CheckItemPassed()
+		{
+			if (!bIsFiltering)
+			{
+				bPassed = true;
+			}
+			else if (CachedFilterOp == EEnvTestFilterOperator::AllPass)
+			{
+				bPassed = bPassed && (NumPassedForItem == NumTestsForItem);
+			}
+			else
+			{
+				bPassed = bPassed && (NumPassedForItem > 0);
+			}
+		}
 	};
 #endif
 
diff --git a/Engine/Source/Runtime/AIModule/Private/EnvironmentQuery/EnvQueryInstance.cpp b/Engine/Source/Runtime/AIModule/Private/EnvironmentQuery/EnvQueryInstance.cpp
--- a/Engine/Source/Runtime/AIModule/Private/EnvironmentQuery/EnvQueryInstance.cpp
+++ b/Engine/Source/Runtime/AIModule/Private/EnvironmentQuery/EnvQueryInstance.cpp
@@ -327,6 +327,10 @@ FEnvQueryInstance::ItemIterator::ItemIterator(const UEnvQueryTest* QueryTest, FE
 	: Instance(&QueryInstance)
 	, CurrentItem(StartingItemIndex != INDEX_NONE ? StartingItemIndex : QueryInstance.CurrentTestStartingItem)
 {
+	CachedFilterOp = QueryTest ? QueryTest->MultipleContextFilterOp.GetValue() : EEnvTestFilterOperator::AllPass;
+	CachedScoreOp = QueryTest ? QueryTest->MultipleContextScoreOp.GetValue() : EEnvTestScoreOperator::AverageScore;
+	bIsFiltering = (QueryTest->TestPurpose == EEnvTestPurpose::Filter) || (QueryTest->TestPurpose == EEnvTestPurpose::FilterAndScore);
+
 	Deadline = QueryInstance.TimeLimit > 0.0 ? (FPlatformTime::Seconds() + QueryInstance.TimeLimit) : -1.0;
 	// it's possible item 'CurrentItem' has been already discarded. Find a valid starting index
 	--CurrentItem;
@@ -346,6 +350,8 @@ void FEnvQueryInstance::ItemIterator::HandleFailedTestResult()
 
 void FEnvQueryInstance::ItemIterator::StoreTestResult()
 {
+	CheckItemPassed();
+
 	if (Instance->IsInSingleItemFinalSearch())
 	{
 		// handle SingleResult mode
@@ -354,7 +360,7 @@ void FEnvQueryInstance::ItemIterator::StoreTestResult()
 			Instance->PickSingleItem(CurrentItem);
 			Instance->bFoundSingleResult = true;
 		}
-		else if (!bPassed || NumPartialScores == 0)
+		else if (!bPassed)
 		{
 			HandleFailedTestResult();
 		}
@@ -365,13 +371,13 @@ void FEnvQueryInstance::ItemIterator::StoreTestResult()
 		{
 			ItemScore = UEnvQueryTypes::SkippedItemValue;
 		}
-		else if (!bPassed || NumPartialScores == 0)
+		else if (!bPassed)
 		{
 			HandleFailedTestResult();
 		}
-		else
+		else if (CachedScoreOp == EEnvTestScoreOperator::AverageScore)
 		{
-			ItemScore /= NumPartialScores;
+			ItemScore /= NumPassedForItem;
 		}
 
 		Instance->ItemDetails[CurrentItem].TestResults[Instance->CurrentTest] = ItemScore;
