fixed pathing grid generator getting stuck on points inside geometry, even though there was valid navmesh on top
#ue4

diff --git a/Engine/Source/Runtime/AIModule/Private/EnvironmentQuery/Generators/EnvQueryGenerator_PathingGrid.cpp b/Engine/Source/Runtime/AIModule/Private/EnvironmentQuery/Generators/EnvQueryGenerator_PathingGrid.cpp
--- a/Engine/Source/Runtime/AIModule/Private/EnvironmentQuery/Generators/EnvQueryGenerator_PathingGrid.cpp
+++ b/Engine/Source/Runtime/AIModule/Private/EnvironmentQuery/Generators/EnvQueryGenerator_PathingGrid.cpp
@@ -8,6 +8,7 @@
 
 UEnvQueryGenerator_PathingGrid::UEnvQueryGenerator_PathingGrid(const FObjectInitializer& ObjectInitializer) : Super(ObjectInitializer)
 {
+	ProjectionData.SetNavmeshOnly();
 	GridSize.DefaultValue = 1000.0f;
 	SpaceBetween.DefaultValue = 100.0f;
 	PathToItem.DefaultValue = true;
@@ -53,6 +54,9 @@ void UEnvQueryGenerator_PathingGrid::ProjectAndFilterNavPoints(TArray<FNavLocati
 	
 	{
 		TArray<NavNodeRef> Polys;
+		TArray<FNavLocation> HitLocations;
+		const FVector ProjectionExtent(ProjectionData.ExtentX, ProjectionData.ExtentX, (ProjectionData.ProjectDown + ProjectionData.ProjectUp) / 2);
+
 		for (int32 ContextIdx = 0; ContextIdx < ContextLocations.Num() && Points.Num(); ContextIdx++)
 		{
 			float CollectDistanceSq = 0.0f;
@@ -73,7 +77,27 @@ void UEnvQueryGenerator_PathingGrid::ProjectAndFilterNavPoints(TArray<FNavLocati
 
 			for (int32 Idx = Points.Num() - 1; Idx >= 0; Idx--)
 			{
-				const bool bHasPath = PathGridHelpers::HasPath(NodePoolData, Points[Idx].NodeRef);
+				bool bHasPath = PathGridHelpers::HasPath(NodePoolData, Points[Idx].NodeRef);
+				if (!bHasPath && Points[Idx].NodeRef != INVALID_NAVNODEREF)
+				{
+					// try projecting it again, maybe it will match valid poly on different height
+					HitLocations.Reset();
+					FVector TestPt(Points[Idx].Location.X, Points[Idx].Location.Y, ContextLocations[ContextIdx].Z);
+
+					NavMeshData->ProjectPointMulti(TestPt, HitLocations, ProjectionExtent, TestPt.Z - ProjectionData.ProjectDown, TestPt.Z + ProjectionData.ProjectUp, NavigationFilterOb, nullptr);
+					for (int32 HitIdx = 0; HitIdx < HitLocations.Num(); HitIdx++)
+					{
+						const bool bHasPathTest = PathGridHelpers::HasPath(NodePoolData, HitLocations[HitIdx].NodeRef);
+						if (bHasPathTest)
+						{
+							Points[Idx] = HitLocations[HitIdx];
+							Points[Idx].Location.Z += ProjectionData.PostProjectionVerticalOffset;
+							bHasPath = true;
+							break;
+						}
+					}
+				}
+
 				if (!bHasPath)
 				{
 					Points.RemoveAt(Idx);
diff --git a/Engine/Source/Runtime/Engine/Private/AI/Navigation/PImplRecastNavMesh.cpp b/Engine/Source/Runtime/Engine/Private/AI/Navigation/PImplRecastNavMesh.cpp
--- a/Engine/Source/Runtime/Engine/Private/AI/Navigation/PImplRecastNavMesh.cpp
+++ b/Engine/Source/Runtime/Engine/Private/AI/Navigation/PImplRecastNavMesh.cpp
@@ -1468,11 +1468,18 @@ bool FPImplRecastNavMesh::ProjectPointMulti(const FVector& Point, TArray<FNavLoc
 				status = NavQuery.projectedPointOnPoly(HitPolys[i], &RcPoint.X, ClosestPoint);
 				if (dtStatusSucceed(status))
 				{
-					FNavLocation HitLocation(Recast2UnrealPoint(ClosestPoint), HitPolys[i]);
-					ensure((HitLocation.Location - AdjustedPoint).SizeSquared2D() < KINDA_SMALL_NUMBER);
-					
-					Result.Add(HitLocation);
-					bSuccess = true;
+					float ExactZ = 0.0f;
+					status = NavQuery.getPolyHeight(HitPolys[i], ClosestPoint, &ExactZ);
+					if (dtStatusSucceed(status))
+					{
+						FNavLocation HitLocation(Recast2UnrealPoint(ClosestPoint), HitPolys[i]);
+						HitLocation.Location.Z = ExactZ;
+
+						ensure((HitLocation.Location - AdjustedPoint).SizeSquared2D() < KINDA_SMALL_NUMBER);
+
+						Result.Add(HitLocation);
+						bSuccess = true;
+					}
 				}
 			}
 		}
