AI: Fixed linker errors with EnvQueryTraceHelper

A specialized template member function is a regular member function and must therefore either be inlined in the header or defined in a compilaton unit. Otherwise, multiple inclusions of the header will result in multiple defined symbols at link time.

diff --git a/Engine/Source/Runtime/AIModule/Private/EnvironmentQuery/EnvQueryTraceHelpers.cpp b/Engine/Source/Runtime/AIModule/Private/EnvironmentQuery/EnvQueryTraceHelpers.cpp
--- a/Engine/Source/Runtime/AIModule/Private/EnvironmentQuery/EnvQueryTraceHelpers.cpp
+++ b/Engine/Source/Runtime/AIModule/Private/EnvironmentQuery/EnvQueryTraceHelpers.cpp
@@ -3,6 +3,159 @@
 #include "AIModulePrivate.h"
 #include "EnvQueryTraceHelpers.h"
 
+
+template<>
+void FEQSHelpers::FBatchTrace::DoSingleSourceMultiDestinations<EEnvTraceShape::Line>(const FVector& Source, TArray<FNavLocation>& Points)
+{
+	FVector HitPos(FVector::ZeroVector);
+	for (int32 Idx = Points.Num() - 1; Idx >= 0; Idx--)
+	{
+		const bool bHit = RunLineTrace(Source, Points[Idx].Location, HitPos);
+		if (bHit)
+		{
+			Points[Idx] = HitPos;
+		}
+		else if (TraceMode == ETraceMode::Discard)
+		{
+			Points.RemoveAt(Idx, 1, false);
+		}
+	}
+}
+
+
+template<>
+void FEQSHelpers::FBatchTrace::DoSingleSourceMultiDestinations<EEnvTraceShape::Box>(const FVector& Source, TArray<FNavLocation>& Points)
+{
+	FVector HitPos(FVector::ZeroVector);
+	for (int32 Idx = Points.Num() - 1; Idx >= 0; Idx--)
+	{
+		const bool bHit = RunBoxTrace(Source, Points[Idx].Location, HitPos);
+		if (bHit)
+		{
+			Points[Idx] = HitPos;
+		}
+		else if (TraceMode == ETraceMode::Discard)
+		{
+			Points.RemoveAt(Idx, 1, false);
+		}
+	}
+}
+
+
+template<>
+void FEQSHelpers::FBatchTrace::DoSingleSourceMultiDestinations<EEnvTraceShape::Sphere>(const FVector& Source, TArray<FNavLocation>& Points)
+{
+	FVector HitPos(FVector::ZeroVector);
+	for (int32 Idx = Points.Num() - 1; Idx >= 0; Idx--)
+	{
+		const bool bHit = RunSphereTrace(Source, Points[Idx].Location, HitPos);
+		if (bHit)
+		{
+			Points[Idx] = HitPos;
+		}
+		else if (TraceMode == ETraceMode::Discard)
+		{
+			Points.RemoveAt(Idx, 1, false);
+		}
+	}
+}
+
+
+template<>
+void FEQSHelpers::FBatchTrace::DoSingleSourceMultiDestinations<EEnvTraceShape::Capsule>(const FVector& Source, TArray<FNavLocation>& Points)
+{
+	FVector HitPos(FVector::ZeroVector);
+	for (int32 Idx = Points.Num() - 1; Idx >= 0; Idx--)
+	{
+		const bool bHit = RunCapsuleTrace(Source, Points[Idx].Location, HitPos);
+		if (bHit)
+		{
+			Points[Idx] = HitPos;
+		}
+		else if (TraceMode == ETraceMode::Discard)
+		{
+			Points.RemoveAt(Idx, 1, false);
+		}
+	}
+}
+
+
+template<>
+void FEQSHelpers::FBatchTrace::DoProject<EEnvTraceShape::Line>(TArray<FNavLocation>& Points, float StartOffsetZ, float EndOffsetZ, float HitOffsetZ)
+{
+	FVector HitPos(FVector::ZeroVector);
+	for (int32 Idx = Points.Num() - 1; Idx >= 0; Idx--)
+	{
+		const bool bHit = RunLineTrace(Points[Idx].Location + FVector(0, 0, StartOffsetZ), Points[Idx].Location + FVector(0, 0, EndOffsetZ), HitPos);
+		if (bHit)
+		{
+			Points[Idx] = HitPos + FVector(0, 0, HitOffsetZ);
+		}
+		else if (TraceMode == ETraceMode::Discard)
+		{
+			Points.RemoveAt(Idx, 1, false);
+		}
+	}
+}
+
+
+template<>
+void FEQSHelpers::FBatchTrace::DoProject<EEnvTraceShape::Box>(TArray<FNavLocation>& Points, float StartOffsetZ, float EndOffsetZ, float HitOffsetZ)
+{
+	FVector HitPos(FVector::ZeroVector);
+	for (int32 Idx = Points.Num() - 1; Idx >= 0; Idx--)
+	{
+		const bool bHit = RunBoxTrace(Points[Idx].Location + FVector(0, 0, StartOffsetZ), Points[Idx].Location + FVector(0, 0, EndOffsetZ), HitPos);
+		if (bHit)
+		{
+			Points[Idx] = HitPos + FVector(0, 0, HitOffsetZ);
+		}
+		else if (TraceMode == ETraceMode::Discard)
+		{
+			Points.RemoveAt(Idx, 1, false);
+		}
+	}
+}
+
+
+template<>
+void FEQSHelpers::FBatchTrace::DoProject<EEnvTraceShape::Sphere>(TArray<FNavLocation>& Points, float StartOffsetZ, float EndOffsetZ, float HitOffsetZ)
+{
+	FVector HitPos(FVector::ZeroVector);
+	for (int32 Idx = Points.Num() - 1; Idx >= 0; Idx--)
+	{
+		const bool bHit = RunSphereTrace(Points[Idx].Location + FVector(0, 0, StartOffsetZ), Points[Idx].Location + FVector(0, 0, EndOffsetZ), HitPos);
+		if (bHit)
+		{
+			Points[Idx] = HitPos + FVector(0, 0, HitOffsetZ);
+		}
+		else if (TraceMode == ETraceMode::Discard)
+		{
+			Points.RemoveAt(Idx, 1, false);
+		}
+	}
+}
+
+
+template<>
+void FEQSHelpers::FBatchTrace::DoProject<EEnvTraceShape::Capsule>(TArray<FNavLocation>& Points, float StartOffsetZ, float EndOffsetZ, float HitOffsetZ)
+{
+	FVector HitPos(FVector::ZeroVector);
+	for (int32 Idx = Points.Num() - 1; Idx >= 0; Idx--)
+	{
+		const bool bHit = RunCapsuleTrace(Points[Idx].Location + FVector(0, 0, StartOffsetZ), Points[Idx].Location + FVector(0, 0, EndOffsetZ), HitPos);
+		if (bHit)
+		{
+			Points[Idx] = HitPos + FVector(0, 0, HitOffsetZ);
+		}
+		else if (TraceMode == ETraceMode::Discard)
+		{
+			Points.RemoveAt(Idx, 1, false);
+		}
+	}
+}
+
+
 void FEQSHelpers::RunNavRaycasts(ANavigationData* NavData, const FEnvTraceData& TraceData, const FVector& SourcePt, TArray<FNavLocation>& Points, ETraceMode TraceMode)
 {
 	TSharedPtr<const FNavigationQueryFilter> NavigationFilter = UNavigationQueryFilter::GetQueryFilter(NavData, TraceData.NavigationFilter);
@@ -34,6 +187,7 @@ void FEQSHelpers::RunNavRaycasts(ANavigationData* NavData, const FEnvTraceData&
 	}
 }
 
+
 void FEQSHelpers::RunNavProjection(ANavigationData* NavData, const FEnvTraceData& TraceData, TArray<FNavLocation>& Points, ETraceMode TraceMode)
 {
 	TSharedPtr<const FNavigationQueryFilter> NavigationFilter = UNavigationQueryFilter::GetQueryFilter(NavData, TraceData.NavigationFilter);
@@ -73,6 +227,7 @@ void FEQSHelpers::RunNavProjection(ANavigationData* NavData, const FEnvTraceData
 	}
 }
 
+
 void FEQSHelpers::RunPhysRaycasts(UWorld* World, const FEnvTraceData& TraceData, const FVector& SourcePt, TArray<FNavLocation>& Points, ETraceMode TraceMode)
 {
 	ECollisionChannel TraceCollisionChannel = UEngineTypes::ConvertToCollisionChannel(TraceData.TraceChannel);
@@ -106,6 +261,7 @@ void FEQSHelpers::RunPhysRaycasts(UWorld* World, const FEnvTraceData& TraceData,
 	}
 }
 
+
 void FEQSHelpers::RunPhysProjection(UWorld* World, const FEnvTraceData& TraceData, TArray<FNavLocation>& Points, ETraceMode TraceMode)
 {
 	ECollisionChannel TraceCollisionChannel = UEngineTypes::ConvertToCollisionChannel(TraceData.TraceChannel);
diff --git a/Engine/Source/Runtime/AIModule/Private/EnvironmentQuery/EnvQueryTraceHelpers.h b/Engine/Source/Runtime/AIModule/Private/EnvironmentQuery/EnvQueryTraceHelpers.h
--- a/Engine/Source/Runtime/AIModule/Private/EnvironmentQuery/EnvQueryTraceHelpers.h
+++ b/Engine/Source/Runtime/AIModule/Private/EnvironmentQuery/EnvQueryTraceHelpers.h
@@ -71,148 +71,28 @@ namespace FEQSHelpers
 	};
 
 	template<>
-	void FBatchTrace::DoSingleSourceMultiDestinations<EEnvTraceShape::Line>(const FVector& Source, TArray<FNavLocation>& Points)
-	{
-		FVector HitPos(FVector::ZeroVector);
-		for (int32 Idx = Points.Num() - 1; Idx >= 0; Idx--)
-		{
-			const bool bHit = RunLineTrace(Source, Points[Idx].Location, HitPos);
-			if (bHit)
-			{
-				Points[Idx] = HitPos;
-			}
-			else if (TraceMode == ETraceMode::Discard)
-			{
-				Points.RemoveAt(Idx, 1, false);
-			}
-		}
-	}
+	void FBatchTrace::DoSingleSourceMultiDestinations<EEnvTraceShape::Line>(const FVector& Source, TArray<FNavLocation>& Points);
 
 	template<>
-	void FBatchTrace::DoSingleSourceMultiDestinations<EEnvTraceShape::Box>(const FVector& Source, TArray<FNavLocation>& Points)
-	{
-		FVector HitPos(FVector::ZeroVector);
-		for (int32 Idx = Points.Num() - 1; Idx >= 0; Idx--)
-		{
-			const bool bHit = RunBoxTrace(Source, Points[Idx].Location, HitPos);
-			if (bHit)
-			{
-				Points[Idx] = HitPos;
-			}
-			else if (TraceMode == ETraceMode::Discard)
-			{
-				Points.RemoveAt(Idx, 1, false);
-			}
-		}
-	}
+	void FBatchTrace::DoSingleSourceMultiDestinations<EEnvTraceShape::Box>(const FVector& Source, TArray<FNavLocation>& Points);
 
 	template<>
-	void FBatchTrace::DoSingleSourceMultiDestinations<EEnvTraceShape::Sphere>(const FVector& Source, TArray<FNavLocation>& Points)
-	{
-		FVector HitPos(FVector::ZeroVector);
-		for (int32 Idx = Points.Num() - 1; Idx >= 0; Idx--)
-		{
-			const bool bHit = RunSphereTrace(Source, Points[Idx].Location, HitPos);
-			if (bHit)
-			{
-				Points[Idx] = HitPos;
-			}
-			else if (TraceMode == ETraceMode::Discard)
-			{
-				Points.RemoveAt(Idx, 1, false);
-			}
-		}
-	}
+	void FBatchTrace::DoSingleSourceMultiDestinations<EEnvTraceShape::Sphere>(const FVector& Source, TArray<FNavLocation>& Points);
 
 	template<>
-	void FBatchTrace::DoSingleSourceMultiDestinations<EEnvTraceShape::Capsule>(const FVector& Source, TArray<FNavLocation>& Points)
-	{
-		FVector HitPos(FVector::ZeroVector);
-		for (int32 Idx = Points.Num() - 1; Idx >= 0; Idx--)
-		{
-			const bool bHit = RunCapsuleTrace(Source, Points[Idx].Location, HitPos);
-			if (bHit)
-			{
-				Points[Idx] = HitPos;
-			}
-			else if (TraceMode == ETraceMode::Discard)
-			{
-				Points.RemoveAt(Idx, 1, false);
-			}
-		}
-	}
+	void FBatchTrace::DoSingleSourceMultiDestinations<EEnvTraceShape::Capsule>(const FVector& Source, TArray<FNavLocation>& Points);
 
 	template<>
-	void FBatchTrace::DoProject<EEnvTraceShape::Line>(TArray<FNavLocation>& Points, float StartOffsetZ, float EndOffsetZ, float HitOffsetZ)
-	{
-		FVector HitPos(FVector::ZeroVector);
-		for (int32 Idx = Points.Num() - 1; Idx >= 0; Idx--)
-		{
-			const bool bHit = RunLineTrace(Points[Idx].Location + FVector(0, 0, StartOffsetZ), Points[Idx].Location + FVector(0, 0, EndOffsetZ), HitPos);
-			if (bHit)
-			{
-				Points[Idx] = HitPos + FVector(0, 0, HitOffsetZ);
-			}
-			else if (TraceMode == ETraceMode::Discard)
-			{
-				Points.RemoveAt(Idx, 1, false);
-			}
-		}
-	}
+	void FBatchTrace::DoProject<EEnvTraceShape::Line>(TArray<FNavLocation>& Points, float StartOffsetZ, float EndOffsetZ, float HitOffsetZ);
 
 	template<>
-	void FBatchTrace::DoProject<EEnvTraceShape::Box>(TArray<FNavLocation>& Points, float StartOffsetZ, float EndOffsetZ, float HitOffsetZ)
-	{
-		FVector HitPos(FVector::ZeroVector);
-		for (int32 Idx = Points.Num() - 1; Idx >= 0; Idx--)
-		{
-			const bool bHit = RunBoxTrace(Points[Idx].Location + FVector(0, 0, StartOffsetZ), Points[Idx].Location + FVector(0, 0, EndOffsetZ), HitPos);
-			if (bHit)
-			{
-				Points[Idx] = HitPos + FVector(0, 0, HitOffsetZ);
-			}
-			else if (TraceMode == ETraceMode::Discard)
-			{
-				Points.RemoveAt(Idx, 1, false);
-			}
-		}
-	}
+	void FBatchTrace::DoProject<EEnvTraceShape::Box>(TArray<FNavLocation>& Points, float StartOffsetZ, float EndOffsetZ, float HitOffsetZ);
 
 	template<>
-	void FBatchTrace::DoProject<EEnvTraceShape::Sphere>(TArray<FNavLocation>& Points, float StartOffsetZ, float EndOffsetZ, float HitOffsetZ)
-	{
-		FVector HitPos(FVector::ZeroVector);
-		for (int32 Idx = Points.Num() - 1; Idx >= 0; Idx--)
-		{
-			const bool bHit = RunSphereTrace(Points[Idx].Location + FVector(0, 0, StartOffsetZ), Points[Idx].Location + FVector(0, 0, EndOffsetZ), HitPos);
-			if (bHit)
-			{
-				Points[Idx] = HitPos + FVector(0, 0, HitOffsetZ);
-			}
-			else if (TraceMode == ETraceMode::Discard)
-			{
-				Points.RemoveAt(Idx, 1, false);
-			}
-		}
-	}
+	void FBatchTrace::DoProject<EEnvTraceShape::Sphere>(TArray<FNavLocation>& Points, float StartOffsetZ, float EndOffsetZ, float HitOffsetZ);
 
 	template<>
-	void FBatchTrace::DoProject<EEnvTraceShape::Capsule>(TArray<FNavLocation>& Points, float StartOffsetZ, float EndOffsetZ, float HitOffsetZ)
-	{
-		FVector HitPos(FVector::ZeroVector);
-		for (int32 Idx = Points.Num() - 1; Idx >= 0; Idx--)
-		{
-			const bool bHit = RunCapsuleTrace(Points[Idx].Location + FVector(0, 0, StartOffsetZ), Points[Idx].Location + FVector(0, 0, EndOffsetZ), HitPos);
-			if (bHit)
-			{
-				Points[Idx] = HitPos + FVector(0, 0, HitOffsetZ);
-			}
-			else if (TraceMode == ETraceMode::Discard)
-			{
-				Points.RemoveAt(Idx, 1, false);
-			}
-		}
-	}
+	void FBatchTrace::DoProject<EEnvTraceShape::Capsule>(TArray<FNavLocation>& Points, float StartOffsetZ, float EndOffsetZ, float HitOffsetZ);
 
 	void RunNavRaycasts(ANavigationData* NavData, const FEnvTraceData& TraceData, const FVector& SourcePt, TArray<FNavLocation>& Points, ETraceMode TraceMode = ETraceMode::Keep);
 	void RunNavProjection(ANavigationData* NavData, const FEnvTraceData& TraceData, TArray<FNavLocation>& Points, ETraceMode TraceMode = ETraceMode::Discard);
