Changed FEQSRenderingHelper::ExtractLocation to take a const reference of query's raw data rather than copying it every time #UE4

diff --git a/Engine/Source/Runtime/AIModule/Private/EnvironmentQuery/EQSRenderingComponent.cpp b/Engine/Source/Runtime/AIModule/Private/EnvironmentQuery/EQSRenderingComponent.cpp
--- a/Engine/Source/Runtime/AIModule/Private/EnvironmentQuery/EQSRenderingComponent.cpp
+++ b/Engine/Source/Runtime/AIModule/Private/EnvironmentQuery/EQSRenderingComponent.cpp
@@ -10,7 +10,7 @@ static const int32 EQSMaxItemsDrawn = 10000;
 
 namespace FEQSRenderingHelper
 {
-	FVector ExtractLocation(TSubclassOf<UEnvQueryItemType> ItemType, TArray<uint8> RawData, const TArray<FEnvQueryItem>& Items, int32 Index)
+	FVector ExtractLocation(TSubclassOf<UEnvQueryItemType> ItemType, const TArray<uint8>& RawData, const TArray<FEnvQueryItem>& Items, int32 Index)
 	{
 		if (Items.IsValidIndex(Index) &&
 			ItemType->IsChildOf(UEnvQueryItemType_VectorBase::StaticClass()))
@@ -145,7 +145,8 @@ void FEQSSceneProxy::CollectEQSData(const FEnvQueryResult* ResultItems, const FE
 
 	// using "mid-results" requires manual normalization
 	const bool bUseMidResults = QueryInstance && (QueryInstance->Items.Num() < QueryInstance->DebugData.DebugItems.Num());
-	const TArray<FEnvQueryItem>& Items = (QueryInstance && bUseMidResults) ? QueryInstance->DebugData.DebugItems : ResultItems->Items;
+	// no point in checking if QueryInstance != null since bUseMidResults == true guarantees that
+	const TArray<FEnvQueryItem>& Items = bUseMidResults ? QueryInstance->DebugData.DebugItems : ResultItems->Items;
 	const TArray<uint8>& RawData = bUseMidResults ? QueryInstance->DebugData.RawData : ResultItems->RawData;
 	const int32 ItemCountLimit = FMath::Clamp(Items.Num(), 0, EQSMaxItemsDrawn);
 	const bool bNoTestsPerformed = QueryInstance != NULL && QueryInstance->CurrentTest <= 0;
