Removed the `check` from UAIPerceptionSystem::MakeNoiseImpl and massaged the function to take that into consideration #UE4

diff --git a/Engine/Source/Runtime/AIModule/Private/Perception/AIPerceptionSystem.cpp b/Engine/Source/Runtime/AIModule/Private/Perception/AIPerceptionSystem.cpp
--- a/Engine/Source/Runtime/AIModule/Private/Perception/AIPerceptionSystem.cpp
+++ b/Engine/Source/Runtime/AIModule/Private/Perception/AIPerceptionSystem.cpp
@@ -378,13 +378,17 @@ bool UAIPerceptionSystem::DeliverDelayedStimuli(UAIPerceptionSystem::EDelayedSti
 
 void UAIPerceptionSystem::MakeNoiseImpl(AActor* NoiseMaker, float Loudness, APawn* NoiseInstigator, const FVector& NoiseLocation)
 {
-	check(NoiseMaker != nullptr || NoiseInstigator != nullptr);
-
+	UE_CLOG(NoiseMaker == nullptr && NoiseInstigator == nullptr, LogAIPerception, Warning
+		, TEXT("UAIPerceptionSystem::MakeNoiseImpl called with both NoiseMaker and NoiseInstigator being null. Unable to resolve UWorld context!"));
+	
 	UWorld* World = NoiseMaker ? NoiseMaker->GetWorld() : (NoiseInstigator ? NoiseInstigator->GetWorld() : nullptr);
 
-	UAIPerceptionSystem::OnEvent(World, FAINoiseEvent(NoiseInstigator ? NoiseInstigator : NoiseMaker
+	if (World)
+	{
+		UAIPerceptionSystem::OnEvent(World, FAINoiseEvent(NoiseInstigator ? NoiseInstigator : NoiseMaker
 			, NoiseLocation
 			, Loudness));
+	}
 }
 
 bool UAIPerceptionSystem::RegisterPerceptionStimuliSource(UObject* WorldContext, TSubclassOf<UAISense> Sense, AActor* Target)
