Refactored a bunch of navigation-related EQS and NavigationFilter functions to take references rather than pointers to NavigationData #UE4

CL contains a deprecation handling as well.

diff --git a/Engine/Source/Runtime/AIModule/Classes/EnvironmentQuery/EnvQueryTypes.h b/Engine/Source/Runtime/AIModule/Classes/EnvironmentQuery/EnvQueryTypes.h
--- a/Engine/Source/Runtime/AIModule/Classes/EnvironmentQuery/EnvQueryTypes.h
+++ b/Engine/Source/Runtime/AIModule/Classes/EnvironmentQuery/EnvQueryTypes.h
@@ -929,7 +929,11 @@ public:
 
 namespace FEQSHelpers
 {
+	AIMODULE_API const ANavigationData* FindNavigationDataForQuery(FEnvQueryInstance& QueryInstance);
+
 #if WITH_RECAST
+
+	DEPRECATED(4.8, "FindNavMeshForQuery is deprecated. Please use FindNavigationDataForQuery")
 	AIMODULE_API const ARecastNavMesh* FindNavMeshForQuery(FEnvQueryInstance& QueryInstance);
 #endif // WITH_RECAST
 }
diff --git a/Engine/Source/Runtime/AIModule/Classes/EnvironmentQuery/Tests/EnvQueryTest_Pathfinding.h b/Engine/Source/Runtime/AIModule/Classes/EnvironmentQuery/Tests/EnvQueryTest_Pathfinding.h
--- a/Engine/Source/Runtime/AIModule/Classes/EnvironmentQuery/Tests/EnvQueryTest_Pathfinding.h
+++ b/Engine/Source/Runtime/AIModule/Classes/EnvironmentQuery/Tests/EnvQueryTest_Pathfinding.h
@@ -55,15 +55,15 @@ class UEnvQueryTest_Pathfinding : public UEnvQueryTest
 
 protected:
 
-	DECLARE_DELEGATE_RetVal_SixParams(bool, FTestPathSignature, const FVector&, const FVector&, EPathFindingMode::Type, const ANavigationData*, UNavigationSystem*, const UObject*);
-	DECLARE_DELEGATE_RetVal_SixParams(float, FFindPathSignature, const FVector&, const FVector&, EPathFindingMode::Type, const ANavigationData*, UNavigationSystem*, const UObject*);
+	DECLARE_DELEGATE_RetVal_SixParams(bool, FTestPathSignature, const FVector&, const FVector&, EPathFindingMode::Type, const ANavigationData&, UNavigationSystem&, const UObject*);
+	DECLARE_DELEGATE_RetVal_SixParams(float, FFindPathSignature, const FVector&, const FVector&, EPathFindingMode::Type, const ANavigationData&, UNavigationSystem&, const UObject*);
 
-	bool TestPathFrom(const FVector& ItemPos, const FVector& ContextPos, EPathFindingMode::Type Mode, const ANavigationData* NavData, UNavigationSystem* NavSys, const UObject* PathOwner) const;
-	bool TestPathTo(const FVector& ItemPos, const FVector& ContextPos, EPathFindingMode::Type Mode, const ANavigationData* NavData, UNavigationSystem* NavSys, const UObject* PathOwner) const;
-	float FindPathCostFrom(const FVector& ItemPos, const FVector& ContextPos, EPathFindingMode::Type Mode, const ANavigationData* NavData, UNavigationSystem* NavSys, const UObject* PathOwner) const;
-	float FindPathCostTo(const FVector& ItemPos, const FVector& ContextPos, EPathFindingMode::Type Mode, const ANavigationData* NavData, UNavigationSystem* NavSys, const UObject* PathOwner) const;
-	float FindPathLengthFrom(const FVector& ItemPos, const FVector& ContextPos, EPathFindingMode::Type Mode, const ANavigationData* NavData, UNavigationSystem* NavSys, const UObject* PathOwner) const;
-	float FindPathLengthTo(const FVector& ItemPos, const FVector& ContextPos, EPathFindingMode::Type Mode, const ANavigationData* NavData, UNavigationSystem* NavSys, const UObject* PathOwner) const;
+	bool TestPathFrom(const FVector& ItemPos, const FVector& ContextPos, EPathFindingMode::Type Mode, const ANavigationData& NavData, UNavigationSystem& NavSys, const UObject* PathOwner) const;
+	bool TestPathTo(const FVector& ItemPos, const FVector& ContextPos, EPathFindingMode::Type Mode, const ANavigationData& NavData, UNavigationSystem& NavSys, const UObject* PathOwner) const;
+	float FindPathCostFrom(const FVector& ItemPos, const FVector& ContextPos, EPathFindingMode::Type Mode, const ANavigationData& NavData, UNavigationSystem& NavSys, const UObject* PathOwner) const;
+	float FindPathCostTo(const FVector& ItemPos, const FVector& ContextPos, EPathFindingMode::Type Mode, const ANavigationData& NavData, UNavigationSystem& NavSys, const UObject* PathOwner) const;
+	float FindPathLengthFrom(const FVector& ItemPos, const FVector& ContextPos, EPathFindingMode::Type Mode, const ANavigationData& NavData, UNavigationSystem& NavSys, const UObject* PathOwner) const;
+	float FindPathLengthTo(const FVector& ItemPos, const FVector& ContextPos, EPathFindingMode::Type Mode, const ANavigationData& NavData, UNavigationSystem& NavSys, const UObject* PathOwner) const;
 
-	ANavigationData* FindNavigationData(UNavigationSystem* NavSys, UObject* Owner) const;
+	ANavigationData* FindNavigationData(UNavigationSystem& NavSys, UObject* Owner) const;
 };
diff --git a/Engine/Source/Runtime/AIModule/Private/AIController.cpp b/Engine/Source/Runtime/AIModule/Private/AIController.cpp
--- a/Engine/Source/Runtime/AIModule/Private/AIController.cpp
+++ b/Engine/Source/Runtime/AIModule/Private/AIController.cpp
@@ -650,30 +650,37 @@ bool AAIController::PreparePathfinding(const FAIMoveRequest& MoveRequest, FPathF
 			NavSys->GetNavDataForProps(GetNavAgentPropertiesRef()) :
 			NavSys->GetAbstractNavData();
 
-		FVector GoalLocation = MoveRequest.GetGoalLocation();
-		if (MoveRequest.HasGoalActor())
+		if (NavData)
 		{
-			const INavAgentInterface* NavGoal = Cast<const INavAgentInterface>(MoveRequest.GetGoalActor());
-			if (NavGoal)
+			FVector GoalLocation = MoveRequest.GetGoalLocation();
+			if (MoveRequest.HasGoalActor())
 			{
-				const FVector Offset = NavGoal->GetMoveGoalOffset(this);
-				GoalLocation = FQuatRotationTranslationMatrix(MoveRequest.GetGoalActor()->GetActorQuat(), NavGoal->GetNavAgentLocation()).TransformPosition(Offset);
+				const INavAgentInterface* NavGoal = Cast<const INavAgentInterface>(MoveRequest.GetGoalActor());
+				if (NavGoal)
+				{
+					const FVector Offset = NavGoal->GetMoveGoalOffset(this);
+					GoalLocation = FQuatRotationTranslationMatrix(MoveRequest.GetGoalActor()->GetActorQuat(), NavGoal->GetNavAgentLocation()).TransformPosition(Offset);
+				}
+				else
+				{
+					GoalLocation = MoveRequest.GetGoalActor()->GetActorLocation();
+				}
 			}
-			else
+
+			Query = FPathFindingQuery(this, *NavData, GetNavAgentLocation(), GoalLocation, UNavigationQueryFilter::GetQueryFilter(*NavData, MoveRequest.GetNavigationFilter()));
+			Query.SetAllowPartialPaths(MoveRequest.IsUsingPartialPaths());
+
+			if (PathFollowingComponent)
 			{
-				GoalLocation = MoveRequest.GetGoalActor()->GetActorLocation();
+				PathFollowingComponent->OnPathfindingQuery(Query);
 			}
-		}
 
-		Query = FPathFindingQuery(this, NavData, GetNavAgentLocation(), GoalLocation, UNavigationQueryFilter::GetQueryFilter(NavData, MoveRequest.GetNavigationFilter()));
-		Query.SetAllowPartialPaths(MoveRequest.IsUsingPartialPaths());
-
-		if (PathFollowingComponent)
+			return true;
+		}
+		else
 		{
-			PathFollowingComponent->OnPathfindingQuery(Query);
+			UE_VLOG(this, LogAINavigation, Warning, TEXT("Unable to find NavigationData instance while calling AAIController::PreparePathfinding"));
 		}
-
-		return true;
 	}
 
 	return false;
diff --git a/Engine/Source/Runtime/AIModule/Private/BehaviorTree/Decorators/BTDecorator_DoesPathExist.cpp b/Engine/Source/Runtime/AIModule/Private/BehaviorTree/Decorators/BTDecorator_DoesPathExist.cpp
--- a/Engine/Source/Runtime/AIModule/Private/BehaviorTree/Decorators/BTDecorator_DoesPathExist.cpp
+++ b/Engine/Source/Runtime/AIModule/Private/BehaviorTree/Decorators/BTDecorator_DoesPathExist.cpp
@@ -61,19 +61,22 @@ bool UBTDecorator_DoesPathExist::CalculateRawConditionValue(UBehaviorTreeCompone
 	{
 		const AAIController* AIOwner = OwnerComp.GetAIOwner();
 		const ANavigationData* NavData = AIOwner ? NavSys->GetNavDataForProps(AIOwner->GetNavAgentPropertiesRef()) : NULL;
-		TSharedPtr<const FNavigationQueryFilter> QueryFilter = UNavigationQueryFilter::GetQueryFilter(NavData, FilterClass);
-
-		if (PathQueryType == EPathExistanceQueryType::NavmeshRaycast2D)
+		if (NavData)
 		{
+			TSharedPtr<const FNavigationQueryFilter> QueryFilter = UNavigationQueryFilter::GetQueryFilter(*NavData, FilterClass);
+
+			if (PathQueryType == EPathExistanceQueryType::NavmeshRaycast2D)
+			{
 #if WITH_RECAST
-			const ARecastNavMesh* RecastNavMesh = Cast<const ARecastNavMesh>(NavData);
-			bHasPath = RecastNavMesh && RecastNavMesh->IsSegmentOnNavmesh(PointA, PointB, QueryFilter);
+				const ARecastNavMesh* RecastNavMesh = Cast<const ARecastNavMesh>(NavData);
+				bHasPath = RecastNavMesh && RecastNavMesh->IsSegmentOnNavmesh(PointA, PointB, QueryFilter);
 #endif
-		}
-		else
-		{		
-			EPathFindingMode::Type TestMode = (PathQueryType == EPathExistanceQueryType::HierarchicalQuery) ? EPathFindingMode::Hierarchical : EPathFindingMode::Regular;
-			bHasPath = NavSys->TestPathSync(FPathFindingQuery(AIOwner, NavData, PointA, PointB, QueryFilter), TestMode);
+			}
+			else
+			{
+				EPathFindingMode::Type TestMode = (PathQueryType == EPathExistanceQueryType::HierarchicalQuery) ? EPathFindingMode::Hierarchical : EPathFindingMode::Regular;
+				bHasPath = NavSys->TestPathSync(FPathFindingQuery(AIOwner, *NavData, PointA, PointB, QueryFilter), TestMode);
+			}
 		}
 	}
 
diff --git a/Engine/Source/Runtime/AIModule/Private/EnvironmentQuery/EnvQueryTraceHelpers.cpp b/Engine/Source/Runtime/AIModule/Private/EnvironmentQuery/EnvQueryTraceHelpers.cpp
--- a/Engine/Source/Runtime/AIModule/Private/EnvironmentQuery/EnvQueryTraceHelpers.cpp
+++ b/Engine/Source/Runtime/AIModule/Private/EnvironmentQuery/EnvQueryTraceHelpers.cpp
@@ -156,7 +156,7 @@ void FEQSHelpers::FBatchTrace::DoProject<EEnvTraceShape::Capsule>(TArray<FNavLoc
 }
 
 
-void FEQSHelpers::RunNavRaycasts(ANavigationData* NavData, const FEnvTraceData& TraceData, const FVector& SourcePt, TArray<FNavLocation>& Points, ETraceMode TraceMode)
+void FEQSHelpers::RunNavRaycasts(const ANavigationData& NavData, const FEnvTraceData& TraceData, const FVector& SourcePt, TArray<FNavLocation>& Points, ETraceMode TraceMode /*= ETraceMode::Keep*/)
 {
 	TSharedPtr<const FNavigationQueryFilter> NavigationFilter = UNavigationQueryFilter::GetQueryFilter(NavData, TraceData.NavigationFilter);
 
@@ -168,7 +168,7 @@ void FEQSHelpers::RunNavRaycasts(ANavigationData* NavData, const FEnvTraceData&
 		RaycastWorkload.Add(FNavigationRaycastWork(SourcePt, ItemLocation.Location));
 	}
 
-	NavData->BatchRaycast(RaycastWorkload, NavigationFilter);
+	NavData.BatchRaycast(RaycastWorkload, NavigationFilter);
 
 	for (int32 Idx = 0; Idx < Points.Num(); Idx++)
 	{
@@ -188,7 +188,7 @@ void FEQSHelpers::RunNavRaycasts(ANavigationData* NavData, const FEnvTraceData&
 }
 
 
-void FEQSHelpers::RunNavProjection(ANavigationData* NavData, const FEnvTraceData& TraceData, TArray<FNavLocation>& Points, ETraceMode TraceMode)
+void FEQSHelpers::RunNavProjection(const ANavigationData& NavData, const FEnvTraceData& TraceData, TArray<FNavLocation>& Points, ETraceMode TraceMode /*= ETraceMode::Discard*/)
 {
 	TSharedPtr<const FNavigationQueryFilter> NavigationFilter = UNavigationQueryFilter::GetQueryFilter(NavData, TraceData.NavigationFilter);
 	TArray<FNavigationProjectionWork> Workload;
@@ -211,7 +211,7 @@ void FEQSHelpers::RunNavProjection(ANavigationData* NavData, const FEnvTraceData
 	}
 
 	const FVector ProjectionExtent(TraceData.ExtentX, TraceData.ExtentX, (TraceData.ProjectDown + TraceData.ProjectUp) / 2);
-	NavData->BatchProjectPoints(Workload, ProjectionExtent, NavigationFilter);
+	NavData.BatchProjectPoints(Workload, ProjectionExtent, NavigationFilter);
 
 	for (int32 Idx = Workload.Num() - 1; Idx >= 0; Idx--)
 	{
diff --git a/Engine/Source/Runtime/AIModule/Private/EnvironmentQuery/EnvQueryTraceHelpers.h b/Engine/Source/Runtime/AIModule/Private/EnvironmentQuery/EnvQueryTraceHelpers.h
--- a/Engine/Source/Runtime/AIModule/Private/EnvironmentQuery/EnvQueryTraceHelpers.h
+++ b/Engine/Source/Runtime/AIModule/Private/EnvironmentQuery/EnvQueryTraceHelpers.h
@@ -94,8 +94,29 @@ namespace FEQSHelpers
 	template<>
 	void FBatchTrace::DoProject<EEnvTraceShape::Capsule>(TArray<FNavLocation>& Points, float StartOffsetZ, float EndOffsetZ, float HitOffsetZ);
 
-	void RunNavRaycasts(ANavigationData* NavData, const FEnvTraceData& TraceData, const FVector& SourcePt, TArray<FNavLocation>& Points, ETraceMode TraceMode = ETraceMode::Keep);
-	void RunNavProjection(ANavigationData* NavData, const FEnvTraceData& TraceData, TArray<FNavLocation>& Points, ETraceMode TraceMode = ETraceMode::Discard);
+	void RunNavRaycasts(const ANavigationData& NavData, const FEnvTraceData& TraceData, const FVector& SourcePt, TArray<FNavLocation>& Points, ETraceMode TraceMode = ETraceMode::Keep);
+	void RunNavProjection(const ANavigationData& NavData, const FEnvTraceData& TraceData, TArray<FNavLocation>& Points, ETraceMode TraceMode = ETraceMode::Discard);
 	void RunPhysRaycasts(UWorld* World, const FEnvTraceData& TraceData, const FVector& SourcePt, TArray<FNavLocation>& Points, ETraceMode TraceMode = ETraceMode::Keep);
 	void RunPhysProjection(UWorld* World, const FEnvTraceData& TraceData, TArray<FNavLocation>& Points, ETraceMode TraceMode = ETraceMode::Discard);
+
+	//----------------------------------------------------------------------//
+	// deprecates
+	//----------------------------------------------------------------------//
+	DEPRECATED(4.8, "This version of RunNavRaycasts is deprecated. Please use ANavigationData reference rather than a pointer version")
+	void RunNavRaycasts(const ANavigationData* NavData, const FEnvTraceData& TraceData, const FVector& SourcePt, TArray<FNavLocation>& Points, ETraceMode TraceMode = ETraceMode::Keep)
+	{
+		if (NavData)
+		{
+			RunNavRaycasts(*NavData, TraceData, SourcePt, Points, TraceMode);
+		}
+	}
+
+	DEPRECATED(4.8, "This version of RunNavProjection is deprecated. Please use ANavigationData reference rather than a pointer version")
+	void RunNavProjection(const ANavigationData* NavData, const FEnvTraceData& TraceData, TArray<FNavLocation>& Points, ETraceMode TraceMode = ETraceMode::Discard)
+	{
+		if (NavData)
+		{
+			RunNavProjection(*NavData, TraceData, Points, TraceMode);
+		}
+	}
 }
diff --git a/Engine/Source/Runtime/AIModule/Private/EnvironmentQuery/EnvQueryTypes.cpp b/Engine/Source/Runtime/AIModule/Private/EnvironmentQuery/EnvQueryTypes.cpp
--- a/Engine/Source/Runtime/AIModule/Private/EnvironmentQuery/EnvQueryTypes.cpp
+++ b/Engine/Source/Runtime/AIModule/Private/EnvironmentQuery/EnvQueryTypes.cpp
@@ -274,7 +274,7 @@ void FEnvFloatParam_DEPRECATED::Convert(UObject* Owner, FAIDataProviderFloatValu
 //----------------------------------------------------------------------//
 namespace FEQSHelpers
 {
-	const ARecastNavMesh* FindNavMeshForQuery(FEnvQueryInstance& QueryInstance)
+	const ANavigationData* FindNavigationDataForQuery(FEnvQueryInstance& QueryInstance)
 	{
 		const UNavigationSystem* NavSys = QueryInstance.World->GetNavigationSystem();
 		
@@ -288,11 +288,26 @@ namespace FEQSHelpers
 		if (NavAgent)
 		{
 			const FNavAgentProperties& NavAgentProps = NavAgent ? NavAgent->GetNavAgentPropertiesRef() : FNavAgentProperties::DefaultProperties;
-			return Cast<const ARecastNavMesh>(NavSys->GetNavDataForProps(NavAgentProps));
+			return NavSys->GetNavDataForProps(NavAgentProps);
 		}
 
-		return Cast<const ARecastNavMesh>(NavSys->GetMainNavData());
+		return NavSys->GetMainNavData();
+	}
+}
+
+//----------------------------------------------------------------------//
+// DEPRECATED
+//----------------------------------------------------------------------//
+#if WITH_RECAST
+
+namespace FEQSHelpers
+{
+	const ARecastNavMesh* FindNavMeshForQuery(FEnvQueryInstance& QueryInstance)
+	{
+		return Cast<const ARecastNavMesh>(FEQSHelpers::FindNavigationDataForQuery(QueryInstance));
 	}
 }
 
+#endif // WITH_RECAST
+
 #undef LOCTEXT_NAMESPACE
diff --git a/Engine/Source/Runtime/AIModule/Private/EnvironmentQuery/Generators/EnvQueryGenerator_OnCircle.cpp b/Engine/Source/Runtime/AIModule/Private/EnvironmentQuery/Generators/EnvQueryGenerator_OnCircle.cpp
--- a/Engine/Source/Runtime/AIModule/Private/EnvironmentQuery/Generators/EnvQueryGenerator_OnCircle.cpp
+++ b/Engine/Source/Runtime/AIModule/Private/EnvironmentQuery/Generators/EnvQueryGenerator_OnCircle.cpp
@@ -155,13 +155,10 @@ void UEnvQueryGenerator_OnCircle::GenerateItemsForCircle(uint8* ContextRawData,
 	{
 		case EEnvQueryTrace::Navigation:
 		{
-			ANavigationData* NavData = nullptr;
-#if WITH_RECAST
-			NavData = (ANavigationData*)FEQSHelpers::FindNavMeshForQuery(OutQueryInstance);
-#endif
+			ANavigationData* NavData = const_cast<ANavigationData*>(FEQSHelpers::FindNavigationDataForQuery(OutQueryInstance));
 			if (NavData)
 			{
-				FEQSHelpers::RunNavRaycasts(NavData, TraceData, CenterLocation, ItemCandidates);
+				FEQSHelpers::RunNavRaycasts(*NavData, TraceData, CenterLocation, ItemCandidates);
 			}
 		}
 			break;
diff --git a/Engine/Source/Runtime/AIModule/Private/EnvironmentQuery/Generators/EnvQueryGenerator_PathingGrid.cpp b/Engine/Source/Runtime/AIModule/Private/EnvironmentQuery/Generators/EnvQueryGenerator_PathingGrid.cpp
--- a/Engine/Source/Runtime/AIModule/Private/EnvironmentQuery/Generators/EnvQueryGenerator_PathingGrid.cpp
+++ b/Engine/Source/Runtime/AIModule/Private/EnvironmentQuery/Generators/EnvQueryGenerator_PathingGrid.cpp
@@ -38,7 +38,7 @@ void UEnvQueryGenerator_PathingGrid::ProjectAndFilterNavPoints(TArray<FNavLocati
 	bool bPathToItem = PathToItem.GetValue();
 	float RangeMultiplierValue = ScanRangeMultiplier.GetValue();
 
-	ARecastNavMesh* NavMeshData = (ARecastNavMesh*)FEQSHelpers::FindNavMeshForQuery(QueryInstance);
+	ARecastNavMesh* NavMeshData = const_cast<ARecastNavMesh*>(static_cast<const ARecastNavMesh*>(FEQSHelpers::FindNavigationDataForQuery(QueryInstance)));
 	if (!NavMeshData)
 	{
 		return;
@@ -48,7 +48,7 @@ void UEnvQueryGenerator_PathingGrid::ProjectAndFilterNavPoints(TArray<FNavLocati
 	QueryInstance.PrepareContext(GenerateAround, ContextLocations);
 
 	TSharedPtr<FNavigationQueryFilter> NavigationFilterOb = (NavigationFilter != nullptr)
-		? UNavigationQueryFilter::GetQueryFilter(NavMeshData, NavigationFilter)->GetCopy()
+		? UNavigationQueryFilter::GetQueryFilter(*NavMeshData, NavigationFilter)->GetCopy()
 		: NavMeshData->GetDefaultQueryFilter()->GetCopy();
 	NavigationFilterOb->SetBacktrackingEnabled(!bPathToItem);
 	
diff --git a/Engine/Source/Runtime/AIModule/Private/EnvironmentQuery/Generators/EnvQueryGenerator_ProjectedPoints.cpp b/Engine/Source/Runtime/AIModule/Private/EnvironmentQuery/Generators/EnvQueryGenerator_ProjectedPoints.cpp
--- a/Engine/Source/Runtime/AIModule/Private/EnvironmentQuery/Generators/EnvQueryGenerator_ProjectedPoints.cpp
+++ b/Engine/Source/Runtime/AIModule/Private/EnvironmentQuery/Generators/EnvQueryGenerator_ProjectedPoints.cpp
@@ -17,18 +17,16 @@ UEnvQueryGenerator_ProjectedPoints::UEnvQueryGenerator_ProjectedPoints(const FOb
 
 void UEnvQueryGenerator_ProjectedPoints::ProjectAndFilterNavPoints(TArray<FNavLocation>& Points, FEnvQueryInstance& QueryInstance) const
 {
-	ANavigationData* NavData = nullptr;
+	const ANavigationData* NavData = nullptr;
 	if (ProjectionData.TraceMode != EEnvQueryTrace::None)
 	{
-#if WITH_RECAST
-		NavData = (ANavigationData*)FEQSHelpers::FindNavMeshForQuery(QueryInstance);
-#endif // WITH_RECAST
+		NavData = FEQSHelpers::FindNavigationDataForQuery(QueryInstance);
 	}
 
 	if (NavData)
 	{
 		FEQSHelpers::ETraceMode Mode = (ProjectionData.TraceMode == EEnvQueryTrace::Navigation) ? FEQSHelpers::ETraceMode::Discard : FEQSHelpers::ETraceMode::Keep;
-		FEQSHelpers::RunNavProjection(NavData, ProjectionData, Points, Mode);
+		FEQSHelpers::RunNavProjection(*NavData, ProjectionData, Points, Mode);
 	}
 
 	if (ProjectionData.TraceMode == EEnvQueryTrace::Geometry)
@@ -39,7 +37,8 @@ void UEnvQueryGenerator_ProjectedPoints::ProjectAndFilterNavPoints(TArray<FNavLo
 
 void UEnvQueryGenerator_ProjectedPoints::StoreNavPoints(const TArray<FNavLocation>& Points, FEnvQueryInstance& QueryInstance) const
 {
-	QueryInstance.ReserveItemData(QueryInstance.Items.Num() + Points.Num());
+	const int32 InitialElementsCount = QueryInstance.Items.Num();
+	QueryInstance.ReserveItemData(InitialElementsCount + Points.Num());
 	for (int32 Idx = 0; Idx < Points.Num(); Idx++)
 	{
 		// store using default function to handle creating new data entry 
@@ -50,7 +49,7 @@ void UEnvQueryGenerator_ProjectedPoints::StoreNavPoints(const TArray<FNavLocatio
 	for (int32 Idx = 0; Idx < Points.Num(); Idx++)
 	{
 		// overwrite with more detailed info
-		UEnvQueryItemType_Point::SetNavValue(DataPtr + QueryInstance.Items[Idx].DataOffset, Points[Idx]);
+		UEnvQueryItemType_Point::SetNavValue(DataPtr + QueryInstance.Items[Idx + InitialElementsCount].DataOffset, Points[Idx]);
 	}
 
 	FEnvQueryOptionInstance& OptionInstance = QueryInstance.Options[QueryInstance.OptionIndex];
diff --git a/Engine/Source/Runtime/AIModule/Private/EnvironmentQuery/Tests/EnvQueryTest_Pathfinding.cpp b/Engine/Source/Runtime/AIModule/Private/EnvironmentQuery/Tests/EnvQueryTest_Pathfinding.cpp
--- a/Engine/Source/Runtime/AIModule/Private/EnvironmentQuery/Tests/EnvQueryTest_Pathfinding.cpp
+++ b/Engine/Source/Runtime/AIModule/Private/EnvironmentQuery/Tests/EnvQueryTest_Pathfinding.cpp
@@ -39,8 +39,12 @@ void UEnvQueryTest_Pathfinding::RunTest(FEnvQueryInstance& QueryInstance) const
 	float MaxThresholdValue = FloatValueMax.GetValue();
 
 	UNavigationSystem* NavSys = QueryInstance.World->GetNavigationSystem();
-	ANavigationData* NavData = FindNavigationData(NavSys, QueryInstance.Owner.Get());
-	if (!NavData)
+	if (NavSys == nullptr)
+	{
+		return;
+	}
+	ANavigationData* NavData = FindNavigationData(*NavSys, QueryInstance.Owner.Get());
+	if (NavData == nullptr)
 	{
 		return;
 	}
@@ -66,7 +70,7 @@ void UEnvQueryTest_Pathfinding::RunTest(FEnvQueryInstance& QueryInstance) const
 			const FVector ItemLocation = GetItemLocation(QueryInstance, *It);
 			for (int32 ContextIndex = 0; ContextIndex < ContextLocations.Num(); ContextIndex++)
 			{
-				const float PathValue = FindPathFunc.Execute(ItemLocation, ContextLocations[ContextIndex], PFMode, NavData, NavSys, QueryInstance.Owner.Get());
+				const float PathValue = FindPathFunc.Execute(ItemLocation, ContextLocations[ContextIndex], PFMode, *NavData, *NavSys, QueryInstance.Owner.Get());
 				It.SetScore(TestPurpose, FilterType, PathValue, MinThresholdValue, MaxThresholdValue);
 
 				if (bDiscardFailed && PathValue >= BIG_NUMBER)
@@ -87,7 +91,7 @@ void UEnvQueryTest_Pathfinding::RunTest(FEnvQueryInstance& QueryInstance) const
 				const FVector ItemLocation = GetItemLocation(QueryInstance, *It);
 				for (int32 ContextIndex = 0; ContextIndex < ContextLocations.Num(); ContextIndex++)
 				{
-					const bool bFoundPath = TestPathTo(ItemLocation, ContextLocations[ContextIndex], PFMode, NavData, NavSys, QueryInstance.Owner.Get());
+					const bool bFoundPath = TestPathTo(ItemLocation, ContextLocations[ContextIndex], PFMode, *NavData, *NavSys, QueryInstance.Owner.Get());
 					It.SetScore(TestPurpose, FilterType, bFoundPath, bWantsPath);
 				}
 			}
@@ -99,7 +103,7 @@ void UEnvQueryTest_Pathfinding::RunTest(FEnvQueryInstance& QueryInstance) const
 				const FVector ItemLocation = GetItemLocation(QueryInstance, *It);
 				for (int32 ContextIndex = 0; ContextIndex < ContextLocations.Num(); ContextIndex++)
 				{
-					const bool bFoundPath = TestPathFrom(ItemLocation, ContextLocations[ContextIndex], PFMode, NavData, NavSys, QueryInstance.Owner.Get());
+					const bool bFoundPath = TestPathFrom(ItemLocation, ContextLocations[ContextIndex], PFMode, *NavData, *NavSys, QueryInstance.Owner.Get());
 					It.SetScore(TestPurpose, FilterType, bFoundPath, bWantsPath);
 				}
 			}
@@ -159,69 +163,69 @@ void UEnvQueryTest_Pathfinding::PostLoad()
 	SetWorkOnFloatValues(TestMode != EEnvTestPathfinding::PathExist);
 }
 
-bool UEnvQueryTest_Pathfinding::TestPathFrom(const FVector& ItemPos, const FVector& ContextPos, EPathFindingMode::Type Mode, const ANavigationData* NavData, UNavigationSystem* NavSys, const UObject* PathOwner) const
+bool UEnvQueryTest_Pathfinding::TestPathFrom(const FVector& ItemPos, const FVector& ContextPos, EPathFindingMode::Type Mode, const ANavigationData& NavData, UNavigationSystem& NavSys, const UObject* PathOwner) const
 {
 	FPathFindingQuery Query(PathOwner, NavData, ItemPos, ContextPos, UNavigationQueryFilter::GetQueryFilter(NavData, FilterClass));
 	Query.SetAllowPartialPaths(false);
 
-	const bool bPathExists = NavSys->TestPathSync(Query, Mode);
+	const bool bPathExists = NavSys.TestPathSync(Query, Mode);
 	return bPathExists;
 }
 
-bool UEnvQueryTest_Pathfinding::TestPathTo(const FVector& ItemPos, const FVector& ContextPos, EPathFindingMode::Type Mode, const ANavigationData* NavData, UNavigationSystem* NavSys, const UObject* PathOwner) const
+bool UEnvQueryTest_Pathfinding::TestPathTo(const FVector& ItemPos, const FVector& ContextPos, EPathFindingMode::Type Mode, const ANavigationData& NavData, UNavigationSystem& NavSys, const UObject* PathOwner) const
 {
 	FPathFindingQuery Query(PathOwner, NavData, ContextPos, ItemPos, UNavigationQueryFilter::GetQueryFilter(NavData, FilterClass));
 	Query.SetAllowPartialPaths(false);
 
-	const bool bPathExists = NavSys->TestPathSync(Query, Mode);
+	const bool bPathExists = NavSys.TestPathSync(Query, Mode);
 	return bPathExists;
 }
 
-float UEnvQueryTest_Pathfinding::FindPathCostFrom(const FVector& ItemPos, const FVector& ContextPos, EPathFindingMode::Type Mode, const ANavigationData* NavData, UNavigationSystem* NavSys, const UObject* PathOwner) const
+float UEnvQueryTest_Pathfinding::FindPathCostFrom(const FVector& ItemPos, const FVector& ContextPos, EPathFindingMode::Type Mode, const ANavigationData& NavData, UNavigationSystem& NavSys, const UObject* PathOwner) const
 {
 	FPathFindingQuery Query(PathOwner, NavData, ItemPos, ContextPos, UNavigationQueryFilter::GetQueryFilter(NavData, FilterClass));
 	Query.SetAllowPartialPaths(false);
 
-	FPathFindingResult Result = NavSys->FindPathSync(Query, Mode);
+	FPathFindingResult Result = NavSys.FindPathSync(Query, Mode);
 	return (Result.IsSuccessful()) ? Result.Path->GetCost() : BIG_NUMBER;
 }
 
-float UEnvQueryTest_Pathfinding::FindPathCostTo(const FVector& ItemPos, const FVector& ContextPos, EPathFindingMode::Type Mode, const ANavigationData* NavData, UNavigationSystem* NavSys, const UObject* PathOwner) const
+float UEnvQueryTest_Pathfinding::FindPathCostTo(const FVector& ItemPos, const FVector& ContextPos, EPathFindingMode::Type Mode, const ANavigationData& NavData, UNavigationSystem& NavSys, const UObject* PathOwner) const
 {
 	FPathFindingQuery Query(PathOwner, NavData, ContextPos, ItemPos, UNavigationQueryFilter::GetQueryFilter(NavData, FilterClass));
 	Query.SetAllowPartialPaths(false);
 
-	FPathFindingResult Result = NavSys->FindPathSync(Query, Mode);
+	FPathFindingResult Result = NavSys.FindPathSync(Query, Mode);
 	return (Result.IsSuccessful()) ? Result.Path->GetCost() : BIG_NUMBER;
 }
 
-float UEnvQueryTest_Pathfinding::FindPathLengthFrom(const FVector& ItemPos, const FVector& ContextPos, EPathFindingMode::Type Mode, const ANavigationData* NavData, UNavigationSystem* NavSys, const UObject* PathOwner) const
+float UEnvQueryTest_Pathfinding::FindPathLengthFrom(const FVector& ItemPos, const FVector& ContextPos, EPathFindingMode::Type Mode, const ANavigationData& NavData, UNavigationSystem& NavSys, const UObject* PathOwner) const
 {
 	FPathFindingQuery Query(PathOwner, NavData, ItemPos, ContextPos, UNavigationQueryFilter::GetQueryFilter(NavData, FilterClass));
 	Query.SetAllowPartialPaths(false);
 
-	FPathFindingResult Result = NavSys->FindPathSync(Query, Mode);
+	FPathFindingResult Result = NavSys.FindPathSync(Query, Mode);
 	return (Result.IsSuccessful()) ? Result.Path->GetLength() : BIG_NUMBER;
 }
 
-float UEnvQueryTest_Pathfinding::FindPathLengthTo(const FVector& ItemPos, const FVector& ContextPos, EPathFindingMode::Type Mode, const ANavigationData* NavData, UNavigationSystem* NavSys, const UObject* PathOwner) const
+float UEnvQueryTest_Pathfinding::FindPathLengthTo(const FVector& ItemPos, const FVector& ContextPos, EPathFindingMode::Type Mode, const ANavigationData& NavData, UNavigationSystem& NavSys, const UObject* PathOwner) const
 {
 	FPathFindingQuery Query(PathOwner, NavData, ContextPos, ItemPos, UNavigationQueryFilter::GetQueryFilter(NavData, FilterClass));
 	Query.SetAllowPartialPaths(false);
 
-	FPathFindingResult Result = NavSys->FindPathSync(Query, Mode);
+	FPathFindingResult Result = NavSys.FindPathSync(Query, Mode);
 	return (Result.IsSuccessful()) ? Result.Path->GetLength() : BIG_NUMBER;
 }
 
-ANavigationData* UEnvQueryTest_Pathfinding::FindNavigationData(UNavigationSystem* NavSys, UObject* Owner) const
+ANavigationData* UEnvQueryTest_Pathfinding::FindNavigationData(UNavigationSystem& NavSys, UObject* Owner) const
 {
 	INavAgentInterface* NavAgent = Cast<INavAgentInterface>(Owner);
 	if (NavAgent)
 	{
-		return NavSys->GetNavDataForProps(NavAgent->GetNavAgentPropertiesRef());
+		return NavSys.GetNavDataForProps(NavAgent->GetNavAgentPropertiesRef());
 	}
 
-	return NavSys->GetMainNavData(FNavigationSystem::DontCreate);
+	return NavSys.GetMainNavData(FNavigationSystem::DontCreate);
 }
 
 #undef LOCTEXT_NAMESPACE
diff --git a/Engine/Source/Runtime/AIModule/Private/EnvironmentQuery/Tests/EnvQueryTest_PathfindingBatch.cpp b/Engine/Source/Runtime/AIModule/Private/EnvironmentQuery/Tests/EnvQueryTest_PathfindingBatch.cpp
--- a/Engine/Source/Runtime/AIModule/Private/EnvironmentQuery/Tests/EnvQueryTest_PathfindingBatch.cpp
+++ b/Engine/Source/Runtime/AIModule/Private/EnvironmentQuery/Tests/EnvQueryTest_PathfindingBatch.cpp
@@ -75,9 +75,13 @@ void UEnvQueryTest_PathfindingBatch::RunTest(FEnvQueryInstance& QueryInstance) c
 	float RangeMultiplierValue = ScanRangeMultiplier.GetValue();
 
 	UNavigationSystem* NavSys = QueryInstance.World->GetNavigationSystem();
-	ANavigationData* NavData = FindNavigationData(NavSys, QueryInstance.Owner.Get());
+	if (NavSys == nullptr)
+	{
+		return;
+	}
+	ANavigationData* NavData = FindNavigationData(*NavSys, QueryInstance.Owner.Get());
 	ARecastNavMesh* NavMeshData = Cast<ARecastNavMesh>(NavData);
-	if (!NavMeshData)
+	if (NavMeshData == nullptr)
 	{
 		return;
 	}
@@ -93,7 +97,7 @@ void UEnvQueryTest_PathfindingBatch::RunTest(FEnvQueryInstance& QueryInstance) c
 	CollectDistanceSq.Init(0.0f, ContextLocations.Num());
 
 	TSharedPtr<FNavigationQueryFilter> NavigationFilter = FilterClass != nullptr
-		? UNavigationQueryFilter::GetQueryFilter(NavMeshData, FilterClass)->GetCopy()
+		? UNavigationQueryFilter::GetQueryFilter(*NavMeshData, FilterClass)->GetCopy()
 		: NavMeshData->GetDefaultQueryFilter()->GetCopy();
 	NavigationFilter->SetBacktrackingEnabled(!bPathToItem);
 	const dtQueryFilter* NavQueryFilter = ((const FRecastQueryFilter*)NavigationFilter->GetImplementation())->GetAsDetourQueryFilter();
diff --git a/Engine/Source/Runtime/Engine/Classes/AI/Navigation/NavFilters/NavigationQueryFilter.h b/Engine/Source/Runtime/Engine/Classes/AI/Navigation/NavFilters/NavigationQueryFilter.h
--- a/Engine/Source/Runtime/Engine/Classes/AI/Navigation/NavFilters/NavigationQueryFilter.h
+++ b/Engine/Source/Runtime/Engine/Classes/AI/Navigation/NavFilters/NavigationQueryFilter.h
@@ -216,13 +216,13 @@ class ENGINE_API UNavigationQueryFilter : public UObject
 	FNavigationFilterFlags ExcludeFlags;
 
 	/** get filter for given navigation data and initialize on first access */
-	TSharedPtr<const FNavigationQueryFilter> GetQueryFilter(const ANavigationData* NavData) const;
+	TSharedPtr<const FNavigationQueryFilter> GetQueryFilter(const ANavigationData& NavData) const;
 	
 	/** helper functions for accessing filter */
-	static TSharedPtr<const FNavigationQueryFilter> GetQueryFilter(const ANavigationData* NavData, UClass* FilterClass);
+	static TSharedPtr<const FNavigationQueryFilter> GetQueryFilter(const ANavigationData& NavData, TSubclassOf<UNavigationQueryFilter> FilterClass);
 
 	template<class T>
-	static TSharedPtr<const FNavigationQueryFilter> GetQueryFilter(const ANavigationData* NavData, UClass* FilterClass = T::StaticClass())
+	static TSharedPtr<const FNavigationQueryFilter> GetQueryFilter(const ANavigationData& NavData, TSubclassOf<UNavigationQueryFilter> FilterClass = T::StaticClass())
 	{
 		return GetQueryFilter(NavData, FilterClass);
 	}
@@ -240,7 +240,40 @@ protected:
 
 	/** find index of area data */
 	int32 FindAreaOverride(TSubclassOf<UNavArea> AreaClass) const;
-
+	
 	/** setup filter for given navigation data, use to create custom filters */
-	virtual void InitializeFilter(const ANavigationData* NavData, FNavigationQueryFilter* Filter) const;
+	virtual void InitializeFilter(const ANavigationData& NavData, FNavigationQueryFilter& Filter) const;
+
+public:
+	//----------------------------------------------------------------------//
+	// deprecated
+	//----------------------------------------------------------------------//
+
+	DEPRECATED(4.8, "This version of InitializeFilter is deprecated. Please use the version taking NavData and Filter references")
+	virtual void InitializeFilter(const ANavigationData* NavData, FNavigationQueryFilter* Filter) const
+	{
+		if (NavData && Filter)
+		{
+			InitializeFilter(*NavData, *Filter);
+		}
+	}
+
+	DEPRECATED(4.8, "This version of GetQueryFilter is deprecated. Please use ANavigationData reference rather than a pointer version")
+	TSharedPtr<const FNavigationQueryFilter> GetQueryFilter(const ANavigationData* NavData) const 
+	{ 
+		return NavData ? GetQueryFilter(*NavData) : nullptr; 
+	}
+
+	DEPRECATED(4.8, "This version of GetQueryFilter is deprecated. Please use ANavigationData reference rather than a pointer version")
+	static TSharedPtr<const FNavigationQueryFilter> GetQueryFilter(const ANavigationData* NavData, TSubclassOf<UNavigationQueryFilter> FilterClass)
+	{
+		return NavData ? GetQueryFilter(*NavData, FilterClass) : nullptr;
+	}
+
+	// will "auto-deprecate" due to the function called inside
+	template<class T>
+	static TSharedPtr<const FNavigationQueryFilter> GetQueryFilter(const ANavigationData* NavData, TSubclassOf<UNavigationQueryFilter> FilterClass = T::StaticClass())
+	{
+		return NavData ? GetQueryFilter(*NavData, FilterClass) : nullptr;
+	}
 };
diff --git a/Engine/Source/Runtime/Engine/Classes/AI/Navigation/NavFilters/RecastFilter_UseDefaultArea.h b/Engine/Source/Runtime/Engine/Classes/AI/Navigation/NavFilters/RecastFilter_UseDefaultArea.h
--- a/Engine/Source/Runtime/Engine/Classes/AI/Navigation/NavFilters/RecastFilter_UseDefaultArea.h
+++ b/Engine/Source/Runtime/Engine/Classes/AI/Navigation/NavFilters/RecastFilter_UseDefaultArea.h
@@ -9,5 +9,5 @@ class URecastFilter_UseDefaultArea : public UNavigationQueryFilter
 {
 	GENERATED_UCLASS_BODY()
 
-	virtual void InitializeFilter(const ANavigationData* NavData, FNavigationQueryFilter* Filter) const override;
+	virtual void InitializeFilter(const ANavigationData& NavData, FNavigationQueryFilter& Filter) const override;
 };
diff --git a/Engine/Source/Runtime/Engine/Classes/AI/Navigation/NavigationData.h b/Engine/Source/Runtime/Engine/Classes/AI/Navigation/NavigationData.h
--- a/Engine/Source/Runtime/Engine/Classes/AI/Navigation/NavigationData.h
+++ b/Engine/Source/Runtime/Engine/Classes/AI/Navigation/NavigationData.h
@@ -866,7 +866,12 @@ struct FAsyncPathFindingQuery : public FPathFindingQuery
 		: QueryID(INVALID_NAVQUERYID)
 	{ }
 
-	FAsyncPathFindingQuery(const UObject* InOwner, const ANavigationData* InNavData, const FVector& Start, const FVector& End, const FNavPathQueryDelegate& Delegate, TSharedPtr<const FNavigationQueryFilter> SourceQueryFilter);
+	DEPRECATED(4.8, "This version of FAsyncPathFindingQuery constructor is deprecated. Please use ANavigationData reference rather than a pointer version")
+	FAsyncPathFindingQuery(const UObject* InOwner, const ANavigationData* InNavData, const FVector& Start, const FVector& End, const FNavPathQueryDelegate& Delegate, TSharedPtr<const FNavigationQueryFilter> SourceQueryFilter)
+	: FAsyncPathFindingQuery(InOwner, *InNavData, Start, End, Delegate, SourceQueryFilter)
+	{}
+
+	FAsyncPathFindingQuery(const UObject* InOwner, const ANavigationData& InNavData, const FVector& Start, const FVector& End, const FNavPathQueryDelegate& Delegate, TSharedPtr<const FNavigationQueryFilter> SourceQueryFilter);
 	FAsyncPathFindingQuery(const FPathFindingQuery& Query, const FNavPathQueryDelegate& Delegate, const EPathFindingMode::Type QueryMode);
 
 protected:
diff --git a/Engine/Source/Runtime/Engine/Classes/AI/Navigation/NavigationSystem.h b/Engine/Source/Runtime/Engine/Classes/AI/Navigation/NavigationSystem.h
--- a/Engine/Source/Runtime/Engine/Classes/AI/Navigation/NavigationSystem.h
+++ b/Engine/Source/Runtime/Engine/Classes/AI/Navigation/NavigationSystem.h
@@ -454,6 +454,7 @@ public:
 	ANavigationData* GetMainNavData(FNavigationSystem::ECreateIfEmpty CreateNewIfNoneFound);
 	/** Returns the world nav mesh object.  Creates one if it doesn't exist. */
 	const ANavigationData* GetMainNavData() const { return MainNavData; }
+	const ANavigationData& GetMainNavDataChecked() const { check(MainNavData); return *MainNavData; }
 
 	ANavigationData* GetAbstractNavData() const { return AbstractNavData; }
 
diff --git a/Engine/Source/Runtime/Engine/Classes/AI/Navigation/NavigationTypes.h b/Engine/Source/Runtime/Engine/Classes/AI/Navigation/NavigationTypes.h
--- a/Engine/Source/Runtime/Engine/Classes/AI/Navigation/NavigationTypes.h
+++ b/Engine/Source/Runtime/Engine/Classes/AI/Navigation/NavigationTypes.h
@@ -483,7 +483,12 @@ struct ENGINE_API FPathFindingQuery
 
 	FPathFindingQuery(const FPathFindingQuery& Source);
 
-	FPathFindingQuery(const UObject* InOwner, const ANavigationData* InNavData, const FVector& Start, const FVector& End, TSharedPtr<const FNavigationQueryFilter> SourceQueryFilter = NULL, FNavPathSharedPtr InPathInstanceToFill = NULL);
+	DEPRECATED(4.8, "This version of FPathFindingQuery's constructor is deprecated. Please use ANavigationData reference rather than a pointer")
+	FPathFindingQuery(const UObject* InOwner, const ANavigationData* InNavData, const FVector& Start, const FVector& End, TSharedPtr<const FNavigationQueryFilter> SourceQueryFilter = NULL, FNavPathSharedPtr InPathInstanceToFill = NULL)
+	: FPathFindingQuery(InOwner, *InNavData, Start, End, SourceQueryFilter, InPathInstanceToFill)
+	{}
+
+	FPathFindingQuery(const UObject* InOwner, const ANavigationData& InNavData, const FVector& Start, const FVector& End, TSharedPtr<const FNavigationQueryFilter> SourceQueryFilter = NULL, FNavPathSharedPtr InPathInstanceToFill = NULL);
 
 	explicit FPathFindingQuery(FNavPathSharedRef PathToRecalculate, const ANavigationData* NavDataOverride = NULL);
 
diff --git a/Engine/Source/Runtime/Engine/Private/AI/Navigation/NavFilters/NavigationQueryFilter.cpp b/Engine/Source/Runtime/Engine/Private/AI/Navigation/NavFilters/NavigationQueryFilter.cpp
--- a/Engine/Source/Runtime/Engine/Private/AI/Navigation/NavFilters/NavigationQueryFilter.cpp
+++ b/Engine/Source/Runtime/Engine/Private/AI/Navigation/NavFilters/NavigationQueryFilter.cpp
@@ -115,31 +115,31 @@ UNavigationQueryFilter::UNavigationQueryFilter(const FObjectInitializer& ObjectI
 	ExcludeFlags.Packed = 0;
 }
 
-TSharedPtr<const FNavigationQueryFilter> UNavigationQueryFilter::GetQueryFilter(const ANavigationData* NavData) const
+TSharedPtr<const FNavigationQueryFilter> UNavigationQueryFilter::GetQueryFilter(const ANavigationData& NavData) const
 {
-	TSharedPtr<const FNavigationQueryFilter> SharedFilter = NavData->GetQueryFilter(GetClass());
+	TSharedPtr<const FNavigationQueryFilter> SharedFilter = NavData.GetQueryFilter(GetClass());
 	if (!SharedFilter.IsValid())
 	{
 		FNavigationQueryFilter* NavFilter = new FNavigationQueryFilter();
-		NavFilter->SetFilterImplementation(NavData->GetDefaultQueryFilterImpl());
+		NavFilter->SetFilterImplementation(NavData.GetDefaultQueryFilterImpl());
 
-		InitializeFilter(NavData, NavFilter);
+		InitializeFilter(NavData, *NavFilter);
 
 		SharedFilter = MakeShareable(NavFilter);
-		((ANavigationData*)NavData)->StoreQueryFilter(GetClass(), SharedFilter);
+		const_cast<ANavigationData&>(NavData).StoreQueryFilter(GetClass(), SharedFilter);
 	}
 
 	return SharedFilter;
 }
 
-void UNavigationQueryFilter::InitializeFilter(const ANavigationData* NavData, FNavigationQueryFilter* Filter) const
+void UNavigationQueryFilter::InitializeFilter(const ANavigationData& NavData, FNavigationQueryFilter& Filter) const
 {
 	// apply overrides
 	for (int32 i = 0; i < Areas.Num(); i++)
 	{
 		const FNavigationFilterArea& AreaData = Areas[i];
 		
-		const int32 AreaId = NavData->GetAreaID(AreaData.AreaClass);
+		const int32 AreaId = NavData.GetAreaID(AreaData.AreaClass);
 		if (AreaId == INDEX_NONE)
 		{
 			continue;
@@ -147,31 +147,38 @@ void UNavigationQueryFilter::InitializeFilter(const ANavigationData* NavData, FN
 
 		if (AreaData.bIsExcluded)
 		{
-			Filter->SetExcludedArea(AreaId);
+			Filter.SetExcludedArea(AreaId);
 		}
 		else
 		{
 			if (AreaData.bOverrideTravelCost)
 			{
-				Filter->SetAreaCost(AreaId, FMath::Max(1.0f, AreaData.TravelCostOverride));
+				Filter.SetAreaCost(AreaId, FMath::Max(1.0f, AreaData.TravelCostOverride));
 			}
 
 			if (AreaData.bOverrideEnteringCost)
 			{
-				Filter->SetFixedAreaEnteringCost(AreaId, FMath::Max(0.0f, AreaData.EnteringCostOverride));
+				Filter.SetFixedAreaEnteringCost(AreaId, FMath::Max(0.0f, AreaData.EnteringCostOverride));
 			}
 		}
 	}
 
 	// apply flags
-	Filter->SetIncludeFlags(IncludeFlags.Packed);
-	Filter->SetExcludeFlags(ExcludeFlags.Packed);
+	Filter.SetIncludeFlags(IncludeFlags.Packed);
+	Filter.SetExcludeFlags(ExcludeFlags.Packed);
 }
 
-TSharedPtr<const FNavigationQueryFilter> UNavigationQueryFilter::GetQueryFilter(const ANavigationData* NavData, UClass* FilterClass)
+TSharedPtr<const FNavigationQueryFilter> UNavigationQueryFilter::GetQueryFilter(const ANavigationData& NavData, TSubclassOf<UNavigationQueryFilter> FilterClass)
 {
-	UNavigationQueryFilter* DefFilterOb = FilterClass ? FilterClass->GetDefaultObject<UNavigationQueryFilter>() : NULL;
-	return NavData && DefFilterOb ? DefFilterOb->GetQueryFilter(NavData) : NULL;
+	if (FilterClass)
+	{
+		UNavigationQueryFilter* DefFilterOb = FilterClass.GetDefaultObject();
+		// no way we have not default object here
+		check(DefFilterOb);
+		return DefFilterOb->GetQueryFilter(NavData);
+	}
+
+	return nullptr;
 }
 
 void UNavigationQueryFilter::AddTravelCostOverride(TSubclassOf<UNavArea> AreaClass, float TravelCost)
diff --git a/Engine/Source/Runtime/Engine/Private/AI/Navigation/NavFilters/RecastFilter_UseDefaultArea.cpp b/Engine/Source/Runtime/Engine/Private/AI/Navigation/NavFilters/RecastFilter_UseDefaultArea.cpp
--- a/Engine/Source/Runtime/Engine/Private/AI/Navigation/NavFilters/RecastFilter_UseDefaultArea.cpp
+++ b/Engine/Source/Runtime/Engine/Private/AI/Navigation/NavFilters/RecastFilter_UseDefaultArea.cpp
@@ -7,10 +7,10 @@ URecastFilter_UseDefaultArea::URecastFilter_UseDefaultArea(const FObjectInitiali
 {
 }
 
-void URecastFilter_UseDefaultArea::InitializeFilter(const ANavigationData* NavData, FNavigationQueryFilter* Filter) const
+void URecastFilter_UseDefaultArea::InitializeFilter(const ANavigationData& NavData, FNavigationQueryFilter& Filter) const
 {
 #if WITH_RECAST
-	Filter->SetFilterImplementation((const INavigationQueryFilterInterface*)ARecastNavMesh::GetNamedFilter(ERecastNamedFilter::FilterOutAreas));
+	Filter.SetFilterImplementation(static_cast<const INavigationQueryFilterInterface*>(ARecastNavMesh::GetNamedFilter(ERecastNamedFilter::FilterOutAreas)));
 #endif // WITH_RECAST
 
 	Super::InitializeFilter(NavData, Filter);
diff --git a/Engine/Source/Runtime/Engine/Private/AI/Navigation/NavigationData.cpp b/Engine/Source/Runtime/Engine/Private/AI/Navigation/NavigationData.cpp
--- a/Engine/Source/Runtime/Engine/Private/AI/Navigation/NavigationData.cpp
+++ b/Engine/Source/Runtime/Engine/Private/AI/Navigation/NavigationData.cpp
@@ -8,8 +8,8 @@
 //----------------------------------------------------------------------//
 // FPathFindingQuery
 //----------------------------------------------------------------------//
-FPathFindingQuery::FPathFindingQuery(const UObject* InOwner, const ANavigationData* InNavData, const FVector& Start, const FVector& End, TSharedPtr<const FNavigationQueryFilter> SourceQueryFilter, FNavPathSharedPtr InPathInstanceToFill)
-: NavData(InNavData)
+FPathFindingQuery::FPathFindingQuery(const UObject* InOwner, const ANavigationData& InNavData, const FVector& Start, const FVector& End, TSharedPtr<const FNavigationQueryFilter> SourceQueryFilter, FNavPathSharedPtr InPathInstanceToFill)
+: NavData(&InNavData)
 , Owner(InOwner)
 , StartLocation(Start)
 , EndLocation(End)
@@ -61,7 +61,7 @@ FPathFindingQuery::FPathFindingQuery(FNavPathSharedRef PathToRecalculate, const
 //----------------------------------------------------------------------//
 uint32 FAsyncPathFindingQuery::LastPathFindingUniqueID = INVALID_NAVQUERYID;
 
-FAsyncPathFindingQuery::FAsyncPathFindingQuery(const UObject* InOwner, const ANavigationData* InNavData, const FVector& Start, const FVector& End, const FNavPathQueryDelegate& Delegate, TSharedPtr<const FNavigationQueryFilter> SourceQueryFilter)
+FAsyncPathFindingQuery::FAsyncPathFindingQuery(const UObject* InOwner, const ANavigationData& InNavData, const FVector& Start, const FVector& End, const FNavPathQueryDelegate& Delegate, TSharedPtr<const FNavigationQueryFilter> SourceQueryFilter)
 : FPathFindingQuery(InOwner, InNavData, Start, End, SourceQueryFilter)
 , QueryID(GetUniqueID())
 , OnDoneDelegate(Delegate)
diff --git a/Engine/Source/Runtime/Engine/Private/AI/Navigation/NavigationSystem.cpp b/Engine/Source/Runtime/Engine/Private/AI/Navigation/NavigationSystem.cpp
--- a/Engine/Source/Runtime/Engine/Private/AI/Navigation/NavigationSystem.cpp
+++ b/Engine/Source/Runtime/Engine/Private/AI/Navigation/NavigationSystem.cpp
@@ -1124,18 +1124,21 @@ void UNavigationSystem::SimpleMoveToActor(AController* Controller, const AActor*
 	else
 	{
 		const ANavigationData* NavData = NavSys->GetNavDataForProps(Controller->GetNavAgentPropertiesRef());
-		FPathFindingQuery Query(Controller, NavData, Controller->GetNavAgentLocation(), Goal->GetActorLocation());
-		FPathFindingResult Result = NavSys->FindPathSync(Query);
-		if (Result.IsSuccessful())
+		if (NavData)
 		{
-			Result.Path->SetGoalActorObservation(*Goal, 100.0f);
+			FPathFindingQuery Query(Controller, *NavData, Controller->GetNavAgentLocation(), Goal->GetActorLocation());
+			FPathFindingResult Result = NavSys->FindPathSync(Query);
+			if (Result.IsSuccessful())
+			{
+				Result.Path->SetGoalActorObservation(*Goal, 100.0f);
 
-			PFollowComp->RequestMove(Result.Path, Goal);
-		}
-		else if (PFollowComp->GetStatus() != EPathFollowingStatus::Idle)
-		{
-			PFollowComp->AbortMove(TEXT("Aborting move due to new move request failing to generate a path"), FAIRequestID::AnyRequest);
-			PFollowComp->SetLastMoveAtGoal(false);
+				PFollowComp->RequestMove(Result.Path, Goal);
+			}
+			else if (PFollowComp->GetStatus() != EPathFollowingStatus::Idle)
+			{
+				PFollowComp->AbortMove(TEXT("Aborting move due to new move request failing to generate a path"), FAIRequestID::AnyRequest);
+				PFollowComp->SetLastMoveAtGoal(false);
+			}
 		}
 	}
 }
@@ -1180,16 +1183,19 @@ void UNavigationSystem::SimpleMoveToLocation(AController* Controller, const FVec
 	else
 	{
 		const ANavigationData* NavData = NavSys->GetNavDataForProps(Controller->GetNavAgentPropertiesRef());
-		FPathFindingQuery Query(Controller, NavData, Controller->GetNavAgentLocation(), GoalLocation);
-		FPathFindingResult Result = NavSys->FindPathSync(Query);
-		if (Result.IsSuccessful())
-		{
-			PFollowComp->RequestMove(Result.Path, NULL);
-		}
-		else if (PFollowComp->GetStatus() != EPathFollowingStatus::Idle)
+		if (NavData)
 		{
-			PFollowComp->AbortMove(TEXT("Aborting move due to new move request failing to generate a path"), FAIRequestID::AnyRequest);
-			PFollowComp->SetLastMoveAtGoal(false);
+			FPathFindingQuery Query(Controller, *NavData, Controller->GetNavAgentLocation(), GoalLocation);
+			FPathFindingResult Result = NavSys->FindPathSync(Query);
+			if (Result.IsSuccessful())
+			{
+				PFollowComp->RequestMove(Result.Path, NULL);
+			}
+			else if (PFollowComp->GetStatus() != EPathFollowingStatus::Idle)
+			{
+				PFollowComp->AbortMove(TEXT("Aborting move due to new move request failing to generate a path"), FAIRequestID::AnyRequest);
+				PFollowComp->SetLastMoveAtGoal(false);
+			}
 		}
 	}
 }
@@ -1231,7 +1237,6 @@ UNavigationPath* UNavigationSystem::FindPathToLocationSynchronously(UObject* Wor
 		UNavigationSystem* NavSys = World->GetNavigationSystem();
 
 		ResultPath = NewObject<UNavigationPath>(NavSys);
-		FPathFindingQuery Query(PathfindingContext, NULL, PathStart, PathEnd);
 		bool bValidPathContext = false;
 		const ANavigationData* NavigationData = NULL;
 
@@ -1258,8 +1263,8 @@ UNavigationPath* UNavigationSystem::FindPathToLocationSynchronously(UObject* Wor
 		}
 
 		check(NavigationData);
-		Query.NavData = NavigationData;		
-		Query.QueryFilter = UNavigationQueryFilter::GetQueryFilter(NavigationData, FilterClass);
+		FPathFindingQuery Query(PathfindingContext, *NavigationData, PathStart, PathEnd);
+		Query.QueryFilter = UNavigationQueryFilter::GetQueryFilter(*NavigationData, FilterClass);
 
 		FPathFindingResult Result = NavSys->FindPathSync(Query, EPathFindingMode::Regular);
 		if (Result.IsSuccessful())
@@ -1307,7 +1312,7 @@ bool UNavigationSystem::NavigationRaycast(UObject* WorldContextObject, const FVe
 
 		if (NavData != NULL)
 		{
-			bRaycastBlocked = NavData->Raycast(RayStart, RayEnd, HitLocation, UNavigationQueryFilter::GetQueryFilter(NavData, FilterClass));
+			bRaycastBlocked = NavData->Raycast(RayStart, RayEnd, HitLocation, UNavigationQueryFilter::GetQueryFilter(*NavData, FilterClass));
 		}
 	}
 
@@ -3347,7 +3352,10 @@ FVector UNavigationSystem::ProjectPointToNavigation(UObject* WorldContextObject,
 	if (NavSys)
 	{
 		ANavigationData* UseNavData = NavData ? NavData : NavSys->GetMainNavData(FNavigationSystem::DontCreate);
-		NavSys->ProjectPointToNavigation(Point, ProjectedPoint, QueryExtent.IsNearlyZero() ? INVALID_NAVEXTENT : QueryExtent, UseNavData, UNavigationQueryFilter::GetQueryFilter(UseNavData, FilterClass));
+		if (UseNavData)
+		{
+			NavSys->ProjectPointToNavigation(Point, ProjectedPoint, QueryExtent.IsNearlyZero() ? INVALID_NAVEXTENT : QueryExtent, UseNavData, UNavigationQueryFilter::GetQueryFilter(*UseNavData, FilterClass));
+		}
 	}
 
 	return ProjectedPoint.Location;
@@ -3362,7 +3370,10 @@ FVector UNavigationSystem::GetRandomPoint(UObject* WorldContextObject, ANavigati
 	if (NavSys)
 	{
 		ANavigationData* UseNavData = NavData ? NavData : NavSys->GetMainNavData(FNavigationSystem::DontCreate);
-		NavSys->GetRandomPoint(RandomPoint, UseNavData, UNavigationQueryFilter::GetQueryFilter(UseNavData, FilterClass));
+		if (UseNavData)
+		{
+			NavSys->GetRandomPoint(RandomPoint, UseNavData, UNavigationQueryFilter::GetQueryFilter(*UseNavData, FilterClass));
+		}
 	}
 
 	return RandomPoint.Location;
@@ -3378,7 +3389,10 @@ FVector UNavigationSystem::GetRandomPointInRadius(UObject* WorldContextObject, c
 	if (NavSys)
 	{
 		ANavigationData* UseNavData = NavData ? NavData : NavSys->GetMainNavData(FNavigationSystem::DontCreate);
-		NavSys->GetRandomPointInRadius(Origin, Radius, RandomPoint, UseNavData, UNavigationQueryFilter::GetQueryFilter(UseNavData, FilterClass));
+		if (UseNavData)
+		{
+			NavSys->GetRandomPointInRadius(Origin, Radius, RandomPoint, UseNavData, UNavigationQueryFilter::GetQueryFilter(*UseNavData, FilterClass));
+		}
 	}
 
 	return RandomPoint.Location;
@@ -3391,7 +3405,10 @@ ENavigationQueryResult::Type UNavigationSystem::GetPathCost(UObject* WorldContex
 	if (NavSys)
 	{
 		ANavigationData* UseNavData = NavData ? NavData : NavSys->GetMainNavData(FNavigationSystem::DontCreate);
-		return NavSys->GetPathCost(PathStart, PathEnd, OutPathCost, UseNavData, UNavigationQueryFilter::GetQueryFilter(UseNavData, FilterClass));
+		if (UseNavData)
+		{
+			return NavSys->GetPathCost(PathStart, PathEnd, OutPathCost, UseNavData, UNavigationQueryFilter::GetQueryFilter(*UseNavData, FilterClass));
+		}
 	}
 
 	return ENavigationQueryResult::Error;
@@ -3406,7 +3423,10 @@ ENavigationQueryResult::Type UNavigationSystem::GetPathLength(UObject* WorldCont
 	if (NavSys)
 	{
 		ANavigationData* UseNavData = NavData ? NavData : NavSys->GetMainNavData(FNavigationSystem::DontCreate);
-		return NavSys->GetPathLength(PathStart, PathEnd, OutPathLength, UseNavData, UNavigationQueryFilter::GetQueryFilter(UseNavData, FilterClass));
+		if (UseNavData)
+		{
+			return NavSys->GetPathLength(PathStart, PathEnd, OutPathLength, UseNavData, UNavigationQueryFilter::GetQueryFilter(*UseNavData, FilterClass));
+		}
 	}
 
 	return ENavigationQueryResult::Error;
diff --git a/Engine/Source/Runtime/Engine/Private/AI/Navigation/NavigationTestingActor.cpp b/Engine/Source/Runtime/Engine/Private/AI/Navigation/NavigationTestingActor.cpp
--- a/Engine/Source/Runtime/Engine/Private/AI/Navigation/NavigationTestingActor.cpp
+++ b/Engine/Source/Runtime/Engine/Private/AI/Navigation/NavigationTestingActor.cpp
@@ -410,8 +410,8 @@ void ANavigationTestingActor::OnPathEvent(FNavigationPath* InvalidatedPath, ENav
 
 FPathFindingQuery ANavigationTestingActor::BuildPathFindingQuery(const ANavigationTestingActor* Goal) const
 {
-	check(Goal);
-	return FPathFindingQuery(this, MyNavData, GetNavAgentLocation(), Goal->GetNavAgentLocation(), UNavigationQueryFilter::GetQueryFilter(MyNavData, FilterClass));
+	check(Goal && MyNavData);
+	return FPathFindingQuery(this, *MyNavData, GetNavAgentLocation(), Goal->GetNavAgentLocation(), UNavigationQueryFilter::GetQueryFilter(*MyNavData, FilterClass));
 }
 
 
