AI: Fixed non-unity linker errors

diff --git a/Engine/Source/Runtime/AIModule/Private/EnvironmentQuery/EnvQueryTraceHelpers.h b/Engine/Source/Runtime/AIModule/Private/EnvironmentQuery/EnvQueryTraceHelpers.h
--- a/Engine/Source/Runtime/AIModule/Private/EnvironmentQuery/EnvQueryTraceHelpers.h
+++ b/Engine/Source/Runtime/AIModule/Private/EnvironmentQuery/EnvQueryTraceHelpers.h
@@ -99,11 +99,10 @@ namespace FEQSHelpers
 	void RunPhysRaycasts(UWorld* World, const FEnvTraceData& TraceData, const FVector& SourcePt, TArray<FNavLocation>& Points, ETraceMode TraceMode = ETraceMode::Keep);
 	void RunPhysProjection(UWorld* World, const FEnvTraceData& TraceData, TArray<FNavLocation>& Points, ETraceMode TraceMode = ETraceMode::Discard);
 
-	//----------------------------------------------------------------------//
-	// deprecates
-	//----------------------------------------------------------------------//
+	// deprecated
+
 	DEPRECATED(4.8, "This version of RunNavRaycasts is deprecated. Please use ANavigationData reference rather than a pointer version")
-	void RunNavRaycasts(const ANavigationData* NavData, const FEnvTraceData& TraceData, const FVector& SourcePt, TArray<FNavLocation>& Points, ETraceMode TraceMode = ETraceMode::Keep)
+	FORCEINLINE void RunNavRaycasts(const ANavigationData* NavData, const FEnvTraceData& TraceData, const FVector& SourcePt, TArray<FNavLocation>& Points, ETraceMode TraceMode = ETraceMode::Keep)
 	{
 		if (NavData)
 		{
@@ -112,7 +111,7 @@ namespace FEQSHelpers
 	}
 
 	DEPRECATED(4.8, "This version of RunNavProjection is deprecated. Please use ANavigationData reference rather than a pointer version")
-	void RunNavProjection(const ANavigationData* NavData, const FEnvTraceData& TraceData, TArray<FNavLocation>& Points, ETraceMode TraceMode = ETraceMode::Discard)
+	FORCEINLINE void RunNavProjection(const ANavigationData* NavData, const FEnvTraceData& TraceData, TArray<FNavLocation>& Points, ETraceMode TraceMode = ETraceMode::Discard)
 	{
 		if (NavData)
 		{
