finishing crowd simulated movement after colliding with goal actor
#ue4 UE-13275

diff --git a/Engine/Source/Runtime/AIModule/Private/Navigation/CrowdFollowingComponent.cpp b/Engine/Source/Runtime/AIModule/Private/Navigation/CrowdFollowingComponent.cpp
--- a/Engine/Source/Runtime/AIModule/Private/Navigation/CrowdFollowingComponent.cpp
+++ b/Engine/Source/Runtime/AIModule/Private/Navigation/CrowdFollowingComponent.cpp
@@ -653,7 +653,13 @@ void UCrowdFollowingComponent::UpdatePathSegment()
 		const FVector CurrentLocation = MovementComp->GetActorFeetLocation();
 		const FVector GoalLocation = GetCurrentTargetLocation();
 
-		if (bFinalPathPart)
+		if (bCollidedWithGoal)
+		{
+			// check if collided with goal actor
+			OnSegmentFinished();
+			OnPathFinished(EPathFollowingResult::Success);
+		}
+		else if (bFinalPathPart)
 		{
 			const FVector ToTarget = (GoalLocation - MovementComp->GetActorFeetLocation());
 			const bool bDirectPath = Path->CastPath<FAbstractNavigationPath>() != NULL;
