fixed crash on reading object value from non object blackboard keys
#ue4 UE-13437

diff --git a/Engine/Source/Runtime/AIModule/Private/BehaviorTree/BlackboardComponent.cpp b/Engine/Source/Runtime/AIModule/Private/BehaviorTree/BlackboardComponent.cpp
--- a/Engine/Source/Runtime/AIModule/Private/BehaviorTree/BlackboardComponent.cpp
+++ b/Engine/Source/Runtime/AIModule/Private/BehaviorTree/BlackboardComponent.cpp
@@ -530,6 +530,13 @@ void UBlackboardComponent::DescribeSelfToVisLog(FVisualLogEntry* Snapshot) const
 
 UObject* UBlackboardComponent::GetValueAsObject(const FName& KeyName) const
 {
+	// temporary fix, will be replaced with type safe entries soon
+	FBlackboard::FKey KeyID = GetKeyID(KeyName);
+	if (GetKeyType(KeyID) != UBlackboardKeyType_Object::StaticClass())
+	{
+		return UBlackboardKeyType_Object::InvalidValue;
+	}
+
 	return GetValue<UBlackboardKeyType_Object>(KeyName);
 }
 
