Making EQS "point" items contain navigation data has been unhacked a bit #UE4

- also made FNavLocation constructor taking an FVector explicit which surfaced some interesting use cases (and fixed those as well!)

Addresses UE-12087

diff --git a/Engine/Source/Runtime/AIModule/Classes/EnvironmentQuery/Items/EnvQueryItemType_Point.h b/Engine/Source/Runtime/AIModule/Classes/EnvironmentQuery/Items/EnvQueryItemType_Point.h
--- a/Engine/Source/Runtime/AIModule/Classes/EnvironmentQuery/Items/EnvQueryItemType_Point.h
+++ b/Engine/Source/Runtime/AIModule/Classes/EnvironmentQuery/Items/EnvQueryItemType_Point.h
@@ -22,3 +22,7 @@ class AIMODULE_API UEnvQueryItemType_Point : public UEnvQueryItemType_VectorBase
 	virtual FVector GetItemLocation(const uint8* RawData) const override;
 	virtual FNavLocation GetItemNavLocation(const uint8* RawData) const;
 };
+
+// a specialization to support saving locations with navigation data already gathered
+template<>
+void FEnvQueryInstance::AddItemData<UEnvQueryItemType_Point, FNavLocation>(FNavLocation ItemValue);
\ No newline at end of file
diff --git a/Engine/Source/Runtime/AIModule/Private/EnvironmentQuery/EnvQueryTraceHelpers.cpp b/Engine/Source/Runtime/AIModule/Private/EnvironmentQuery/EnvQueryTraceHelpers.cpp
--- a/Engine/Source/Runtime/AIModule/Private/EnvironmentQuery/EnvQueryTraceHelpers.cpp
+++ b/Engine/Source/Runtime/AIModule/Private/EnvironmentQuery/EnvQueryTraceHelpers.cpp
@@ -13,7 +13,7 @@ void FEQSHelpers::FBatchTrace::DoSingleSourceMultiDestinations<EEnvTraceShape::L
 		const bool bHit = RunLineTrace(Source, Points[Idx].Location, HitPos);
 		if (bHit)
 		{
-			Points[Idx] = HitPos;
+			Points[Idx] = FNavLocation(HitPos);
 		}
 		else if (TraceMode == ETraceMode::Discard)
 		{
@@ -32,7 +32,7 @@ void FEQSHelpers::FBatchTrace::DoSingleSourceMultiDestinations<EEnvTraceShape::B
 		const bool bHit = RunBoxTrace(Source, Points[Idx].Location, HitPos);
 		if (bHit)
 		{
-			Points[Idx] = HitPos;
+			Points[Idx] = FNavLocation(HitPos);
 		}
 		else if (TraceMode == ETraceMode::Discard)
 		{
@@ -51,7 +51,7 @@ void FEQSHelpers::FBatchTrace::DoSingleSourceMultiDestinations<EEnvTraceShape::S
 		const bool bHit = RunSphereTrace(Source, Points[Idx].Location, HitPos);
 		if (bHit)
 		{
-			Points[Idx] = HitPos;
+			Points[Idx] = FNavLocation(HitPos);
 		}
 		else if (TraceMode == ETraceMode::Discard)
 		{
@@ -70,7 +70,7 @@ void FEQSHelpers::FBatchTrace::DoSingleSourceMultiDestinations<EEnvTraceShape::C
 		const bool bHit = RunCapsuleTrace(Source, Points[Idx].Location, HitPos);
 		if (bHit)
 		{
-			Points[Idx] = HitPos;
+			Points[Idx] = FNavLocation(HitPos);
 		}
 		else if (TraceMode == ETraceMode::Discard)
 		{
@@ -89,7 +89,7 @@ void FEQSHelpers::FBatchTrace::DoProject<EEnvTraceShape::Line>(TArray<FNavLocati
 		const bool bHit = RunLineTrace(Points[Idx].Location + FVector(0, 0, StartOffsetZ), Points[Idx].Location + FVector(0, 0, EndOffsetZ), HitPos);
 		if (bHit)
 		{
-			Points[Idx] = HitPos + FVector(0, 0, HitOffsetZ);
+			Points[Idx] = FNavLocation(HitPos + FVector(0, 0, HitOffsetZ));
 		}
 		else if (TraceMode == ETraceMode::Discard)
 		{
@@ -108,7 +108,7 @@ void FEQSHelpers::FBatchTrace::DoProject<EEnvTraceShape::Box>(TArray<FNavLocatio
 		const bool bHit = RunBoxTrace(Points[Idx].Location + FVector(0, 0, StartOffsetZ), Points[Idx].Location + FVector(0, 0, EndOffsetZ), HitPos);
 		if (bHit)
 		{
-			Points[Idx] = HitPos + FVector(0, 0, HitOffsetZ);
+			Points[Idx] = FNavLocation(HitPos + FVector(0, 0, HitOffsetZ));
 		}
 		else if (TraceMode == ETraceMode::Discard)
 		{
@@ -127,7 +127,7 @@ void FEQSHelpers::FBatchTrace::DoProject<EEnvTraceShape::Sphere>(TArray<FNavLoca
 		const bool bHit = RunSphereTrace(Points[Idx].Location + FVector(0, 0, StartOffsetZ), Points[Idx].Location + FVector(0, 0, EndOffsetZ), HitPos);
 		if (bHit)
 		{
-			Points[Idx] = HitPos + FVector(0, 0, HitOffsetZ);
+			Points[Idx] = FNavLocation(HitPos + FVector(0, 0, HitOffsetZ));
 		}
 		else if (TraceMode == ETraceMode::Discard)
 		{
@@ -146,7 +146,7 @@ void FEQSHelpers::FBatchTrace::DoProject<EEnvTraceShape::Capsule>(TArray<FNavLoc
 		const bool bHit = RunCapsuleTrace(Points[Idx].Location + FVector(0, 0, StartOffsetZ), Points[Idx].Location + FVector(0, 0, EndOffsetZ), HitPos);
 		if (bHit)
 		{
-			Points[Idx] = HitPos + FVector(0, 0, HitOffsetZ);
+			Points[Idx] = FNavLocation(HitPos + FVector(0, 0, HitOffsetZ));
 		}
 		else if (TraceMode == ETraceMode::Discard)
 		{
diff --git a/Engine/Source/Runtime/AIModule/Private/EnvironmentQuery/Generators/EnvQueryGenerator_Donut.cpp b/Engine/Source/Runtime/AIModule/Private/EnvironmentQuery/Generators/EnvQueryGenerator_Donut.cpp
--- a/Engine/Source/Runtime/AIModule/Private/EnvironmentQuery/Generators/EnvQueryGenerator_Donut.cpp
+++ b/Engine/Source/Runtime/AIModule/Private/EnvironmentQuery/Generators/EnvQueryGenerator_Donut.cpp
@@ -123,7 +123,7 @@ void UEnvQueryGenerator_Donut::GenerateItems(FEnvQueryInstance& QueryInstance) c
 				const FVector RingPos(RingRadius * CosValue, RingRadius * SinValue, 0.0f);
 				for (int32 ContextIdx = 0; ContextIdx < CenterPoints.Num(); ContextIdx++)
 				{
-					const FVector PointPos = CenterPoints[ContextIdx] + RingPos;
+					const FNavLocation PointPos = FNavLocation(CenterPoints[ContextIdx] + RingPos);
 					Points.Add(PointPos);
 				}
 			}
diff --git a/Engine/Source/Runtime/AIModule/Private/EnvironmentQuery/Generators/EnvQueryGenerator_OnCircle.cpp b/Engine/Source/Runtime/AIModule/Private/EnvironmentQuery/Generators/EnvQueryGenerator_OnCircle.cpp
--- a/Engine/Source/Runtime/AIModule/Private/EnvironmentQuery/Generators/EnvQueryGenerator_OnCircle.cpp
+++ b/Engine/Source/Runtime/AIModule/Private/EnvironmentQuery/Generators/EnvQueryGenerator_OnCircle.cpp
@@ -148,7 +148,7 @@ void UEnvQueryGenerator_OnCircle::GenerateItemsForCircle(uint8* ContextRawData,
 
 	for (int32 Step = 0; Step < StepsCount; ++Step)
 	{
-		ItemCandidates[Step] = CenterLocation + StartDirection.RotateAngleAxis(AngleStep*Step, FVector::UpVector);
+		ItemCandidates[Step] = FNavLocation(CenterLocation + StartDirection.RotateAngleAxis(AngleStep*Step, FVector::UpVector));
 	}
 
 	switch (TraceData.TraceMode)
diff --git a/Engine/Source/Runtime/AIModule/Private/EnvironmentQuery/Generators/EnvQueryGenerator_ProjectedPoints.cpp b/Engine/Source/Runtime/AIModule/Private/EnvironmentQuery/Generators/EnvQueryGenerator_ProjectedPoints.cpp
--- a/Engine/Source/Runtime/AIModule/Private/EnvironmentQuery/Generators/EnvQueryGenerator_ProjectedPoints.cpp
+++ b/Engine/Source/Runtime/AIModule/Private/EnvironmentQuery/Generators/EnvQueryGenerator_ProjectedPoints.cpp
@@ -42,14 +42,7 @@ void UEnvQueryGenerator_ProjectedPoints::StoreNavPoints(const TArray<FNavLocatio
 	for (int32 Idx = 0; Idx < Points.Num(); Idx++)
 	{
 		// store using default function to handle creating new data entry 
-		QueryInstance.AddItemData<UEnvQueryItemType_Point>(Points[Idx].Location);
-	}
-
-	uint8* DataPtr = QueryInstance.RawData.GetData();
-	for (int32 Idx = 0; Idx < Points.Num(); Idx++)
-	{
-		// overwrite with more detailed info
-		UEnvQueryItemType_Point::SetNavValue(DataPtr + QueryInstance.Items[Idx + InitialElementsCount].DataOffset, Points[Idx]);
+		QueryInstance.AddItemData<UEnvQueryItemType_Point>(Points[Idx]);
 	}
 
 	FEnvQueryOptionInstance& OptionInstance = QueryInstance.Options[QueryInstance.OptionIndex];
diff --git a/Engine/Source/Runtime/AIModule/Private/EnvironmentQuery/Generators/EnvQueryGenerator_SimpleGrid.cpp b/Engine/Source/Runtime/AIModule/Private/EnvironmentQuery/Generators/EnvQueryGenerator_SimpleGrid.cpp
--- a/Engine/Source/Runtime/AIModule/Private/EnvironmentQuery/Generators/EnvQueryGenerator_SimpleGrid.cpp
+++ b/Engine/Source/Runtime/AIModule/Private/EnvironmentQuery/Generators/EnvQueryGenerator_SimpleGrid.cpp
@@ -37,7 +37,7 @@ void UEnvQueryGenerator_SimpleGrid::GenerateItems(FEnvQueryInstance& QueryInstan
 		{
 			for (int32 IndexY = 0; IndexY <= ItemCount; ++IndexY)
 			{
-				const FVector TestPoint = ContextLocations[ContextIndex] - FVector(DensityValue * (IndexX - ItemCountHalf), DensityValue * (IndexY - ItemCountHalf), 0);
+				const FNavLocation TestPoint = FNavLocation(ContextLocations[ContextIndex] - FVector(DensityValue * (IndexX - ItemCountHalf), DensityValue * (IndexY - ItemCountHalf), 0));
 				GridPoints.Add(TestPoint);
 			}
 		}
diff --git a/Engine/Source/Runtime/AIModule/Private/EnvironmentQuery/Items/EnvQueryItemType_Point.cpp b/Engine/Source/Runtime/AIModule/Private/EnvironmentQuery/Items/EnvQueryItemType_Point.cpp
--- a/Engine/Source/Runtime/AIModule/Private/EnvironmentQuery/Items/EnvQueryItemType_Point.cpp
+++ b/Engine/Source/Runtime/AIModule/Private/EnvironmentQuery/Items/EnvQueryItemType_Point.cpp
@@ -3,6 +3,18 @@
 #include "AIModulePrivate.h"
 #include "EnvironmentQuery/Items/EnvQueryItemType_Point.h"
 
+template<>
+void FEnvQueryInstance::AddItemData<UEnvQueryItemType_Point, FNavLocation>(FNavLocation ItemValue)
+{
+	DEC_MEMORY_STAT_BY(STAT_AI_EQS_InstanceMemory, RawData.GetAllocatedSize() + Items.GetAllocatedSize());
+
+	const int32 DataOffset = RawData.AddUninitialized(ValueSize);
+	UEnvQueryItemType_Point::SetNavValue(RawData.GetData() + DataOffset, ItemValue);
+	Items.Add(FEnvQueryItem(DataOffset));
+
+	INC_MEMORY_STAT_BY(STAT_AI_EQS_InstanceMemory, RawData.GetAllocatedSize() + Items.GetAllocatedSize());
+}
+
 UEnvQueryItemType_Point::UEnvQueryItemType_Point(const FObjectInitializer& ObjectInitializer) : Super(ObjectInitializer)
 {
 	ValueSize = sizeof(FNavLocation);
diff --git a/Engine/Source/Runtime/AIModule/Private/Navigation/PathFollowingComponent.cpp b/Engine/Source/Runtime/AIModule/Private/Navigation/PathFollowingComponent.cpp
--- a/Engine/Source/Runtime/AIModule/Private/Navigation/PathFollowingComponent.cpp
+++ b/Engine/Source/Runtime/AIModule/Private/Navigation/PathFollowingComponent.cpp
@@ -1261,7 +1261,7 @@ FNavLocation UPathFollowingComponent::GetCurrentNavLocation() const
 	// get navigation location of moved actor
 	if (MovementComp == NULL || GetWorld() == NULL || GetWorld()->GetNavigationSystem() == NULL)
 	{
-		return FVector::ZeroVector;
+		return FNavLocation();
 	}
 
 	const FVector OwnerLoc = MovementComp->GetActorNavLocation();
diff --git a/Engine/Source/Runtime/Engine/Classes/AI/Navigation/NavigationTypes.h b/Engine/Source/Runtime/Engine/Classes/AI/Navigation/NavigationTypes.h
--- a/Engine/Source/Runtime/Engine/Classes/AI/Navigation/NavigationTypes.h
+++ b/Engine/Source/Runtime/Engine/Classes/AI/Navigation/NavigationTypes.h
@@ -229,7 +229,7 @@ struct FNavLocation
 	NavNodeRef NodeRef;
 
 	FNavLocation() : Location(FVector::ZeroVector), NodeRef(INVALID_NAVNODEREF) {}
-	FNavLocation(const FVector& InLocation, NavNodeRef InNodeRef = INVALID_NAVNODEREF) 
+	explicit FNavLocation(const FVector& InLocation, NavNodeRef InNodeRef = INVALID_NAVNODEREF) 
 		: Location(InLocation), NodeRef(InNodeRef) {}
 };
 
diff --git a/Engine/Source/Runtime/Engine/Classes/AI/Navigation/RecastNavMesh.h b/Engine/Source/Runtime/Engine/Classes/AI/Navigation/RecastNavMesh.h
--- a/Engine/Source/Runtime/Engine/Classes/AI/Navigation/RecastNavMesh.h
+++ b/Engine/Source/Runtime/Engine/Classes/AI/Navigation/RecastNavMesh.h
@@ -627,6 +627,7 @@ public:
 
 		FORCEINLINE int32 GetMaxCorridorSize() const { return MAX_PATH_CORRIDOR_POLYS; }
 		FORCEINLINE bool HasHit() const { return HitTime != FLT_MAX; }
+		FORCEINLINE NavNodeRef GetLastNodeRef() const { return CorridorPolysCount > 0 ? CorridorPolys[CorridorPolysCount] : INVALID_NAVNODEREF; }
 	};
 
 	//----------------------------------------------------------------------//
diff --git a/Engine/Source/Runtime/Engine/Private/AI/Navigation/NavigationTestingActor.cpp b/Engine/Source/Runtime/Engine/Private/AI/Navigation/NavigationTestingActor.cpp
--- a/Engine/Source/Runtime/Engine/Private/AI/Navigation/NavigationTestingActor.cpp
+++ b/Engine/Source/Runtime/Engine/Private/AI/Navigation/NavigationTestingActor.cpp
@@ -410,8 +410,13 @@ void ANavigationTestingActor::OnPathEvent(FNavigationPath* InvalidatedPath, ENav
 
 FPathFindingQuery ANavigationTestingActor::BuildPathFindingQuery(const ANavigationTestingActor* Goal) const
 {
-	check(Goal && MyNavData);
-	return FPathFindingQuery(this, *MyNavData, GetNavAgentLocation(), Goal->GetNavAgentLocation(), UNavigationQueryFilter::GetQueryFilter(*MyNavData, FilterClass));
+	check(Goal);
+	if (MyNavData)
+	{
+		return FPathFindingQuery(this, *MyNavData, GetNavAgentLocation(), Goal->GetNavAgentLocation(), UNavigationQueryFilter::GetQueryFilter(*MyNavData, FilterClass));
+	}
+	
+	return FPathFindingQuery();
 }
 
 
diff --git a/Engine/Source/Runtime/Engine/Private/AI/Navigation/RecastNavMesh.cpp b/Engine/Source/Runtime/Engine/Private/AI/Navigation/RecastNavMesh.cpp
--- a/Engine/Source/Runtime/Engine/Private/AI/Navigation/RecastNavMesh.cpp
+++ b/Engine/Source/Runtime/Engine/Private/AI/Navigation/RecastNavMesh.cpp
@@ -1809,7 +1809,7 @@ void ARecastNavMesh::BatchRaycast(TArray<FNavigationRaycastWork>& Workload, TSha
 			if (dtStatusSucceed(RaycastStatus) && RaycastResult.HasHit())
 			{
 				WorkItem.bDidHit = true;
-				WorkItem.HitLocation = WorkItem.RayStart + (WorkItem.RayEnd - WorkItem.RayStart) * RaycastResult.HitTime;
+				WorkItem.HitLocation = FNavLocation(WorkItem.RayStart + (WorkItem.RayEnd - WorkItem.RayStart) * RaycastResult.HitTime, RaycastResult.GetLastNodeRef());
 			}
 		}
 	}
