Made it possible for Blackboard key observer request unregistering as part of handling "value change" notification #UE4

The change involves deprecating old FOnBlackboardChange and introduction new FOnBlackboardChangeNotification that allows a function to return a value indicating if it's still interested in observing given BB key

diff --git a/Engine/Source/Runtime/AIModule/Classes/BehaviorTree/BehaviorTreeTypes.h b/Engine/Source/Runtime/AIModule/Classes/BehaviorTree/BehaviorTreeTypes.h
--- a/Engine/Source/Runtime/AIModule/Classes/BehaviorTree/BehaviorTreeTypes.h
+++ b/Engine/Source/Runtime/AIModule/Classes/BehaviorTree/BehaviorTreeTypes.h
@@ -44,8 +44,15 @@ namespace FBlackboard
 	const FKey InvalidKey = FKey(-1);
 }
 
+enum class EBlackboardNotificationResult : uint8
+{
+	RemoveObserver,
+	ContinueObserving
+};
+
 // delegate defines
 DECLARE_DELEGATE_TwoParams(FOnBlackboardChange, const UBlackboardComponent&, FBlackboard::FKey /*key ID*/);
+DECLARE_DELEGATE_RetVal_TwoParams(EBlackboardNotificationResult, FOnBlackboardChangeNotification, const UBlackboardComponent&, FBlackboard::FKey /*key ID*/);
 
 namespace BTSpecialChild
 {
diff --git a/Engine/Source/Runtime/AIModule/Classes/BehaviorTree/BlackboardComponent.h b/Engine/Source/Runtime/AIModule/Classes/BehaviorTree/BlackboardComponent.h
--- a/Engine/Source/Runtime/AIModule/Classes/BehaviorTree/BlackboardComponent.h
+++ b/Engine/Source/Runtime/AIModule/Classes/BehaviorTree/BlackboardComponent.h
@@ -67,8 +67,12 @@ public:
 	bool HasValidAsset() const;
 
 	/** register observer for blackboard key */
+	DEPRECATED(4.8, "FOnBlackboardChange delegate has been deprecated. Please use FOnBlackboardChangeNotification instead.")
 	FDelegateHandle RegisterObserver(FBlackboard::FKey KeyID, UObject* NotifyOwner, FOnBlackboardChange ObserverDelegate);
 
+	/** register observer for blackboard key */
+	FDelegateHandle RegisterObserver(FBlackboard::FKey KeyID, UObject* NotifyOwner, FOnBlackboardChangeNotification ObserverDelegate);
+
 	/** unregister observer from blackboard key */
 	DELEGATE_DEPRECATED("This overload of UnregisterObserver is deprecated, instead pass the result of RegisterObserver.")
 	void UnregisterObserver(FBlackboard::FKey KeyID, FOnBlackboardChange ObserverDelegate);
@@ -283,9 +287,15 @@ protected:
 	UPROPERTY(transient)
 	TArray<UBlackboardKeyType*> KeyInstances;
 
-	/** observers registered for blackboard keys */
-	TMultiMap<uint8, FOnBlackboardChange> Observers;
+private:
+	/** observers registered for blackboard keys
+	 *	@note this is a legacy variable, left hanging only to support deprecated code until it does away */
+	TMultiMap<uint8, FOnBlackboardChange> Observers_DEPRECATED;
 
+protected:
+	/** observers registered for blackboard keys */
+	mutable TMultiMap<uint8, FOnBlackboardChangeNotification> Observers;
+	
 	/** observers registered from owner objects */
 	TMultiMap<UObject*, FDelegateHandle> ObserverHandles;
 
diff --git a/Engine/Source/Runtime/AIModule/Classes/BehaviorTree/Decorators/BTDecorator_Blackboard.h b/Engine/Source/Runtime/AIModule/Classes/BehaviorTree/Decorators/BTDecorator_Blackboard.h
--- a/Engine/Source/Runtime/AIModule/Classes/BehaviorTree/Decorators/BTDecorator_Blackboard.h
+++ b/Engine/Source/Runtime/AIModule/Classes/BehaviorTree/Decorators/BTDecorator_Blackboard.h
@@ -34,7 +34,7 @@ class AIMODULE_API UBTDecorator_Blackboard : public UBTDecorator_BlackboardBase
 	GENERATED_UCLASS_BODY()
 
 	virtual bool CalculateRawConditionValue(UBehaviorTreeComponent& OwnerComp, uint8* NodeMemory) const override;
-	virtual void OnBlackboardChange(const UBlackboardComponent& Blackboard, FBlackboard::FKey ChangedKeyID) override;
+	virtual EBlackboardNotificationResult OnBlackboardKeyValueChange(const UBlackboardComponent& Blackboard, FBlackboard::FKey ChangedKeyID) override;
 	virtual void DescribeRuntimeValues(const UBehaviorTreeComponent& OwnerComp, uint8* NodeMemory, EBTDescriptionVerbosity::Type Verbosity, TArray<FString>& Values) const override;
 	virtual FString GetStaticDescription() const override;
 
diff --git a/Engine/Source/Runtime/AIModule/Classes/BehaviorTree/Decorators/BTDecorator_BlackboardBase.h b/Engine/Source/Runtime/AIModule/Classes/BehaviorTree/Decorators/BTDecorator_BlackboardBase.h
--- a/Engine/Source/Runtime/AIModule/Classes/BehaviorTree/Decorators/BTDecorator_BlackboardBase.h
+++ b/Engine/Source/Runtime/AIModule/Classes/BehaviorTree/Decorators/BTDecorator_BlackboardBase.h
@@ -13,9 +13,12 @@ class AIMODULE_API UBTDecorator_BlackboardBase : public UBTDecorator
 	/** initialize any asset related data */
 	virtual void InitializeFromAsset(UBehaviorTree& Asset) override;
 
-	/** notify about change in blackboard keys */
+	DEPRECATED(4.8, "OnBlackboardChange has been deprecated. Please switch over to using OnBlackboardKeyValueChange, it's different just by a return type.")
 	virtual void OnBlackboardChange(const UBlackboardComponent& Blackboard, FBlackboard::FKey ChangedKeyID);
 
+	/** notify about change in blackboard keys */
+	virtual EBlackboardNotificationResult OnBlackboardKeyValueChange(const UBlackboardComponent& Blackboard, FBlackboard::FKey ChangedKeyID);
+
 #if WITH_EDITOR
 	virtual FName GetNodeIconName() const override;
 #endif
diff --git a/Engine/Source/Runtime/AIModule/Classes/BehaviorTree/Decorators/BTDecorator_BlueprintBase.h b/Engine/Source/Runtime/AIModule/Classes/BehaviorTree/Decorators/BTDecorator_BlueprintBase.h
--- a/Engine/Source/Runtime/AIModule/Classes/BehaviorTree/Decorators/BTDecorator_BlueprintBase.h
+++ b/Engine/Source/Runtime/AIModule/Classes/BehaviorTree/Decorators/BTDecorator_BlueprintBase.h
@@ -39,7 +39,7 @@ public:
 	virtual void PostLoad() override;
 
 	/** notify about changes in blackboard */
-	void OnBlackboardChange(const UBlackboardComponent& Blackboard, FBlackboard::FKey ChangedKeyID);
+	EBlackboardNotificationResult OnBlackboardKeyValueChange(const UBlackboardComponent& Blackboard, FBlackboard::FKey ChangedKeyID);
 
 	virtual FString GetStaticDescription() const override;
 	virtual void DescribeRuntimeValues(const UBehaviorTreeComponent& OwnerComp, uint8* NodeMemory, EBTDescriptionVerbosity::Type Verbosity, TArray<FString>& Values) const override;
diff --git a/Engine/Source/Runtime/AIModule/Classes/BehaviorTree/Decorators/BTDecorator_CompareBBEntries.h b/Engine/Source/Runtime/AIModule/Classes/BehaviorTree/Decorators/BTDecorator_CompareBBEntries.h
--- a/Engine/Source/Runtime/AIModule/Classes/BehaviorTree/Decorators/BTDecorator_CompareBBEntries.h
+++ b/Engine/Source/Runtime/AIModule/Classes/BehaviorTree/Decorators/BTDecorator_CompareBBEntries.h
@@ -47,7 +47,7 @@ public:
 
 	virtual void InitializeFromAsset(UBehaviorTree& Asset) override;
 
-	virtual void OnBlackboardChange(const UBlackboardComponent& Blackboard, FBlackboard::FKey ChangedKeyID);
+	virtual EBlackboardNotificationResult OnBlackboardKeyValueChange(const UBlackboardComponent& Blackboard, FBlackboard::FKey ChangedKeyID);
 	virtual void OnBecomeRelevant(UBehaviorTreeComponent& OwnerComp, uint8* NodeMemory) override;
 	virtual void OnCeaseRelevant(UBehaviorTreeComponent& OwnerComp, uint8* NodeMemory) override;
 
diff --git a/Engine/Source/Runtime/AIModule/Classes/BehaviorTree/Decorators/BTDecorator_ConditionalLoop.h b/Engine/Source/Runtime/AIModule/Classes/BehaviorTree/Decorators/BTDecorator_ConditionalLoop.h
--- a/Engine/Source/Runtime/AIModule/Classes/BehaviorTree/Decorators/BTDecorator_ConditionalLoop.h
+++ b/Engine/Source/Runtime/AIModule/Classes/BehaviorTree/Decorators/BTDecorator_ConditionalLoop.h
@@ -14,7 +14,7 @@ class AIMODULE_API UBTDecorator_ConditionalLoop : public UBTDecorator_Blackboard
 	GENERATED_UCLASS_BODY()
 
 	virtual bool CalculateRawConditionValue(UBehaviorTreeComponent& OwnerComp, uint8* NodeMemory) const override;
-	virtual void OnBlackboardChange(const UBlackboardComponent& Blackboard, FBlackboard::FKey ChangedKeyID) override;
+	virtual EBlackboardNotificationResult OnBlackboardKeyValueChange(const UBlackboardComponent& Blackboard, FBlackboard::FKey ChangedKeyID) override;
 	virtual void OnNodeDeactivation(FBehaviorTreeSearchData& SearchData, EBTNodeResult::Type NodeResult) override;
 
 #if WITH_EDITOR
diff --git a/Engine/Source/Runtime/AIModule/Classes/BehaviorTree/Decorators/BTDecorator_IsBBEntryOfClass.h b/Engine/Source/Runtime/AIModule/Classes/BehaviorTree/Decorators/BTDecorator_IsBBEntryOfClass.h
--- a/Engine/Source/Runtime/AIModule/Classes/BehaviorTree/Decorators/BTDecorator_IsBBEntryOfClass.h
+++ b/Engine/Source/Runtime/AIModule/Classes/BehaviorTree/Decorators/BTDecorator_IsBBEntryOfClass.h
@@ -19,7 +19,7 @@ protected:
 	TSubclassOf<UObject> TestClass;
 
 	virtual bool CalculateRawConditionValue(UBehaviorTreeComponent& OwnerComp, uint8* NodeMemory) const override;
-	virtual void OnBlackboardChange(const UBlackboardComponent& Blackboard, FBlackboard::FKey ChangedKeyID) override;
+	virtual EBlackboardNotificationResult OnBlackboardKeyValueChange(const UBlackboardComponent& Blackboard, FBlackboard::FKey ChangedKeyID) override;
 	virtual void DescribeRuntimeValues(const UBehaviorTreeComponent& OwnerComp, uint8* NodeMemory, EBTDescriptionVerbosity::Type Verbosity, TArray<FString>& Values) const override;
 	virtual FString GetStaticDescription() const override;
 };
diff --git a/Engine/Source/Runtime/AIModule/Private/BehaviorTree/BlackboardComponent.cpp b/Engine/Source/Runtime/AIModule/Private/BehaviorTree/BlackboardComponent.cpp
--- a/Engine/Source/Runtime/AIModule/Private/BehaviorTree/BlackboardComponent.cpp
+++ b/Engine/Source/Runtime/AIModule/Private/BehaviorTree/BlackboardComponent.cpp
@@ -278,7 +278,7 @@ int32 UBlackboardComponent::GetNumKeys() const
 	return BlackboardAsset ? BlackboardAsset->GetNumKeys() : 0;
 }
 
-FDelegateHandle UBlackboardComponent::RegisterObserver(FBlackboard::FKey KeyID, UObject* NotifyOwner, FOnBlackboardChange ObserverDelegate)
+FDelegateHandle UBlackboardComponent::RegisterObserver(FBlackboard::FKey KeyID, UObject* NotifyOwner, FOnBlackboardChangeNotification ObserverDelegate)
 {
 	for (auto It = Observers.CreateConstKeyIterator(KeyID); It; ++It)
 	{
@@ -295,15 +295,17 @@ FDelegateHandle UBlackboardComponent::RegisterObserver(FBlackboard::FKey KeyID,
 	return Handle;
 }
 
-void UBlackboardComponent::UnregisterObserver(FBlackboard::FKey KeyID, FOnBlackboardChange ObserverDelegate)
+void UBlackboardComponent::UnregisterObserver(FBlackboard::FKey KeyID, FDelegateHandle ObserverHandle)
 {
-	for (auto It = Observers.CreateKeyIterator(KeyID); It; ++It)
+	// this loop is only here for deprecated code support. Will go away
+	// DEPRECATED BEGIN
+	for (auto It = Observers_DEPRECATED.CreateKeyIterator(KeyID); It; ++It)
 	{
-		if (It.Value().DEPRECATED_Compare(ObserverDelegate))
+		if (It.Value().GetHandle() == ObserverHandle)
 		{
 			for (auto HandleIt = ObserverHandles.CreateIterator(); HandleIt; ++HandleIt)
 			{
-				if (HandleIt.Value() == It.Value().GetHandle())
+				if (HandleIt.Value() == ObserverHandle)
 				{
 					HandleIt.RemoveCurrent();
 					break;
@@ -314,10 +316,8 @@ void UBlackboardComponent::UnregisterObserver(FBlackboard::FKey KeyID, FOnBlackb
 			break;
 		}
 	}
-}
+	// DEPRECATED END
 
-void UBlackboardComponent::UnregisterObserver(FBlackboard::FKey KeyID, FDelegateHandle ObserverHandle)
-{
 	for (auto It = Observers.CreateKeyIterator(KeyID); It; ++It)
 	{
 		if (It.Value().GetHandle() == ObserverHandle)
@@ -341,6 +341,18 @@ void UBlackboardComponent::UnregisterObserversFrom(UObject* NotifyOwner)
 {
 	for (auto It = ObserverHandles.CreateKeyIterator(NotifyOwner); It; ++It)
 	{
+		// deprecated code support
+		// DEPRECATED BEGIN
+		for (auto ObsIt = Observers_DEPRECATED.CreateIterator(); ObsIt; ++ObsIt)
+		{
+			if (ObsIt.Value().GetHandle() == It.Value())
+			{
+				ObsIt.RemoveCurrent();
+				break;
+			}
+		}
+		// DEPRECATED END
+
 		for (auto ObsIt = Observers.CreateIterator(); ObsIt; ++ObsIt)
 		{
 			if (ObsIt.Value().GetHandle() == It.Value())
@@ -380,11 +392,24 @@ void UBlackboardComponent::NotifyObservers(FBlackboard::FKey KeyID) const
 	}
 	else
 	{
-		for (TMultiMap<uint8, FOnBlackboardChange>::TConstKeyIterator KeyIt(Observers, KeyID); KeyIt; ++KeyIt)
+		// DEPRECATED BEGIN
+		for (TMultiMap<uint8, FOnBlackboardChange>::TConstKeyIterator KeyIt(Observers_DEPRECATED, KeyID); KeyIt; ++KeyIt)
 		{
 			const FOnBlackboardChange& ObserverDelegate = KeyIt.Value();
 			ObserverDelegate.ExecuteIfBound(*this, KeyID);
 		}
+		// DEPRECATED END
+
+		for (TMultiMap<uint8, FOnBlackboardChangeNotification>::TKeyIterator KeyIt(Observers, KeyID); KeyIt; ++KeyIt)
+		{
+			const FOnBlackboardChangeNotification& ObserverDelegate = KeyIt.Value();
+			const bool bWantsToContinueObserving = ObserverDelegate.IsBound() && (ObserverDelegate.Execute(*this, KeyID) == EBlackboardNotificationResult::ContinueObserving);
+
+			if (bWantsToContinueObserving == false)
+			{
+				KeyIt.RemoveCurrent();
+			}
+		}
 	}
 }
 
@@ -443,8 +468,17 @@ FString UBlackboardComponent::GetDebugInfoString(EBlackboardDescription::Type Mo
 		DebugString += TEXT("Observed Keys:\n");
 
 		TArray<uint8> ObserversKeys;
-		if (Observers.GetKeys(ObserversKeys) > 0)
+		if (Observers.Num() > 0
+			// DEPRECATED BEGIN
+			|| Observers_DEPRECATED.Num() > 0
+			// DEPRECATED END
+			)
 		{
+			Observers.GetKeys(ObserversKeys);
+			// DEPRECATED BEGIN
+			Observers_DEPRECATED.GetKeys(ObserversKeys);
+			// DEPRECATED END
+
 			for (int32 KeyIndex = 0; KeyIndex < ObserversKeys.Num(); ++KeyIndex)
 			{
 				const FBlackboard::FKey KeyID = ObserversKeys[KeyIndex];
@@ -874,3 +908,44 @@ FRotator UBlackboardComponent::GetValueAsRotator(FBlackboard::FKey KeyID) const
 {
 	return GetValue<UBlackboardKeyType_Rotator>(KeyID);
 }
+
+//----------------------------------------------------------------------//
+// DEPRECATED
+//----------------------------------------------------------------------//
+FDelegateHandle UBlackboardComponent::RegisterObserver(FBlackboard::FKey KeyID, UObject* NotifyOwner, FOnBlackboardChange ObserverDelegate)
+{
+	for (auto It = Observers_DEPRECATED.CreateConstKeyIterator(KeyID); It; ++It)
+	{
+		// If the pair's value matches, return a pointer to it.
+		if (It.Value().GetHandle() == ObserverDelegate.GetHandle())
+		{
+			return It.Value().GetHandle();
+		}
+	}
+
+	FDelegateHandle Handle = Observers_DEPRECATED.Add(KeyID, ObserverDelegate).GetHandle();
+	ObserverHandles.Add(NotifyOwner, Handle);
+
+	return Handle;
+}
+
+void UBlackboardComponent::UnregisterObserver(FBlackboard::FKey KeyID, FOnBlackboardChange ObserverDelegate)
+{
+	for (auto It = Observers_DEPRECATED.CreateKeyIterator(KeyID); It; ++It)
+	{
+		if (It.Value().DEPRECATED_Compare(ObserverDelegate))
+		{
+			for (auto HandleIt = ObserverHandles.CreateIterator(); HandleIt; ++HandleIt)
+			{
+				if (HandleIt.Value() == It.Value().GetHandle())
+				{
+					HandleIt.RemoveCurrent();
+					break;
+				}
+			}
+
+			It.RemoveCurrent();
+			break;
+		}
+	}
+}
diff --git a/Engine/Source/Runtime/AIModule/Private/BehaviorTree/Decorators/BTDecorator_Blackboard.cpp b/Engine/Source/Runtime/AIModule/Private/BehaviorTree/Decorators/BTDecorator_Blackboard.cpp
--- a/Engine/Source/Runtime/AIModule/Private/BehaviorTree/Decorators/BTDecorator_Blackboard.cpp
+++ b/Engine/Source/Runtime/AIModule/Private/BehaviorTree/Decorators/BTDecorator_Blackboard.cpp
@@ -55,43 +55,54 @@ bool UBTDecorator_Blackboard::EvaluateOnBlackboard(const UBlackboardComponent& B
 	return bResult;
 }
 
-void UBTDecorator_Blackboard::OnBlackboardChange(const UBlackboardComponent& Blackboard, FBlackboard::FKey ChangedKeyID)
+EBlackboardNotificationResult UBTDecorator_Blackboard::OnBlackboardKeyValueChange(const UBlackboardComponent& Blackboard, FBlackboard::FKey ChangedKeyID)
 {
 	UBehaviorTreeComponent* BehaviorComp = (UBehaviorTreeComponent*)Blackboard.GetBrainComponent();
-	if (BlackboardKey.GetSelectedKeyID() == ChangedKeyID &&
-		GetFlowAbortMode() != EBTFlowAbortMode::None &&
-		BehaviorComp)
+
+	if (BehaviorComp == nullptr)
+	{
+		return EBlackboardNotificationResult::RemoveObserver;
+	}
+
+	EBlackboardNotificationResult Result = EBlackboardNotificationResult::ContinueObserving;
+
+	if (BlackboardKey.GetSelectedKeyID() == ChangedKeyID && GetFlowAbortMode() != EBTFlowAbortMode::None)
 	{
 		if (NotifyObserver == EBTBlackboardRestart::ResultChange)
 		{
 			const bool bIsExecutingBranch = BehaviorComp->IsExecutingBranch(this, GetChildIndex());
 			const bool bPass = EvaluateOnBlackboard(Blackboard);
 
-			UE_VLOG(BehaviorComp->GetOwner(), LogBehaviorTree, Verbose, TEXT("%s, OnBlackboardChange[%s] pass:%d executing:%d => %s"),
+			UE_VLOG(BehaviorComp->GetOwner(), LogBehaviorTree, Verbose, TEXT("%s, OnBlackboardKeyValueChange[%s] pass:%d executing:%d => %s"),
 				*UBehaviorTreeTypes::DescribeNodeHelper(this),
 				*Blackboard.GetKeyName(ChangedKeyID).ToString(), bPass, bIsExecutingBranch,
 				(bIsExecutingBranch && !bPass) || (!bIsExecutingBranch && bPass) ? TEXT("restart") : TEXT("skip"));
 
-			if ((bIsExecutingBranch && !bPass) ||
-				(!bIsExecutingBranch && bPass))
+			// this is basically a XOR, but written down this way for readability's sake
+			if ((bIsExecutingBranch && !bPass) || (!bIsExecutingBranch && bPass))
 			{
 				BehaviorComp->RequestExecution(this);
+				Result = EBlackboardNotificationResult::RemoveObserver;
 			}
 			else if (!bIsExecutingBranch && !bPass && GetParentNode() && GetParentNode()->Children.IsValidIndex(GetChildIndex()))
 			{
 				const UBTCompositeNode* BranchRoot = GetParentNode()->Children[GetChildIndex()].ChildComposite;
 				BehaviorComp->UnregisterAuxNodesInBranch(BranchRoot);
+				Result = EBlackboardNotificationResult::RemoveObserver;
 			}
 		}
 		else
 		{
-			UE_VLOG(BehaviorComp->GetOwner(), LogBehaviorTree, Verbose, TEXT("%s, OnBlackboardChange[%s] => restart"),
+			UE_VLOG(BehaviorComp->GetOwner(), LogBehaviorTree, Verbose, TEXT("%s, OnBlackboardKeyValueChange[%s] => restart"),
 				*UBehaviorTreeTypes::DescribeNodeHelper(this),
 				*Blackboard.GetKeyName(ChangedKeyID).ToString());
 
 			BehaviorComp->RequestExecution(this);
+			Result = EBlackboardNotificationResult::RemoveObserver;
 		}
 	}
+
+	return Result;
 }
 
 void UBTDecorator_Blackboard::DescribeRuntimeValues(const UBehaviorTreeComponent& OwnerComp, uint8* NodeMemory, EBTDescriptionVerbosity::Type Verbosity, TArray<FString>& Values) const
diff --git a/Engine/Source/Runtime/AIModule/Private/BehaviorTree/Decorators/BTDecorator_BlackboardBase.cpp b/Engine/Source/Runtime/AIModule/Private/BehaviorTree/Decorators/BTDecorator_BlackboardBase.cpp
--- a/Engine/Source/Runtime/AIModule/Private/BehaviorTree/Decorators/BTDecorator_BlackboardBase.cpp
+++ b/Engine/Source/Runtime/AIModule/Private/BehaviorTree/Decorators/BTDecorator_BlackboardBase.cpp
@@ -26,7 +26,7 @@ void UBTDecorator_BlackboardBase::OnBecomeRelevant(UBehaviorTreeComponent& Owner
 	if (BlackboardComp)
 	{
 		auto KeyID = BlackboardKey.GetSelectedKeyID();
-		BlackboardComp->RegisterObserver(KeyID, this, FOnBlackboardChange::CreateUObject(this, &UBTDecorator_BlackboardBase::OnBlackboardChange));
+		BlackboardComp->RegisterObserver(KeyID, this, FOnBlackboardChangeNotification::CreateUObject(this, &UBTDecorator_BlackboardBase::OnBlackboardKeyValueChange));
 	}
 }
 
@@ -39,13 +39,19 @@ void UBTDecorator_BlackboardBase::OnCeaseRelevant(UBehaviorTreeComponent& OwnerC
 	}
 }
 
-void UBTDecorator_BlackboardBase::OnBlackboardChange(const UBlackboardComponent& Blackboard, FBlackboard::FKey ChangedKeyID)
+EBlackboardNotificationResult UBTDecorator_BlackboardBase::OnBlackboardKeyValueChange(const UBlackboardComponent& Blackboard, FBlackboard::FKey ChangedKeyID)
 {
 	UBehaviorTreeComponent* BehaviorComp = (UBehaviorTreeComponent*)Blackboard.GetBrainComponent();
-	if (BlackboardKey.GetSelectedKeyID() == ChangedKeyID && BehaviorComp)
+	if (BehaviorComp == nullptr)
+	{
+		return EBlackboardNotificationResult::RemoveObserver;
+	}
+
+	if (BlackboardKey.GetSelectedKeyID() == ChangedKeyID)
 	{
 		BehaviorComp->RequestExecution(this);		
 	}
+	return EBlackboardNotificationResult::ContinueObserving;
 }
 
 #if WITH_EDITOR
@@ -56,3 +62,11 @@ FName UBTDecorator_BlackboardBase::GetNodeIconName() const
 }
 
 #endif	// WITH_EDITOR
+
+//----------------------------------------------------------------------//
+// DEPRECATED
+//----------------------------------------------------------------------//
+void UBTDecorator_BlackboardBase::OnBlackboardChange(const UBlackboardComponent& Blackboard, FBlackboard::FKey ChangedKeyID)
+{
+	OnBlackboardKeyValueChange(Blackboard, ChangedKeyID);
+}
\ No newline at end of file
diff --git a/Engine/Source/Runtime/AIModule/Private/BehaviorTree/Decorators/BTDecorator_BlueprintBase.cpp b/Engine/Source/Runtime/AIModule/Private/BehaviorTree/Decorators/BTDecorator_BlueprintBase.cpp
--- a/Engine/Source/Runtime/AIModule/Private/BehaviorTree/Decorators/BTDecorator_BlueprintBase.cpp
+++ b/Engine/Source/Runtime/AIModule/Private/BehaviorTree/Decorators/BTDecorator_BlueprintBase.cpp
@@ -101,7 +101,7 @@ void UBTDecorator_BlueprintBase::OnBecomeRelevant(UBehaviorTreeComponent& OwnerC
 			const FBlackboard::FKey KeyID = BlackboardComp->GetKeyID(ObservedKeyNames[NameIndex]);
 			if (KeyID != FBlackboard::InvalidKey)
 			{
-				BlackboardComp->RegisterObserver(KeyID, this, FOnBlackboardChange::CreateUObject(this, &UBTDecorator_BlueprintBase::OnBlackboardChange));
+				BlackboardComp->RegisterObserver(KeyID, this, FOnBlackboardChangeNotification::CreateUObject(this, &UBTDecorator_BlueprintBase::OnBlackboardKeyValueChange));
 			}
 		}
 	}
@@ -262,13 +262,14 @@ void UBTDecorator_BlueprintBase::DescribeRuntimeValues(const UBehaviorTreeCompon
 	}
 }
 
-void UBTDecorator_BlueprintBase::OnBlackboardChange(const UBlackboardComponent& Blackboard, FBlackboard::FKey ChangedKeyID)
+EBlackboardNotificationResult UBTDecorator_BlueprintBase::OnBlackboardKeyValueChange(const UBlackboardComponent& Blackboard, FBlackboard::FKey ChangedKeyID)
 {
 	UBehaviorTreeComponent* BehaviorComp = (UBehaviorTreeComponent*)Blackboard.GetBrainComponent();
 	if (BehaviorComp && GetShouldAbort(*BehaviorComp))
 	{
-		BehaviorComp->RequestExecution(this);		
+		BehaviorComp->RequestExecution(this);
 	}
+	return BehaviorComp ? EBlackboardNotificationResult::ContinueObserving : EBlackboardNotificationResult::RemoveObserver;
 }
 
 void UBTDecorator_BlueprintBase::OnInstanceDestroyed(UBehaviorTreeComponent& OwnerComp)
diff --git a/Engine/Source/Runtime/AIModule/Private/BehaviorTree/Decorators/BTDecorator_CompareBBEntries.cpp b/Engine/Source/Runtime/AIModule/Private/BehaviorTree/Decorators/BTDecorator_CompareBBEntries.cpp
--- a/Engine/Source/Runtime/AIModule/Private/BehaviorTree/Decorators/BTDecorator_CompareBBEntries.cpp
+++ b/Engine/Source/Runtime/AIModule/Private/BehaviorTree/Decorators/BTDecorator_CompareBBEntries.cpp
@@ -64,8 +64,8 @@ void UBTDecorator_CompareBBEntries::OnBecomeRelevant(UBehaviorTreeComponent& Own
 	UBlackboardComponent* BlackboardComp = OwnerComp.GetBlackboardComponent();
 	if (BlackboardComp)
 	{
-		BlackboardComp->RegisterObserver(BlackboardKeyA.GetSelectedKeyID(), this, FOnBlackboardChange::CreateUObject(this, &UBTDecorator_CompareBBEntries::OnBlackboardChange));
-		BlackboardComp->RegisterObserver(BlackboardKeyB.GetSelectedKeyID(), this, FOnBlackboardChange::CreateUObject(this, &UBTDecorator_CompareBBEntries::OnBlackboardChange));
+		BlackboardComp->RegisterObserver(BlackboardKeyA.GetSelectedKeyID(), this, FOnBlackboardChangeNotification::CreateUObject(this, &UBTDecorator_CompareBBEntries::OnBlackboardKeyValueChange));
+		BlackboardComp->RegisterObserver(BlackboardKeyB.GetSelectedKeyID(), this, FOnBlackboardChangeNotification::CreateUObject(this, &UBTDecorator_CompareBBEntries::OnBlackboardKeyValueChange));
 	}
 }
 
@@ -78,13 +78,19 @@ void UBTDecorator_CompareBBEntries::OnCeaseRelevant(UBehaviorTreeComponent& Owne
 	}
 }
 
-void UBTDecorator_CompareBBEntries::OnBlackboardChange(const UBlackboardComponent& Blackboard, FBlackboard::FKey ChangedKeyID)
+EBlackboardNotificationResult UBTDecorator_CompareBBEntries::OnBlackboardKeyValueChange(const UBlackboardComponent& Blackboard, FBlackboard::FKey ChangedKeyID)
 {
 	UBehaviorTreeComponent* BehaviorComp = static_cast<UBehaviorTreeComponent*>(Blackboard.GetBrainComponent());
-	if (BehaviorComp && (BlackboardKeyA.GetSelectedKeyID() == ChangedKeyID || BlackboardKeyB.GetSelectedKeyID() == ChangedKeyID))
+	if (BehaviorComp == nullptr)
+	{
+		return EBlackboardNotificationResult::RemoveObserver;
+	}
+	else if (BlackboardKeyA.GetSelectedKeyID() == ChangedKeyID || BlackboardKeyB.GetSelectedKeyID() == ChangedKeyID)
 	{
 		BehaviorComp->RequestExecution(this);		
 	}
+
+	return EBlackboardNotificationResult::ContinueObserving;
 }
 
 #if WITH_EDITOR
diff --git a/Engine/Source/Runtime/AIModule/Private/BehaviorTree/Decorators/BTDecorator_ConditionalLoop.cpp b/Engine/Source/Runtime/AIModule/Private/BehaviorTree/Decorators/BTDecorator_ConditionalLoop.cpp
--- a/Engine/Source/Runtime/AIModule/Private/BehaviorTree/Decorators/BTDecorator_ConditionalLoop.cpp
+++ b/Engine/Source/Runtime/AIModule/Private/BehaviorTree/Decorators/BTDecorator_ConditionalLoop.cpp
@@ -20,9 +20,10 @@ bool UBTDecorator_ConditionalLoop::CalculateRawConditionValue(UBehaviorTreeCompo
 	return true;
 }
 
-void UBTDecorator_ConditionalLoop::OnBlackboardChange(const UBlackboardComponent& Blackboard, FBlackboard::FKey ChangedKeyID)
+EBlackboardNotificationResult UBTDecorator_ConditionalLoop::OnBlackboardKeyValueChange(const UBlackboardComponent& Blackboard, FBlackboard::FKey ChangedKeyID)
 {
 	// empty, don't react to blackboard value changes
+	return EBlackboardNotificationResult::RemoveObserver;
 }
 
 void UBTDecorator_ConditionalLoop::OnNodeDeactivation(FBehaviorTreeSearchData& SearchData, EBTNodeResult::Type NodeResult)
diff --git a/Engine/Source/Runtime/AIModule/Private/BehaviorTree/Decorators/BTDecorator_IsBBEntryOfClass.cpp b/Engine/Source/Runtime/AIModule/Private/BehaviorTree/Decorators/BTDecorator_IsBBEntryOfClass.cpp
--- a/Engine/Source/Runtime/AIModule/Private/BehaviorTree/Decorators/BTDecorator_IsBBEntryOfClass.cpp
+++ b/Engine/Source/Runtime/AIModule/Private/BehaviorTree/Decorators/BTDecorator_IsBBEntryOfClass.cpp
@@ -26,7 +26,7 @@ bool UBTDecorator_IsBBEntryOfClass::CalculateRawConditionValue(UBehaviorTreeComp
 	return false;
 }
 
-void UBTDecorator_IsBBEntryOfClass::OnBlackboardChange(const UBlackboardComponent& Blackboard, FBlackboard::FKey ChangedKeyID)
+EBlackboardNotificationResult UBTDecorator_IsBBEntryOfClass::OnBlackboardKeyValueChange(const UBlackboardComponent& Blackboard, FBlackboard::FKey ChangedKeyID)
 {
 	check(Cast<UBehaviorTreeComponent>(Blackboard.GetBrainComponent()));
 	check(BlackboardKey.GetSelectedKeyID() == ChangedKeyID);
@@ -36,6 +36,8 @@ void UBTDecorator_IsBBEntryOfClass::OnBlackboardChange(const UBlackboardComponen
 	{
 		static_cast<UBehaviorTreeComponent*>(Blackboard.GetBrainComponent())->RequestExecution(this);
 	}
+
+	return EBlackboardNotificationResult::ContinueObserving;
 }
 
 void UBTDecorator_IsBBEntryOfClass::DescribeRuntimeValues(const UBehaviorTreeComponent& OwnerComp, uint8* NodeMemory, EBTDescriptionVerbosity::Type Verbosity, TArray<FString>& Values) const
