Improved UI behavior of BTNodes' names #UE4

diff --git a/Engine/Source/Runtime/AIModule/Classes/BehaviorTree/BTNode.h b/Engine/Source/Runtime/AIModule/Classes/BehaviorTree/BTNode.h
--- a/Engine/Source/Runtime/AIModule/Classes/BehaviorTree/BTNode.h
+++ b/Engine/Source/Runtime/AIModule/Classes/BehaviorTree/BTNode.h
@@ -140,8 +140,6 @@ class AIMODULE_API UBTNode : public UObject
 	/** Gets called only for instanced nodes(bCreateNodeInstance == true). In practive overriden by BP-implemented BT nodes */
 	virtual void SetOwner(AActor* ActorOwner) {}
 
-protected:
-
 	/** node name */
 	UPROPERTY(Category=Description, EditAnywhere)
 	FString NodeName;
@@ -225,11 +223,6 @@ FORCEINLINE UBTNode* UBTNode::GetNextNode() const
 }
 #endif
 
-FORCEINLINE FString UBTNode::GetNodeName() const
-{
-	return NodeName.Len() ? NodeName : GetClass()->GetName();
-}
-
 FORCEINLINE uint16 UBTNode::GetExecutionIndex() const
 {
 	return ExecutionIndex;
diff --git a/Engine/Source/Runtime/AIModule/Classes/BehaviorTree/Services/BTService_BlueprintBase.h b/Engine/Source/Runtime/AIModule/Classes/BehaviorTree/Services/BTService_BlueprintBase.h
--- a/Engine/Source/Runtime/AIModule/Classes/BehaviorTree/Services/BTService_BlueprintBase.h
+++ b/Engine/Source/Runtime/AIModule/Classes/BehaviorTree/Services/BTService_BlueprintBase.h
@@ -19,8 +19,6 @@ class AIMODULE_API UBTService_BlueprintBase : public UBTService
 {
 	GENERATED_UCLASS_BODY()
 
-	virtual void PostInitProperties() override;
-
 	virtual void DescribeRuntimeValues(const UBehaviorTreeComponent& OwnerComp, uint8* NodeMemory, EBTDescriptionVerbosity::Type Verbosity, TArray<FString>& Values) const override;
 	virtual void OnInstanceDestroyed(UBehaviorTreeComponent& OwnerComp) override;
 
diff --git a/Engine/Source/Runtime/AIModule/Classes/BehaviorTree/Tasks/BTTask_BlueprintBase.h b/Engine/Source/Runtime/AIModule/Classes/BehaviorTree/Tasks/BTTask_BlueprintBase.h
--- a/Engine/Source/Runtime/AIModule/Classes/BehaviorTree/Tasks/BTTask_BlueprintBase.h
+++ b/Engine/Source/Runtime/AIModule/Classes/BehaviorTree/Tasks/BTTask_BlueprintBase.h
@@ -19,10 +19,7 @@ UCLASS(Abstract, Blueprintable)
 class AIMODULE_API UBTTask_BlueprintBase : public UBTTaskNode
 {
 	GENERATED_UCLASS_BODY()
-
-	/** setup node name */
-	virtual void PostInitProperties() override;
-
+	
 	virtual EBTNodeResult::Type ExecuteTask(UBehaviorTreeComponent& OwnerComp, uint8* NodeMemory) override;
 	virtual EBTNodeResult::Type AbortTask(UBehaviorTreeComponent& OwnerComp, uint8* NodeMemory) override;
 
diff --git a/Engine/Source/Runtime/AIModule/Private/BehaviorTree/BTNode.cpp b/Engine/Source/Runtime/AIModule/Private/BehaviorTree/BTNode.cpp
--- a/Engine/Source/Runtime/AIModule/Private/BehaviorTree/BTNode.cpp
+++ b/Engine/Source/Runtime/AIModule/Private/BehaviorTree/BTNode.cpp
@@ -11,7 +11,6 @@
 
 UBTNode::UBTNode(const FObjectInitializer& ObjectInitializer) : Super(ObjectInitializer)
 {
-	NodeName = "";
 	ParentNode = NULL;
 	TreeAsset = NULL;
 	ExecutionIndex = 0;
@@ -143,6 +142,11 @@ UBTNode* UBTNode::GetNodeInstance(FBehaviorTreeSearchData& SearchData) const
 	return GetNodeInstance(SearchData.OwnerComp, GetNodeMemory<uint8>(SearchData));
 }
 
+FString UBTNode::GetNodeName() const
+{
+	return NodeName.Len() ? NodeName : UBehaviorTreeTypes::GetShortTypeName(this);
+}
+
 FString UBTNode::GetRuntimeDescription(const UBehaviorTreeComponent& OwnerComp, uint8* NodeMemory, EBTDescriptionVerbosity::Type Verbosity) const
 {
 	FString Description = NodeName.Len() ? FString::Printf(TEXT("%d. %s [%s]"), ExecutionIndex, *NodeName, *GetStaticDescription()) : GetStaticDescription();
diff --git a/Engine/Source/Runtime/AIModule/Private/BehaviorTree/Decorators/BTDecorator_BlueprintBase.cpp b/Engine/Source/Runtime/AIModule/Private/BehaviorTree/Decorators/BTDecorator_BlueprintBase.cpp
--- a/Engine/Source/Runtime/AIModule/Private/BehaviorTree/Decorators/BTDecorator_BlueprintBase.cpp
+++ b/Engine/Source/Runtime/AIModule/Private/BehaviorTree/Decorators/BTDecorator_BlueprintBase.cpp
@@ -38,8 +38,6 @@ void UBTDecorator_BlueprintBase::InitializeProperties()
 		BlueprintNodeHelpers::CollectPropertyData(this, StopAtClass, PropertyData);
 
 		bIsObservingBB = BlueprintNodeHelpers::HasAnyBlackboardSelectors(this, StopAtClass);
-
-		NodeName = BlueprintNodeHelpers::GetNodeName(this);
 	}
 }
 
diff --git a/Engine/Source/Runtime/AIModule/Private/BehaviorTree/Services/BTService_BlueprintBase.cpp b/Engine/Source/Runtime/AIModule/Private/BehaviorTree/Services/BTService_BlueprintBase.cpp
--- a/Engine/Source/Runtime/AIModule/Private/BehaviorTree/Services/BTService_BlueprintBase.cpp
+++ b/Engine/Source/Runtime/AIModule/Private/BehaviorTree/Services/BTService_BlueprintBase.cpp
@@ -27,12 +27,6 @@ UBTService_BlueprintBase::UBTService_BlueprintBase(const FObjectInitializer& Obj
 	}
 }
 
-void UBTService_BlueprintBase::PostInitProperties()
-{
-	Super::PostInitProperties();
-	NodeName = BlueprintNodeHelpers::GetNodeName(this);
-}
-
 void UBTService_BlueprintBase::SetOwner(AActor* InActorOwner)
 {
 	ActorOwner = InActorOwner;
diff --git a/Engine/Source/Runtime/AIModule/Private/BehaviorTree/Tasks/BTTask_BlueprintBase.cpp b/Engine/Source/Runtime/AIModule/Private/BehaviorTree/Tasks/BTTask_BlueprintBase.cpp
--- a/Engine/Source/Runtime/AIModule/Private/BehaviorTree/Tasks/BTTask_BlueprintBase.cpp
+++ b/Engine/Source/Runtime/AIModule/Private/BehaviorTree/Tasks/BTTask_BlueprintBase.cpp
@@ -23,12 +23,6 @@ UBTTask_BlueprintBase::UBTTask_BlueprintBase(const FObjectInitializer& ObjectIni
 	}
 }
 
-void UBTTask_BlueprintBase::PostInitProperties()
-{
-	Super::PostInitProperties();
-	NodeName = BlueprintNodeHelpers::GetNodeName(this);
-}
-
 void UBTTask_BlueprintBase::SetOwner(AActor* InActorOwner) 
 { 
 	ActorOwner = InActorOwner;
