Fixed editor-time changes to parent BB asset not propagating to derived BB assets #UE4

UE-13833

diff --git a/Engine/Source/Runtime/AIModule/Classes/BehaviorTree/BlackboardData.h b/Engine/Source/Runtime/AIModule/Classes/BehaviorTree/BlackboardData.h
--- a/Engine/Source/Runtime/AIModule/Classes/BehaviorTree/BlackboardData.h
+++ b/Engine/Source/Runtime/AIModule/Classes/BehaviorTree/BlackboardData.h
@@ -84,6 +84,7 @@ public:
 #if WITH_EDITOR
 	virtual void PreEditChange(class FEditPropertyChain& PropertyAboutToChange) override;
 	virtual void PostEditChangeProperty(struct FPropertyChangedEvent& PropertyChangedEvent) override;
+	void PropagateKeyChangesToDerivedBlackboardAssets();
 #endif
 
 	/** @return true if blackboard keys are not conflicting with parent key chain */
diff --git a/Engine/Source/Runtime/AIModule/Private/BehaviorTree/BlackboardData.cpp b/Engine/Source/Runtime/AIModule/Private/BehaviorTree/BlackboardData.cpp
--- a/Engine/Source/Runtime/AIModule/Private/BehaviorTree/BlackboardData.cpp
+++ b/Engine/Source/Runtime/AIModule/Private/BehaviorTree/BlackboardData.cpp
@@ -134,6 +134,7 @@ void UBlackboardData::PostEditChangeProperty(FPropertyChangedEvent& PropertyChan
 {
 	static const FName NAME_Parent = GET_MEMBER_NAME_CHECKED(UBlackboardData, Parent);
 	static const FName NAME_InstanceSynced = GET_MEMBER_NAME_CHECKED(FBlackboardEntry, bInstanceSynced);
+	static const FName NAME_Keys = GET_MEMBER_NAME_CHECKED(UBlackboardData, Keys);
 
 	Super::PostEditChangeProperty(PropertyChangedEvent);
 
@@ -158,6 +159,27 @@ void UBlackboardData::PostEditChangeProperty(FPropertyChangedEvent& PropertyChan
 			UpdateIfHasSynchronizedKeys();
 		}
 	}
+	if (PropertyChangedEvent.MemberProperty)
+	{
+		if (PropertyChangedEvent.MemberProperty->GetFName() == NAME_Keys)
+		{
+			// look for BB assets using this one as a parent and update them as well
+			PropagateKeyChangesToDerivedBlackboardAssets();
+		}
+	}
+}
+
+void UBlackboardData::PropagateKeyChangesToDerivedBlackboardAssets()
+{
+	for (TObjectIterator<UBlackboardData> It; It; ++It)
+	{
+		if (It->Parent == this)
+		{
+			It->UpdateParentKeys();
+			It->UpdateIfHasSynchronizedKeys();
+			It->PropagateKeyChangesToDerivedBlackboardAssets();
+		}
+	}
 }
 
 #endif // WITH_EDITOR
