Added a safety check to not crash if BTTaskNode, that's no longer part of active BT, gets a notification via UBTTaskNode::ReceivedMessage call #UE4

diff --git a/Engine/Source/Runtime/AIModule/Private/BehaviorTree/BTTaskNode.cpp b/Engine/Source/Runtime/AIModule/Private/BehaviorTree/BTTaskNode.cpp
--- a/Engine/Source/Runtime/AIModule/Private/BehaviorTree/BTTaskNode.cpp
+++ b/Engine/Source/Runtime/AIModule/Private/BehaviorTree/BTTaskNode.cpp
@@ -58,11 +58,20 @@ void UBTTaskNode::WrappedOnTaskFinished(UBehaviorTreeComponent& OwnerComp, uint8
 
 void UBTTaskNode::ReceivedMessage(UBrainComponent* BrainComp, const FAIMessage& Message)
 {
-	UBehaviorTreeComponent* OwnerComp = (UBehaviorTreeComponent*)BrainComp;
-	uint16 InstanceIdx = OwnerComp->FindInstanceContainingNode(this);
-	uint8* NodeMemory = GetNodeMemory<uint8>(OwnerComp->InstanceStack[InstanceIdx]);
-
-	OnMessage(*OwnerComp, NodeMemory, Message.MessageName, Message.RequestID, Message.Status == FAIMessage::Success);
+	UBehaviorTreeComponent* OwnerComp = static_cast<UBehaviorTreeComponent*>(BrainComp);
+	check(OwnerComp);
+	
+	const uint16 InstanceIdx = OwnerComp->FindInstanceContainingNode(this);
+	if (InstanceIdx)
+	{
+		uint8* NodeMemory = GetNodeMemory<uint8>(OwnerComp->InstanceStack[InstanceIdx]);
+		OnMessage(*OwnerComp, NodeMemory, Message.MessageName, Message.RequestID, Message.Status == FAIMessage::Success);
+	}
+	else
+	{
+		UE_VLOG(OwnerComp->GetOwner(), LogBehaviorTree, Warning, TEXT("UBTTaskNode::ReceivedMessage called while %s node no longer in active BT")
+			, *GetNodeName());
+	}
 }
 
 void UBTTaskNode::TickTask(UBehaviorTreeComponent& OwnerComp, uint8* NodeMemory, float DeltaSeconds)
