fixed crash in pathfinding batch EQS test
#ue4, merged from dev

diff --git a/Engine/Source/Runtime/AIModule/Private/EnvironmentQuery/Tests/EnvQueryTest_PathfindingBatch.cpp b/Engine/Source/Runtime/AIModule/Private/EnvironmentQuery/Tests/EnvQueryTest_PathfindingBatch.cpp
--- a/Engine/Source/Runtime/AIModule/Private/EnvironmentQuery/Tests/EnvQueryTest_PathfindingBatch.cpp
+++ b/Engine/Source/Runtime/AIModule/Private/EnvironmentQuery/Tests/EnvQueryTest_PathfindingBatch.cpp
@@ -22,29 +22,29 @@ UEnvQueryTest_PathfindingBatch::UEnvQueryTest_PathfindingBatch(const FObjectInit
 #if WITH_RECAST
 namespace NodePoolHelpers
 {
-	static bool HasPath(const FRecastDebugPathfindingData& NodePool, const FNavigationProjectionWork& TestPt)
+	static bool HasPath(const FRecastDebugPathfindingData& NodePool, const FNavigationProjectionWork* TestPt)
 	{
-		FRecastDebugPathfindingNode SearchKey(TestPt.OutLocation.NodeRef);
+		FRecastDebugPathfindingNode SearchKey(TestPt->OutLocation.NodeRef);
 		const FRecastDebugPathfindingNode* MyNode = NodePool.Nodes.Find(SearchKey);
 		return MyNode != nullptr;
 	}
 
-	static float GetPathLength(const FRecastDebugPathfindingData& NodePool, const FNavigationProjectionWork& TestPt, const dtQueryFilter* Filter)
+	static float GetPathLength(const FRecastDebugPathfindingData& NodePool, const FNavigationProjectionWork* TestPt, const dtQueryFilter* Filter)
 	{
-		FRecastDebugPathfindingNode SearchKey(TestPt.OutLocation.NodeRef);
+		FRecastDebugPathfindingNode SearchKey(TestPt->OutLocation.NodeRef);
 		const FRecastDebugPathfindingNode* MyNode = NodePool.Nodes.Find(SearchKey);
 		if (MyNode)
 		{
-			float LastSegmentLength = FVector::Dist(MyNode->NodePos, TestPt.OutLocation.Location);
+			float LastSegmentLength = FVector::Dist(MyNode->NodePos, TestPt->OutLocation.Location);
 			return MyNode->Length + LastSegmentLength;
 		}
 
 		return BIG_NUMBER;
 	}
 
-	static float GetPathCost(const FRecastDebugPathfindingData& NodePool, const FNavigationProjectionWork& TestPt, const dtQueryFilter* Filter)
+	static float GetPathCost(const FRecastDebugPathfindingData& NodePool, const FNavigationProjectionWork* TestPt, const dtQueryFilter* Filter)
 	{
-		FRecastDebugPathfindingNode SearchKey(TestPt.OutLocation.NodeRef);
+		FRecastDebugPathfindingNode SearchKey(TestPt->OutLocation.NodeRef);
 		const FRecastDebugPathfindingNode* MyNode = NodePool.Nodes.Find(SearchKey);
 		if (MyNode)
 		{
@@ -54,7 +54,7 @@ namespace NodePoolHelpers
 		return BIG_NUMBER;
 	}
 
-	typedef float(*PathParamFunc)(const FRecastDebugPathfindingData&, const FNavigationProjectionWork&, const dtQueryFilter*);
+	typedef float(*PathParamFunc)(const FRecastDebugPathfindingData&, const FNavigationProjectionWork*, const dtQueryFilter*);
 }
 
 void UEnvQueryTest_PathfindingBatch::RunTest(FEnvQueryInstance& QueryInstance) const
@@ -142,15 +142,16 @@ void UEnvQueryTest_PathfindingBatch::RunTest(FEnvQueryInstance& QueryInstance) c
 		}
 	}
 
+	FNavigationProjectionWork* ProjectedPtr = TestPoints.GetTypedData();
 	if (GetWorkOnFloatValues())
 	{
 		NodePoolHelpers::PathParamFunc Func[] = { nullptr, NodePoolHelpers::GetPathCost, NodePoolHelpers::GetPathLength };
 		FEnvQueryInstance::ItemIterator It(this, QueryInstance);
-		for (It.IgnoreTimeLimit(); It; ++It)
+		for (It.IgnoreTimeLimit(); It; ++It, ProjectedPtr++)
 		{
 			for (int32 ContextIndex = 0; ContextIndex < ContextLocations.Num(); ContextIndex++)
 			{
-				const float PathValue = Func[TestMode](NodePoolData[ContextIndex], TestPoints[*It], NavQueryFilter);
+				const float PathValue = Func[TestMode](NodePoolData[ContextIndex], ProjectedPtr, NavQueryFilter);
 				It.SetScore(TestPurpose, FilterType, PathValue, MinThresholdValue, MaxThresholdValue);
 
 				if (bDiscardFailed && PathValue >= BIG_NUMBER)
@@ -163,11 +164,11 @@ void UEnvQueryTest_PathfindingBatch::RunTest(FEnvQueryInstance& QueryInstance) c
 	else
 	{
 		FEnvQueryInstance::ItemIterator It(this, QueryInstance);
-		for (It.IgnoreTimeLimit(); It; ++It)
+		for (It.IgnoreTimeLimit(); It; ++It, ProjectedPtr++)
 		{
 			for (int32 ContextIndex = 0; ContextIndex < ContextLocations.Num(); ContextIndex++)
 			{
-				const bool bFoundPath = NodePoolHelpers::HasPath(NodePoolData[ContextIndex], TestPoints[*It]);
+				const bool bFoundPath = NodePoolHelpers::HasPath(NodePoolData[ContextIndex], ProjectedPtr);
 				It.SetScore(TestPurpose, FilterType, bFoundPath, bWantsPath);
 			}
 		}
