EQS item iterator will now use FORCEINLINE_EXPLICIT_OPERATOR_BOOL
#ue4, merged from dev

diff --git a/Engine/Source/Runtime/AIModule/Classes/EnvironmentQuery/EnvQueryTypes.h b/Engine/Source/Runtime/AIModule/Classes/EnvironmentQuery/EnvQueryTypes.h
--- a/Engine/Source/Runtime/AIModule/Classes/EnvironmentQuery/EnvQueryTypes.h
+++ b/Engine/Source/Runtime/AIModule/Classes/EnvironmentQuery/EnvQueryTypes.h
@@ -822,7 +822,7 @@ public:
 			Deadline = -1.0f;
 		}
 
-		operator bool() const
+		FORCEINLINE_EXPLICIT_OPERATOR_BOOL() const
 		{
 			return CurrentItem < Instance->Items.Num() && !Instance->bFoundSingleResult && (Deadline < 0 || FPlatformTime::Seconds() < Deadline);
 		}
