replaced operator * from EQS item iterators with new function: GetIndex()
#ue4

diff --git a/Engine/Source/Runtime/AIModule/Classes/EnvironmentQuery/EnvQueryTest.h b/Engine/Source/Runtime/AIModule/Classes/EnvironmentQuery/EnvQueryTest.h
--- a/Engine/Source/Runtime/AIModule/Classes/EnvironmentQuery/EnvQueryTest.h
+++ b/Engine/Source/Runtime/AIModule/Classes/EnvironmentQuery/EnvQueryTest.h
@@ -142,7 +142,7 @@ public:
 	/** helper: get location of item */
 	FORCEINLINE FVector GetItemLocation(FEnvQueryInstance& QueryInstance, const FEnvQueryInstance::ItemIterator& Iterator) const
 	{
-		return GetItemLocation(QueryInstance, *Iterator);
+		return GetItemLocation(QueryInstance, Iterator.GetIndex());
 	}
 
 	/** helper: get location of item */
@@ -151,7 +151,7 @@ public:
 	/** helper: get location of item */
 	FORCEINLINE FRotator GetItemRotation(FEnvQueryInstance& QueryInstance, const FEnvQueryInstance::ItemIterator& Iterator) const
 	{
-		return GetItemRotation(QueryInstance, *Iterator);
+		return GetItemRotation(QueryInstance, Iterator.GetIndex());
 	}
 
 	/** helper: get actor from item */
@@ -160,7 +160,7 @@ public:
 	/** helper: get actor from item */
 	FORCEINLINE AActor* GetItemActor(FEnvQueryInstance& QueryInstance, const FEnvQueryInstance::ItemIterator& Iterator) const
 	{
-		return GetItemActor(QueryInstance, *Iterator);
+		return GetItemActor(QueryInstance, Iterator.GetIndex());
 	}
 
 	/** normalize scores in range */
diff --git a/Engine/Source/Runtime/AIModule/Classes/EnvironmentQuery/EnvQueryTypes.h b/Engine/Source/Runtime/AIModule/Classes/EnvironmentQuery/EnvQueryTypes.h
--- a/Engine/Source/Runtime/AIModule/Classes/EnvironmentQuery/EnvQueryTypes.h
+++ b/Engine/Source/Runtime/AIModule/Classes/EnvironmentQuery/EnvQueryTypes.h
@@ -822,14 +822,20 @@ public:
 			Deadline = -1.0f;
 		}
 
-		FORCEINLINE_EXPLICIT_OPERATOR_BOOL() const
+		int32 GetIndex() const
 		{
-			return CurrentItem < Instance->Items.Num() && !Instance->bFoundSingleResult && (Deadline < 0 || FPlatformTime::Seconds() < Deadline);
+			return CurrentItem;
 		}
 
+		DEPRECATED(4.8, "This function is now deprecatewd, please use GetIndex() for current index or GetItemData() for raw data pointer")
 		int32 operator*() const
 		{
-			return CurrentItem; 
+			return GetIndex();
+		}
+
+		FORCEINLINE_EXPLICIT_OPERATOR_BOOL() const
+		{
+			return CurrentItem < Instance->Items.Num() && !Instance->bFoundSingleResult && (Deadline < 0 || FPlatformTime::Seconds() < Deadline);
 		}
 
 		void operator++()
diff --git a/Engine/Source/Runtime/AIModule/Private/EnvironmentQuery/Tests/EnvQueryTest_Distance.cpp b/Engine/Source/Runtime/AIModule/Private/EnvironmentQuery/Tests/EnvQueryTest_Distance.cpp
--- a/Engine/Source/Runtime/AIModule/Private/EnvironmentQuery/Tests/EnvQueryTest_Distance.cpp
+++ b/Engine/Source/Runtime/AIModule/Private/EnvironmentQuery/Tests/EnvQueryTest_Distance.cpp
@@ -50,7 +50,7 @@ void UEnvQueryTest_Distance::RunTest(FEnvQueryInstance& QueryInstance) const
 		case EEnvTestDistance::Distance3D:	
 			for (FEnvQueryInstance::ItemIterator It(this, QueryInstance); It; ++It)
 			{
-				const FVector ItemLocation = GetItemLocation(QueryInstance, *It);
+				const FVector ItemLocation = GetItemLocation(QueryInstance, It.GetIndex());
 				for (int32 ContextIndex = 0; ContextIndex < ContextLocations.Num(); ContextIndex++)
 				{
 					const float Distance = CalcDistance3D(ItemLocation, ContextLocations[ContextIndex]);
@@ -62,7 +62,7 @@ void UEnvQueryTest_Distance::RunTest(FEnvQueryInstance& QueryInstance) const
 		case EEnvTestDistance::Distance2D:	
 			for (FEnvQueryInstance::ItemIterator It(this, QueryInstance); It; ++It)
 			{
-				const FVector ItemLocation = GetItemLocation(QueryInstance, *It);
+				const FVector ItemLocation = GetItemLocation(QueryInstance, It.GetIndex());
 				for (int32 ContextIndex = 0; ContextIndex < ContextLocations.Num(); ContextIndex++)
 				{
 					const float Distance = CalcDistance2D(ItemLocation, ContextLocations[ContextIndex]);
@@ -74,7 +74,7 @@ void UEnvQueryTest_Distance::RunTest(FEnvQueryInstance& QueryInstance) const
 		case EEnvTestDistance::DistanceZ:	
 			for (FEnvQueryInstance::ItemIterator It(this, QueryInstance); It; ++It)
 			{
-				const FVector ItemLocation = GetItemLocation(QueryInstance, *It);
+				const FVector ItemLocation = GetItemLocation(QueryInstance, It.GetIndex());
 				for (int32 ContextIndex = 0; ContextIndex < ContextLocations.Num(); ContextIndex++)
 				{
 					const float Distance = CalcDistanceZ(ItemLocation, ContextLocations[ContextIndex]);
diff --git a/Engine/Source/Runtime/AIModule/Private/EnvironmentQuery/Tests/EnvQueryTest_Dot.cpp b/Engine/Source/Runtime/AIModule/Private/EnvironmentQuery/Tests/EnvQueryTest_Dot.cpp
--- a/Engine/Source/Runtime/AIModule/Private/EnvironmentQuery/Tests/EnvQueryTest_Dot.cpp
+++ b/Engine/Source/Runtime/AIModule/Private/EnvironmentQuery/Tests/EnvQueryTest_Dot.cpp
@@ -57,8 +57,8 @@ void UEnvQueryTest_Dot::RunTest(FEnvQueryInstance& QueryInstance) const
 		// update lines for contexts using current item
 		if (bUpdateLineAPerItem || bUpdateLineBPerItem)
 		{
-			const FVector ItemLocation = (LineA.DirMode == EEnvDirection::Rotation && LineB.DirMode == EEnvDirection::Rotation) ? FVector::ZeroVector : GetItemLocation(QueryInstance, *It);
-			const FRotator ItemRotation = (LineA.DirMode == EEnvDirection::Rotation || LineB.DirMode == EEnvDirection::Rotation) ? GetItemRotation(QueryInstance, *It) : FRotator::ZeroRotator;
+			const FVector ItemLocation = (LineA.DirMode == EEnvDirection::Rotation && LineB.DirMode == EEnvDirection::Rotation) ? FVector::ZeroVector : GetItemLocation(QueryInstance, It.GetIndex());
+			const FRotator ItemRotation = (LineA.DirMode == EEnvDirection::Rotation || LineB.DirMode == EEnvDirection::Rotation) ? GetItemRotation(QueryInstance, It.GetIndex()) : FRotator::ZeroRotator;
 
 			if (bUpdateLineAPerItem)
 			{
diff --git a/Engine/Source/Runtime/AIModule/Private/EnvironmentQuery/Tests/EnvQueryTest_GameplayTags.cpp b/Engine/Source/Runtime/AIModule/Private/EnvironmentQuery/Tests/EnvQueryTest_GameplayTags.cpp
--- a/Engine/Source/Runtime/AIModule/Private/EnvironmentQuery/Tests/EnvQueryTest_GameplayTags.cpp
+++ b/Engine/Source/Runtime/AIModule/Private/EnvironmentQuery/Tests/EnvQueryTest_GameplayTags.cpp
@@ -58,7 +58,7 @@ void UEnvQueryTest_GameplayTags::RunTest(FEnvQueryInstance& QueryInstance) const
 	// loop through all items
 	for (FEnvQueryInstance::ItemIterator It(this, QueryInstance); It; ++It)
 	{
-		AActor* ItemActor = GetItemActor(QueryInstance, *It);
+		AActor* ItemActor = GetItemActor(QueryInstance, It.GetIndex());
 		IGameplayTagAssetInterface* GameplayTagAssetInterface = Cast<IGameplayTagAssetInterface>(ItemActor);
 		if (GameplayTagAssetInterface != NULL)
 		{
diff --git a/Engine/Source/Runtime/AIModule/Private/EnvironmentQuery/Tests/EnvQueryTest_Pathfinding.cpp b/Engine/Source/Runtime/AIModule/Private/EnvironmentQuery/Tests/EnvQueryTest_Pathfinding.cpp
--- a/Engine/Source/Runtime/AIModule/Private/EnvironmentQuery/Tests/EnvQueryTest_Pathfinding.cpp
+++ b/Engine/Source/Runtime/AIModule/Private/EnvironmentQuery/Tests/EnvQueryTest_Pathfinding.cpp
@@ -67,7 +67,7 @@ void UEnvQueryTest_Pathfinding::RunTest(FEnvQueryInstance& QueryInstance) const
 		NavData->BeginBatchQuery();
 		for (FEnvQueryInstance::ItemIterator It(this, QueryInstance); It; ++It)
 		{
-			const FVector ItemLocation = GetItemLocation(QueryInstance, *It);
+			const FVector ItemLocation = GetItemLocation(QueryInstance, It.GetIndex());
 			for (int32 ContextIndex = 0; ContextIndex < ContextLocations.Num(); ContextIndex++)
 			{
 				const float PathValue = FindPathFunc.Execute(ItemLocation, ContextLocations[ContextIndex], PFMode, *NavData, *NavSys, QueryInstance.Owner.Get());
@@ -88,7 +88,7 @@ void UEnvQueryTest_Pathfinding::RunTest(FEnvQueryInstance& QueryInstance) const
 		{
 			for (FEnvQueryInstance::ItemIterator It(this, QueryInstance); It; ++It)
 			{
-				const FVector ItemLocation = GetItemLocation(QueryInstance, *It);
+				const FVector ItemLocation = GetItemLocation(QueryInstance, It.GetIndex());
 				for (int32 ContextIndex = 0; ContextIndex < ContextLocations.Num(); ContextIndex++)
 				{
 					const bool bFoundPath = TestPathTo(ItemLocation, ContextLocations[ContextIndex], PFMode, *NavData, *NavSys, QueryInstance.Owner.Get());
@@ -100,7 +100,7 @@ void UEnvQueryTest_Pathfinding::RunTest(FEnvQueryInstance& QueryInstance) const
 		{
 			for (FEnvQueryInstance::ItemIterator It(this, QueryInstance); It; ++It)
 			{
-				const FVector ItemLocation = GetItemLocation(QueryInstance, *It);
+				const FVector ItemLocation = GetItemLocation(QueryInstance, It.GetIndex());
 				for (int32 ContextIndex = 0; ContextIndex < ContextLocations.Num(); ContextIndex++)
 				{
 					const bool bFoundPath = TestPathFrom(ItemLocation, ContextLocations[ContextIndex], PFMode, *NavData, *NavSys, QueryInstance.Owner.Get());
diff --git a/Engine/Source/Runtime/AIModule/Private/EnvironmentQuery/Tests/EnvQueryTest_Trace.cpp b/Engine/Source/Runtime/AIModule/Private/EnvironmentQuery/Tests/EnvQueryTest_Trace.cpp
--- a/Engine/Source/Runtime/AIModule/Private/EnvironmentQuery/Tests/EnvQueryTest_Trace.cpp
+++ b/Engine/Source/Runtime/AIModule/Private/EnvironmentQuery/Tests/EnvQueryTest_Trace.cpp
@@ -77,8 +77,8 @@ void UEnvQueryTest_Trace::RunTest(FEnvQueryInstance& QueryInstance) const
 
 	for (FEnvQueryInstance::ItemIterator It(this, QueryInstance); It; ++It)
 	{
-		const FVector ItemLocation = GetItemLocation(QueryInstance, *It) + FVector(0,0,ItemZ);
-		AActor* ItemActor = GetItemActor(QueryInstance, *It);
+		const FVector ItemLocation = GetItemLocation(QueryInstance, It.GetIndex()) + FVector(0, 0, ItemZ);
+		AActor* ItemActor = GetItemActor(QueryInstance, It.GetIndex());
 
 		for (int32 ContextIndex = 0; ContextIndex < ContextLocations.Num(); ContextIndex++)
 		{
