Fixed a bug in BT task's message hanging that would result in tasks waiting for messages never end #UE4

diff --git a/Engine/Source/Runtime/AIModule/Private/BehaviorTree/BTTaskNode.cpp b/Engine/Source/Runtime/AIModule/Private/BehaviorTree/BTTaskNode.cpp
--- a/Engine/Source/Runtime/AIModule/Private/BehaviorTree/BTTaskNode.cpp
+++ b/Engine/Source/Runtime/AIModule/Private/BehaviorTree/BTTaskNode.cpp
@@ -62,7 +62,7 @@ void UBTTaskNode::ReceivedMessage(UBrainComponent* BrainComp, const FAIMessage&
 	check(OwnerComp);
 	
 	const uint16 InstanceIdx = OwnerComp->FindInstanceContainingNode(this);
-	if (InstanceIdx)
+	if (OwnerComp->InstanceStack.IsValidIndex(InstanceIdx))
 	{
 		uint8* NodeMemory = GetNodeMemory<uint8>(OwnerComp->InstanceStack[InstanceIdx]);
 		OnMessage(*OwnerComp, NodeMemory, Message.MessageName, Message.RequestID, Message.Status == FAIMessage::Success);
