A bunch of stabilization fixes to AIPerceptionSystem #UE4

* Fixed a bug where registering the very first source for yet not instanced AI sense would result in registering that source with an invalid FAISenseID
* added handling of registered stimuli source actors ending play
    - note: this will get updated once we get instigators with OnEndPlay broadcasts
* Depracated FPerceptionChannelFilter and renamed it to a more precise FPerceptionChannelWhitelist
* fixed ReportPerceptionEvent not hiding the context pin

diff --git a/Engine/Source/Runtime/AIModule/Classes/Perception/AIPerceptionComponent.h b/Engine/Source/Runtime/AIModule/Classes/Perception/AIPerceptionComponent.h
--- a/Engine/Source/Runtime/AIModule/Classes/Perception/AIPerceptionComponent.h
+++ b/Engine/Source/Runtime/AIModule/Classes/Perception/AIPerceptionComponent.h
@@ -147,7 +147,7 @@ protected:
 	UPROPERTY(Transient)
 	AAIController* AIOwner;
 
-	FPerceptionChannelFilter PerceptionFilter;
+	FPerceptionChannelWhitelist PerceptionFilter;
 
 private:
 	TActorPerceptionContainer PerceptualData;
@@ -185,8 +185,9 @@ public:
 	
 	void GetLocationAndDirection(FVector& Location, FVector& Direction) const;
 	const AActor* GetBodyActor() const;
+	AActor* GetMutableBodyActor();
 
-	FORCEINLINE const FPerceptionChannelFilter GetPerceptionFilter() const { return PerceptionFilter; }
+	FORCEINLINE const FPerceptionChannelWhitelist GetPerceptionFilter() const { return PerceptionFilter; }
 
 	FGenericTeamId GetTeamIdentifier() const;
 	FORCEINLINE FPerceptionListenerID GetListenerId() const { return PerceptionListenerId; }
diff --git a/Engine/Source/Runtime/AIModule/Classes/Perception/AIPerceptionSystem.h b/Engine/Source/Runtime/AIModule/Classes/Perception/AIPerceptionSystem.h
--- a/Engine/Source/Runtime/AIModule/Classes/Perception/AIPerceptionSystem.h
+++ b/Engine/Source/Runtime/AIModule/Classes/Perception/AIPerceptionSystem.h
@@ -42,11 +42,19 @@ protected:
 	UPROPERTY(config, EditAnywhere, Category = Perception)
 	float PerceptionAgingRate;
 
+	FActorEndPlaySignature::FDelegate StimuliSourceEndPlayDelegate;
+
+	// not a UPROPERTY on purpose so that we have a control over when stuff gets removed from the map
+	TMap<const AActor*, FPerceptionStimuliSource> RegisteredStimuliSources;
+
 	/** gets set to true if as as result of stimuli aging (that's done outside of Tick, on timer)
 	 *	one of listeners requires an update. The update, as usual is tone in Tick where 
 	 *	bSomeListenersNeedUpdateDueToStimuliAging gets reset to false */
 	uint32 bSomeListenersNeedUpdateDueToStimuliAging : 1;
 
+	/** gets set to true when perception system gets notified about a stimuli source's end play */
+	uint32 bStimuliSourcesRefreshRequired : 1;
+
 	struct FDelayedStimulus
 	{
 		float DeliveryTimestamp;
@@ -65,6 +73,11 @@ protected:
 		FPerceptionSourceRegistration(FAISenseID InSenseID, AActor* SourceActor)
 			: SenseID(InSenseID), Source(SourceActor)
 		{}
+
+		FORCEINLINE bool operator==(const FPerceptionSourceRegistration& Other) const
+		{
+			return SenseID == Other.SenseID && Source == Other.Source;
+		}
 	};
 	TArray<FPerceptionSourceRegistration> SourcesToRegister;
 
@@ -103,27 +116,36 @@ public:
 	UFUNCTION(BlueprintCallable, Category = "AI|Perception")
 	void ReportEvent(UAISenseEvent* PerceptionEvent);
 
-	UFUNCTION(BlueprintCallable, Category = "AI|Perception", meta = (WorldContext = "WorldContext"))
+	UFUNCTION(BlueprintCallable, Category = "AI|Perception", meta = (HidePin = "WorldContext", DefaultToSelf = "WorldContext"))
 	static void ReportPerceptionEvent(UObject* WorldContext, UAISenseEvent* PerceptionEvent);
 	
 	template<typename FSenseClass>
 	void RegisterSource(AActor& SourceActor);
 
+	/** 
+	 *	unregisters given actor from the list of active stimuli sources
+	 */
+	void UnregisterSource(AActor& SourceActor);
+
 	void RegisterDelayedStimulus(FPerceptionListenerID ListenerId, float Delay, AActor* Instigator, const FAIStimulus& Stimulus);
 
 	static UAIPerceptionSystem* GetCurrent(UObject* WorldContextObject);
 
 	static void MakeNoiseImpl(AActor* NoiseMaker, float Loudness, APawn* NoiseInstigator, const FVector& NoiseLocation);
 
-	UFUNCTION(BlueprintCallable, Category = "AI|Perception", meta = (WorldContext = "WorldContext"))
+	UFUNCTION(BlueprintCallable, Category = "AI|Perception", meta = (HidePin = "WorldContext", DefaultToSelf = "WorldContext"))
 	static bool RegisterPerceptionStimuliSource(UObject* WorldContext, TSubclassOf<UAISense> Sense, AActor* Target);
 
-	void RegisterSenseClass(TSubclassOf<UAISense> SenseClass);
+	FAISenseID RegisterSenseClass(TSubclassOf<UAISense> SenseClass);
 
 	UFUNCTION(BlueprintCallable, Category = "AI|Perception")
 	static TSubclassOf<UAISense> GetSenseClassForStimulus(UObject* WorldContext, const FAIStimulus& Stimulus);
 	
 protected:
+
+	UFUNCTION()
+	void OnPerceptionStimuliSourceEndPlay(EEndPlayReason::Type EndPlayReason);
+	
 	/** requests registration of a given actor as a perception data source for specified sense */
 	void RegisterSource(FAISenseID SenseID, AActor& SourceActor);
 
diff --git a/Engine/Source/Runtime/AIModule/Classes/Perception/AIPerceptionTypes.h b/Engine/Source/Runtime/AIModule/Classes/Perception/AIPerceptionTypes.h
--- a/Engine/Source/Runtime/AIModule/Classes/Perception/AIPerceptionTypes.h
+++ b/Engine/Source/Runtime/AIModule/Classes/Perception/AIPerceptionTypes.h
@@ -32,14 +32,14 @@ enum class EAISenseNotifyType : uint8
 	OnPerceptionChange, // from "visible" to "not visible" or vice versa
 };
 
-struct FPerceptionChannelFilter
+struct FPerceptionChannelWhitelist
 {
-	static const uint32 AllChannels = uint32(-1);
+	typedef int32 FFlagsContainer;
 
-	uint32 AcceptedChannelsMask;
+	FFlagsContainer AcceptedChannelsMask;
 
 	// by default accept all
-	FPerceptionChannelFilter() : AcceptedChannelsMask(AllChannels)
+	FPerceptionChannelWhitelist() : AcceptedChannelsMask()
 	{}
 
 	void Clear()
@@ -47,13 +47,13 @@ struct FPerceptionChannelFilter
 		AcceptedChannelsMask = 0;
 	}
 
-	FORCEINLINE FPerceptionChannelFilter& FilterOutChannel(FAISenseID Channel)
+	FORCEINLINE FPerceptionChannelWhitelist& FilterOutChannel(FAISenseID Channel)
 	{
 		AcceptedChannelsMask &= ~(1 << Channel);
 		return *this;
 	}
 
-	FORCEINLINE FPerceptionChannelFilter& AcceptChannel(FAISenseID Channel)
+	FORCEINLINE_DEBUGGABLE FPerceptionChannelWhitelist& AcceptChannel(FAISenseID Channel)
 	{
 		AcceptedChannelsMask |= (1 << Channel);
 		return *this;
@@ -63,6 +63,68 @@ struct FPerceptionChannelFilter
 	{
 		return (AcceptedChannelsMask & (1 << Channel)) != 0;
 	}
+
+	FORCEINLINE FPerceptionChannelWhitelist& MergeFilterIn(const FPerceptionChannelWhitelist& OtherFilter)
+	{
+		AcceptedChannelsMask |= OtherFilter.AcceptedChannelsMask;
+		return *this;
+	}
+
+	FORCEINLINE FFlagsContainer GetAcceptedChannelsMask() const 
+	{ 
+		return AcceptedChannelsMask;
+	}
+
+	struct FConstIterator
+	{
+	private:
+		FFlagsContainer RemainingChannelsToTest;
+		const FPerceptionChannelWhitelist& Whitelist;
+		int32 CurrentIndex;
+
+	public:
+		FConstIterator(const FPerceptionChannelWhitelist& InWhitelist)
+			: RemainingChannelsToTest((FFlagsContainer)-1), Whitelist(InWhitelist)
+			, CurrentIndex(INDEX_NONE)
+		{
+			FindNextAcceptedChannel();
+		}
+
+		FORCEINLINE void FindNextAcceptedChannel()
+		{
+			const FFlagsContainer& Flags = Whitelist.GetAcceptedChannelsMask();
+
+			while ((RemainingChannelsToTest & Flags) != 0 && ((1 << ++CurrentIndex) | Flags) == 0)
+			{
+				RemainingChannelsToTest &= ~(1 << CurrentIndex);
+			}
+		}
+
+		FORCEINLINE_EXPLICIT_OPERATOR_BOOL() const
+		{
+			return (RemainingChannelsToTest & Whitelist.GetAcceptedChannelsMask()) != 0;
+		}
+
+		FORCEINLINE int32 operator*() const
+		{
+			return CurrentIndex;
+		}
+
+		FORCEINLINE void operator++()
+		{
+			// mark "old" index as already used
+			RemainingChannelsToTest &= ~(1 << CurrentIndex);
+			FindNextAcceptedChannel();
+		}
+	};
+};
+
+struct FPerceptionChannelFilter : public FPerceptionChannelWhitelist
+{
+	DEPRECATED(4.5, "FPerceptionChannelFilter has been renamed to FPerceptionChannelWhitelist. Please use that instead.")
+	FPerceptionChannelFilter()
+	{
+	}
 };
 
 USTRUCT(BlueprintType)
@@ -181,7 +243,7 @@ struct AIMODULE_API FPerceptionListener
 {
 	TWeakObjectPtr<UAIPerceptionComponent> Listener;
 
-	FPerceptionChannelFilter Filter;
+	FPerceptionChannelWhitelist Filter;
 
 	FVector CachedLocation;
 	FVector CachedDirection;
@@ -229,6 +291,12 @@ private:
 	FORCEINLINE void MarkForStimulusProcessing() { bHasStimulusToProcess = true; }
 };
 
+struct AIMODULE_API FPerceptionStimuliSource
+{
+	TWeakObjectPtr<AActor> SourceActor;
+	FPerceptionChannelWhitelist RelevantSenses;
+};
+
 namespace AIPerception
 {
 	typedef TMap<FPerceptionListenerID, FPerceptionListener> FListenerMap;
diff --git a/Engine/Source/Runtime/AIModule/Classes/Perception/AISense.h b/Engine/Source/Runtime/AIModule/Classes/Perception/AISense.h
--- a/Engine/Source/Runtime/AIModule/Classes/Perception/AISense.h
+++ b/Engine/Source/Runtime/AIModule/Classes/Perception/AISense.h
@@ -89,6 +89,10 @@ public:
 
 	//virtual void RegisterSources(TArray<AActor&> SourceActors) {}
 	virtual void RegisterSource(AActor& SourceActors){}
+	virtual void UnregisterSource(AActor& SourceActors){}
+	// @note this function should not be needed once AActor.OnEndPlay broadcast includes instigator as one of its params
+	// since implementations on this function would end up being more expensive then precise UnregisterSource calls
+	virtual void CleanseInvalidSources() {}
 	virtual void RegisterWrappedEvent(UAISenseEvent& PerceptionEvent);
 	virtual FAISenseID UpdateSenseID();
 
diff --git a/Engine/Source/Runtime/AIModule/Classes/Perception/AISense_Sight.h b/Engine/Source/Runtime/AIModule/Classes/Perception/AISense_Sight.h
--- a/Engine/Source/Runtime/AIModule/Classes/Perception/AISense_Sight.h
+++ b/Engine/Source/Runtime/AIModule/Classes/Perception/AISense_Sight.h
@@ -147,6 +147,8 @@ public:
 	void RegisterEvent(const FAISightEvent& Event);	
 
 	virtual void RegisterSource(AActor& SourceActors) override;
+	virtual void UnregisterSource(AActor& SourceActor) override;
+	virtual void CleanseInvalidSources() override;
 	
 protected:
 	virtual float Update() override;
@@ -165,7 +167,7 @@ protected:
 		Sort
 	};
 	void RemoveAllQueriesByListener(const FPerceptionListener& Listener, FQueriesOperationPostProcess PostProcess);
-	void RemoveAllQueriesToTarget(const FName& TargetId, FQueriesOperationPostProcess PostProcess);
+	void RemoveAllQueriesToTarget(const FAISightTarget::FTargetId& TargetId, FQueriesOperationPostProcess PostProcess);
 
 	/** returns information whether new LoS queries have been added */
 	bool RegisterTarget(AActor& TargetActor, FQueriesOperationPostProcess PostProcess);
diff --git a/Engine/Source/Runtime/AIModule/Private/Perception/AIPerceptionComponent.cpp b/Engine/Source/Runtime/AIModule/Private/Perception/AIPerceptionComponent.cpp
--- a/Engine/Source/Runtime/AIModule/Private/Perception/AIPerceptionComponent.cpp
+++ b/Engine/Source/Runtime/AIModule/Private/Perception/AIPerceptionComponent.cpp
@@ -195,7 +195,10 @@ void UAIPerceptionComponent::OnUnregister()
 
 void UAIPerceptionComponent::OnOwnerEndPlay(EEndPlayReason::Type EndPlayReason)
 {
-	CleanUp();
+	if (EndPlayReason != EEndPlayReason::EndPlayInEditor && EndPlayReason != EEndPlayReason::Quit)
+	{
+		CleanUp();
+	}
 }
 
 void UAIPerceptionComponent::CleanUp()
@@ -206,6 +209,11 @@ void UAIPerceptionComponent::CleanUp()
 		if (AIPerceptionSys != nullptr)
 		{
 			AIPerceptionSys->UnregisterListener(*this);
+			AActor* MutableBodyActor = GetMutableBodyActor();
+			if (MutableBodyActor)
+			{
+				AIPerceptionSys->UnregisterSource(*MutableBodyActor);
+			}
 		}
 
 		if (HasAnyFlags(RF_BeginDestroyed) == false)
@@ -354,6 +362,11 @@ const AActor* UAIPerceptionComponent::GetBodyActor() const
 	return Cast<AActor>(GetOuter());
 }
 
+AActor* UAIPerceptionComponent::GetMutableBodyActor()
+{
+	return const_cast<AActor*>(GetBodyActor());
+}
+
 void UAIPerceptionComponent::RegisterStimulus(AActor* Source, const FAIStimulus& Stimulus)
 {
 	StimuliToProcess.Add(FStimulusToProcess(Source, Stimulus));
diff --git a/Engine/Source/Runtime/AIModule/Private/Perception/AIPerceptionSystem.cpp b/Engine/Source/Runtime/AIModule/Private/Perception/AIPerceptionSystem.cpp
--- a/Engine/Source/Runtime/AIModule/Private/Perception/AIPerceptionSystem.cpp
+++ b/Engine/Source/Runtime/AIModule/Private/Perception/AIPerceptionSystem.cpp
@@ -31,11 +31,13 @@ UAIPerceptionSystem::UAIPerceptionSystem(const FObjectInitializer& ObjectInitial
 	: Super(ObjectInitializer)
 	, PerceptionAgingRate(0.3f)
 	, bSomeListenersNeedUpdateDueToStimuliAging(false)
+	, bStimuliSourcesRefreshRequired(false)
 	, CurrentTime(0.f)
 {
+	StimuliSourceEndPlayDelegate.BindDynamic(this, &UAIPerceptionSystem::OnPerceptionStimuliSourceEndPlay);
 }
 
-void UAIPerceptionSystem::RegisterSenseClass(TSubclassOf<UAISense> SenseClass)
+FAISenseID UAIPerceptionSystem::RegisterSenseClass(TSubclassOf<UAISense> SenseClass)
 {
 	check(SenseClass);
 	FAISenseID SenseID = UAISense::GetSenseID(SenseClass);
@@ -47,7 +49,7 @@ void UAIPerceptionSystem::RegisterSenseClass(TSubclassOf<UAISense> SenseClass)
 		if (SenseID.IsValid() == false)
 		{
 			// @todo log a message here
-			return;
+			return FAISenseID::InvalidID();
 		}
 	}
 
@@ -73,6 +75,8 @@ void UAIPerceptionSystem::RegisterSenseClass(TSubclassOf<UAISense> SenseClass)
 		UE_VLOG(this, LogAIPerception, Log, TEXT("Registering sense %s"), *Senses[SenseID]->GetName());
 #endif
 	}
+
+	return SenseID;
 }
 
 UWorld* UAIPerceptionSystem::GetWorld() const
@@ -99,11 +103,10 @@ TStatId UAIPerceptionSystem::GetStatId() const
 	RETURN_QUICK_DECLARE_CYCLE_STAT(UAIPerceptionSystem, STATGROUP_Tickables);
 }
 
-
 void UAIPerceptionSystem::RegisterSource(FAISenseID SenseID, AActor& SourceActor)
 {
 	ensure(IsSenseInstantiated(SenseID));
-	SourcesToRegister.Add(FPerceptionSourceRegistration(SenseID, &SourceActor));
+	SourcesToRegister.AddUnique(FPerceptionSourceRegistration(SenseID, &SourceActor));
 }
 
 void UAIPerceptionSystem::PerformSourceRegistration()
@@ -113,9 +116,17 @@ void UAIPerceptionSystem::PerformSourceRegistration()
 	for (const auto& PercSource : SourcesToRegister)
 	{
 		AActor* SourceActor = PercSource.Source.Get();
-		if (SourceActor)
+		if (SourceActor != nullptr && SourceActor->IsPendingKillPending() == false)
 		{
 			Senses[PercSource.SenseID]->RegisterSource(*SourceActor);
+
+			// hook into notification about actor's EndPlay to remove it as a source 
+			SourceActor->OnEndPlay.AddUnique(StimuliSourceEndPlayDelegate);
+
+			// store information we have this actor as given sense's source
+			FPerceptionStimuliSource& StimuliSource = RegisteredStimuliSources.FindOrAdd(SourceActor);
+			StimuliSource.SourceActor = SourceActor;
+			StimuliSource.RelevantSenses.AcceptChannel(PercSource.SenseID);
 		}
 	}
 
@@ -164,6 +175,34 @@ void UAIPerceptionSystem::Tick(float DeltaSeconds)
 
 		bool bNeedsUpdate = false;
 
+		if (bStimuliSourcesRefreshRequired == true)
+		{
+			// this is a bit naive, but we don't get notifications from the engine which Actor ended play, 
+			// we just know one did. Let's just refresh the map removing all no-longer-valid instances
+			// @todo: we should be just calling senses directly here to unregister specific source
+			// but at this point source actor is invalid already so we can't really pin and use it
+			// this will go away once OnEndPlay broadcast passes in its instigator
+			FPerceptionChannelWhitelist SensesToUpdate;
+			for (auto It = RegisteredStimuliSources.CreateIterator(); It; ++It)
+			{
+				if (It->Value.SourceActor.IsValid() == false)
+				{
+					SensesToUpdate.MergeFilterIn(It->Value.RelevantSenses);
+					It.RemoveCurrent();
+				}
+			}
+
+			for (FPerceptionChannelWhitelist::FConstIterator It(SensesToUpdate); It && Senses.IsValidIndex(*It); ++It)
+			{
+				if (Senses[*It] != nullptr)
+				{
+					Senses[*It]->CleanseInvalidSources();
+				}
+			}
+
+			bStimuliSourcesRefreshRequired = false;
+		}
+		
 		if (SourcesToRegister.Num() > 0)
 		{
 			PerformSourceRegistration();
@@ -311,12 +350,33 @@ void UAIPerceptionSystem::UnregisterListener(UAIPerceptionComponent& Listener)
 	}
 }
 
+void UAIPerceptionSystem::UnregisterSource(AActor& SourceActor)
+{
+	SourceActor.OnEndPlay.Remove(StimuliSourceEndPlayDelegate);
+
+	const FPerceptionStimuliSource* StimuliSource = RegisteredStimuliSources.Find(&SourceActor);
+	if (StimuliSource)
+	{
+		for (int32 SenseID = 0; SenseID < Senses.Num(); ++SenseID)
+		{
+			if (StimuliSource->RelevantSenses.ShouldRespondToChannel(Senses[SenseID]->GetSenseID()))
+			{
+				Senses[SenseID]->UnregisterSource(SourceActor);
+			}
+		}
+	}
+	else
+	{
+		UE_VLOG(this, LogAIPerception, Log, TEXT("UnregisterSource called for %s but it doesn't seem to be registered as a source"), *SourceActor.GetName());
+	}
+}
+
 void UAIPerceptionSystem::OnListenerRemoved(const FPerceptionListener& NewListener)
 {
 	UAISense** SenseInstance = Senses.GetData();
 	for (int32 SenseID = 0; SenseID < Senses.Num(); ++SenseID, ++SenseInstance)
 	{
-		if (*SenseInstance != nullptr)
+		if (*SenseInstance != nullptr && NewListener.HasSense((*SenseInstance)->GetSenseID()))
 		{
 			(*SenseInstance)->OnListenerRemoved(NewListener);
 		}
@@ -400,15 +460,18 @@ bool UAIPerceptionSystem::RegisterPerceptionStimuliSource(UObject* WorldContext,
 		if (World && World->GetAISystem())
 		{
 			UAISystem* AISys = Cast<UAISystem>(World->GetAISystem());
-			if (AISys && AISys->GetPerceptionSystem())
+			if (AISys != nullptr && AISys->GetPerceptionSystem() != nullptr)
 			{
-				const FAISenseID SenseID = UAISense::GetSenseID(Sense);
-				if (AISys->GetPerceptionSystem()->IsSenseInstantiated(SenseID) == false)
+				// just a cache
+				UAIPerceptionSystem* PerceptionSys = AISys->GetPerceptionSystem();
+
+				FAISenseID SenseID = UAISense::GetSenseID(Sense);
+				if (PerceptionSys->IsSenseInstantiated(SenseID) == false)
 				{
-					AISys->GetPerceptionSystem()->RegisterSenseClass(Sense);
+					SenseID = PerceptionSys->RegisterSenseClass(Sense);
 				}
 
-				AISys->GetPerceptionSystem()->RegisterSource(SenseID, *Target);
+				PerceptionSys->RegisterSource(SenseID, *Target);
 				bResult = true;
 			}
 		}
@@ -417,6 +480,14 @@ bool UAIPerceptionSystem::RegisterPerceptionStimuliSource(UObject* WorldContext,
 	return bResult;
 }
 
+void UAIPerceptionSystem::OnPerceptionStimuliSourceEndPlay(EEndPlayReason::Type EndPlayReason)
+{
+	// this tells us just that _a_ source has been removed. We need to parse through all sources and find which one was it
+	// this is a fall-back behavior, if source gets unregistered manually this function won't get called 
+	// note, the actual work will be done on system's tick
+	bStimuliSourcesRefreshRequired = true;
+}
+
 TSubclassOf<UAISense> UAIPerceptionSystem::GetSenseClassForStimulus(UObject* WorldContext, const FAIStimulus& Stimulus)
 {
 	TSubclassOf<UAISense> Result = nullptr;
diff --git a/Engine/Source/Runtime/AIModule/Private/Perception/AISense_Sight.cpp b/Engine/Source/Runtime/AIModule/Private/Perception/AISense_Sight.cpp
--- a/Engine/Source/Runtime/AIModule/Private/Perception/AISense_Sight.cpp
+++ b/Engine/Source/Runtime/AIModule/Private/Perception/AISense_Sight.cpp
@@ -257,6 +257,9 @@ float UAISense_Sight::Update()
 
 		if (InvalidTargets.Num() > 0)
 		{
+			// this should not be happening since UAIPerceptionSystem::OnPerceptionStimuliSourceEndPlay introduction
+			UE_VLOG(GetPerceptionSystem(), LogAIPerception, Error, TEXT("Invalid sight targets found during UAISense_Sight::Update call"));
+
 			for (const auto& TargetId : InvalidTargets)
 			{
 				// remove affected queries
@@ -287,6 +290,44 @@ void UAISense_Sight::RegisterSource(AActor& SourceActor)
 	RegisterTarget(SourceActor, Sort);
 }
 
+void UAISense_Sight::UnregisterSource(AActor& SourceActor)
+{
+	const FAISightTarget::FTargetId AsTargetId = SourceActor.GetFName();
+	FAISightTarget* AsTarget = ObservedTargets.Find(AsTargetId);
+	if (AsTarget != nullptr)
+	{
+		RemoveAllQueriesToTarget(AsTargetId, DontSort);
+	}
+}
+
+void UAISense_Sight::CleanseInvalidSources()
+{
+	bool bInvalidSourcesFound = false;
+	for (TMap<FName, FAISightTarget>::TIterator ItTarget(ObservedTargets); ItTarget; ++ItTarget)
+	{
+		if (ItTarget->Value.Target.IsValid() == false)
+		{
+			// remove affected queries
+			RemoveAllQueriesToTarget(ItTarget->Key, DontSort);
+			// remove target itself
+			ItTarget.RemoveCurrent();
+
+			bInvalidSourcesFound = true;
+		}
+	}
+
+	if (bInvalidSourcesFound)
+	{
+		// remove holes
+		ObservedTargets.Compact();
+		SortQueries();
+	}
+	else
+	{
+		UE_VLOG(GetPerceptionSystem(), LogAIPerception, Error, TEXT("UAISense_Sight::CleanseInvalidSources called and no invalid targets were found"));
+	}
+}
+
 bool UAISense_Sight::RegisterTarget(AActor& TargetActor, FQueriesOperationPostProcess PostProcess)
 {
 	SCOPE_CYCLE_COUNTER(STAT_AI_Sense_Sight);
@@ -429,20 +470,9 @@ void UAISense_Sight::OnListenerRemovedImpl(const FPerceptionListener& UpdatedLis
 
 	DigestedProperties.FindAndRemoveChecked(UpdatedListener.GetListenerID());
 
-	if (UpdatedListener.Listener.IsValid())
-	{
-		// see if this listener is a Target as well
-		const FAISightTarget::FTargetId AsTargetId = UpdatedListener.GetBodyActorName();
-		FAISightTarget* AsTarget = ObservedTargets.Find(AsTargetId);
-		if (AsTarget != NULL)
-		{
-			RemoveAllQueriesToTarget(AsTargetId, Sort);			
-		}
-	}
-	else
-	{
-		//@todo quite possible there are left over sight queries with this listener as target
-	}
+	// note: there use to be code to remove all queries _to_ listener here as well
+	// but that was wrong - the fact that a listener gets unregistered doesn't have to
+	// mean it's being removed from the game altogether.
 }
 
 void UAISense_Sight::RemoveAllQueriesByListener(const FPerceptionListener& Listener, FQueriesOperationPostProcess PostProcess)
@@ -473,7 +503,7 @@ void UAISense_Sight::RemoveAllQueriesByListener(const FPerceptionListener& Liste
 	}
 }
 
-void UAISense_Sight::RemoveAllQueriesToTarget(const FName& TargetId, FQueriesOperationPostProcess PostProcess)
+void UAISense_Sight::RemoveAllQueriesToTarget(const FAISightTarget::FTargetId& TargetId, FQueriesOperationPostProcess PostProcess)
 {
 	SCOPE_CYCLE_COUNTER(STAT_AI_Sense_Sight);
 
