Fixed a horrible crash resulting from BT decorators having wrong ChildIndex assigned, which in turn was a result of UBTNode::GetSpecialNodeMemory returning valid pointers for non-instanced nodes #UE4

Merged over from dev branch

diff --git a/Engine/Source/Runtime/AIModule/Classes/BehaviorTree/BTNode.h b/Engine/Source/Runtime/AIModule/Classes/BehaviorTree/BTNode.h
--- a/Engine/Source/Runtime/AIModule/Classes/BehaviorTree/BTNode.h
+++ b/Engine/Source/Runtime/AIModule/Classes/BehaviorTree/BTNode.h
@@ -293,5 +293,6 @@ const T* UBTNode::GetNodeMemory(const FBehaviorTreeInstance& BTInstance) const
 template<typename T>
 T* UBTNode::GetSpecialNodeMemory(uint8* NodeMemory) const
 {
-	return (T*)(NodeMemory - ((GetSpecialMemorySize() + 3) & ~3));
+	const int32 SpecialMemorySize = GetSpecialMemorySize();
+	return SpecialMemorySize ? (T*)(NodeMemory - ((SpecialMemorySize + 3) & ~3)) : nullptr;
 }
diff --git a/Engine/Source/Runtime/AIModule/Private/BehaviorTree/BTCompositeNode.cpp b/Engine/Source/Runtime/AIModule/Private/BehaviorTree/BTCompositeNode.cpp
--- a/Engine/Source/Runtime/AIModule/Private/BehaviorTree/BTCompositeNode.cpp
+++ b/Engine/Source/Runtime/AIModule/Private/BehaviorTree/BTCompositeNode.cpp
@@ -354,6 +354,13 @@ static bool UpdateOperationStack(const UBehaviorTreeComponent& OwnerComp, FStrin
 
 bool UBTCompositeNode::DoDecoratorsAllowExecution(UBehaviorTreeComponent& OwnerComp, int32 InstanceIdx, int32 ChildIdx) const
 {
+	if (Children.IsValidIndex(ChildIdx) == false)
+	{
+		UE_VLOG(OwnerComp.GetOwner(), LogBehaviorTree, Error, TEXT("%s: DoDecoratorsAllowExecution called with ChildIdx = %d which is not a valid child index")
+			, *UBehaviorTreeTypes::DescribeNodeHelper(this), ChildIdx);
+		return false;
+	}
+
 	const FBTCompositeChild& ChildInfo = Children[ChildIdx];
 	bool bResult = true;
 
