Fixed a bug in blackboard key comparing #UE4

diff --git a/Engine/Source/Runtime/AIModule/Private/BehaviorTree/BlackboardComponent.cpp b/Engine/Source/Runtime/AIModule/Private/BehaviorTree/BlackboardComponent.cpp
--- a/Engine/Source/Runtime/AIModule/Private/BehaviorTree/BlackboardComponent.cpp
+++ b/Engine/Source/Runtime/AIModule/Private/BehaviorTree/BlackboardComponent.cpp
@@ -436,7 +436,7 @@ EBlackboardCompare::Type UBlackboardComponent::CompareKeyValues(TSubclassOf<UBla
 	const uint8* KeyAMemory = GetKeyRawData(KeyA) + (KeyInstances[KeyA] ? sizeof(FBlackboardInstancedKeyMemory) : 0);
 	const uint8* KeyBMemory = GetKeyRawData(KeyB) + (KeyInstances[KeyB] ? sizeof(FBlackboardInstancedKeyMemory) : 0);
 
-	const UBlackboardKeyType* KeyAOb = KeyInstances[KeyA] ? KeyInstances[KeyA] : GetDefault<UBlackboardKeyType>();
+	const UBlackboardKeyType* KeyAOb = KeyInstances[KeyA] ? KeyInstances[KeyA] : KeyType->GetDefaultObject<UBlackboardKeyType>();
 	return KeyAOb->CompareValues(*this, KeyAMemory, KeyInstances[KeyB], KeyBMemory);
 }
 
diff --git a/Engine/Source/Runtime/AIModule/Private/BehaviorTree/Decorators/BTDecorator_CompareBBEntries.cpp b/Engine/Source/Runtime/AIModule/Private/BehaviorTree/Decorators/BTDecorator_CompareBBEntries.cpp
--- a/Engine/Source/Runtime/AIModule/Private/BehaviorTree/Decorators/BTDecorator_CompareBBEntries.cpp
+++ b/Engine/Source/Runtime/AIModule/Private/BehaviorTree/Decorators/BTDecorator_CompareBBEntries.cpp
@@ -44,6 +44,7 @@ bool UBTDecorator_CompareBBEntries::CalculateRawConditionValue(UBehaviorTreeComp
 			&& int32(EBlackboardCompare::NotEqual) == int32(EBlackBoardEntryComparison::NotEqual),
 			"These values need to be equal");
 
+		// note that this comparison relies on assumption asserted above
 		return int32(Operator) == int32(Result);
 	}
 
