Fixed AIController crashing if called to posses a Null pawn while already having a pawn #UE4

UE-14433

diff --git a/Engine/Source/Runtime/AIModule/Private/AIController.cpp b/Engine/Source/Runtime/AIModule/Private/AIController.cpp
--- a/Engine/Source/Runtime/AIModule/Private/AIController.cpp
+++ b/Engine/Source/Runtime/AIModule/Private/AIController.cpp
@@ -438,7 +438,7 @@ void AAIController::Possess(APawn* InPawn)
 {
 	Super::Possess(InPawn);
 
-	if (!GetPawn())
+	if (GetPawn() == nullptr || InPawn == nullptr)
 	{
 		return;
 	}
