hidden deprecated native enum class in blackboard's dropdown
#ue4

diff --git a/Engine/Source/Runtime/AIModule/Classes/BehaviorTree/Blackboard/BlackboardKeyType_NativeEnum.h b/Engine/Source/Runtime/AIModule/Classes/BehaviorTree/Blackboard/BlackboardKeyType_NativeEnum.h
--- a/Engine/Source/Runtime/AIModule/Classes/BehaviorTree/Blackboard/BlackboardKeyType_NativeEnum.h
+++ b/Engine/Source/Runtime/AIModule/Classes/BehaviorTree/Blackboard/BlackboardKeyType_NativeEnum.h
@@ -6,7 +6,7 @@
 
 // DEPRECATED, please use UBlackboardKeyType_Enum instead
 
-UCLASS(NotEditInlineNew)
+UCLASS(NotEditInlineNew, HideDropDown)
 class AIMODULE_API UBlackboardKeyType_NativeEnum : public UBlackboardKeyType
 {
 	GENERATED_UCLASS_BODY()
