Made sure AI's BrainComponent points at the BehaviorTreeComponent AI owns when RunBehaviorTree gets called #UE4

UE-14758

diff --git a/Engine/Source/Runtime/AIModule/Private/AIController.cpp b/Engine/Source/Runtime/AIModule/Private/AIController.cpp
--- a/Engine/Source/Runtime/AIModule/Private/AIController.cpp
+++ b/Engine/Source/Runtime/AIModule/Private/AIController.cpp
@@ -810,9 +810,12 @@ bool AAIController::RunBehaviorTree(UBehaviorTree* BTAsset)
 		{
 			UE_VLOG(this, LogBehaviorTree, Log, TEXT("RunBehaviorTree: spawning BehaviorTreeComponent.."));
 
-			BrainComponent = BTComp = NewObject<UBehaviorTreeComponent>(this, TEXT("BTComponent"));
-			BrainComponent->RegisterComponent();
+			BTComp = NewObject<UBehaviorTreeComponent>(this, TEXT("BTComponent"));
+			BTComp->RegisterComponent();
 		}
+		
+		// make sure BrainComponent points at the newly created BT component
+		BrainComponent = BTComp;
 
 		check(BTComp != NULL);
 		BTComp->StartTree(*BTAsset, EBTExecutionMode::Looped);
