added sanity check in behavior tree tick function before accessing instance data
#ue4 UE-14760

diff --git a/Engine/Source/Runtime/AIModule/Private/BehaviorTree/BehaviorTreeComponent.cpp b/Engine/Source/Runtime/AIModule/Private/BehaviorTree/BehaviorTreeComponent.cpp
--- a/Engine/Source/Runtime/AIModule/Private/BehaviorTree/BehaviorTreeComponent.cpp
+++ b/Engine/Source/Runtime/AIModule/Private/BehaviorTree/BehaviorTreeComponent.cpp
@@ -964,12 +964,12 @@ void UBehaviorTreeComponent::TickComponent(float DeltaTime, enum ELevelTick Tick
 	SCOPE_CYCLE_COUNTER(STAT_AI_BehaviorTree_Tick);
 
 	check(this != nullptr && this->IsPendingKill() == false);
-		
+
 	if (bRequestedFlowUpdate)
 	{
 		ProcessExecutionRequest();
 	}
-	
+
 	if (InstanceStack.Num() == 0 || !bIsRunning)
 	{
 		return;
@@ -978,7 +978,7 @@ void UBehaviorTreeComponent::TickComponent(float DeltaTime, enum ELevelTick Tick
 	// tick active auxiliary nodes and parallel tasks (in execution order, before task)
 	for (int32 InstanceIndex = 0; InstanceIndex < InstanceStack.Num(); InstanceIndex++)
 	{
-		FBehaviorTreeInstance& InstanceInfo = InstanceStack[InstanceIndex];		
+		FBehaviorTreeInstance& InstanceInfo = InstanceStack[InstanceIndex];
 		for (int32 AuxIndex = 0; AuxIndex < InstanceInfo.ActiveAuxNodes.Num(); AuxIndex++)
 		{
 			const UBTAuxiliaryNode* AuxNode = InstanceInfo.ActiveAuxNodes[AuxIndex];
@@ -995,13 +995,16 @@ void UBehaviorTreeComponent::TickComponent(float DeltaTime, enum ELevelTick Tick
 	}
 
 	// tick active task
-	FBehaviorTreeInstance& ActiveInstance = InstanceStack[ActiveInstanceIdx];
-	if (ActiveInstance.ActiveNodeType == EBTActiveNode::ActiveTask ||
-		ActiveInstance.ActiveNodeType == EBTActiveNode::AbortingTask)
+	if (InstanceStack.IsValidIndex(ActiveInstanceIdx))
 	{
-		UBTTaskNode* ActiveTask = (UBTTaskNode*)ActiveInstance.ActiveNode;
-		uint8* NodeMemory = ActiveTask->GetNodeMemory<uint8>(ActiveInstance);
-		ActiveTask->WrappedTickTask(*this, NodeMemory, DeltaTime);
+		FBehaviorTreeInstance& ActiveInstance = InstanceStack[ActiveInstanceIdx];
+		if (ActiveInstance.ActiveNodeType == EBTActiveNode::ActiveTask ||
+			ActiveInstance.ActiveNodeType == EBTActiveNode::AbortingTask)
+		{
+			UBTTaskNode* ActiveTask = (UBTTaskNode*)ActiveInstance.ActiveNode;
+			uint8* NodeMemory = ActiveTask->GetNodeMemory<uint8>(ActiveInstance);
+			ActiveTask->WrappedTickTask(*this, NodeMemory, DeltaTime);
+		}
 	}
 }
 
