Get/SetValue functions in blackboard will now check type of key
#ue4 UE-13569

diff --git a/Engine/Source/Runtime/AIModule/Classes/BehaviorTree/BlackboardComponent.h b/Engine/Source/Runtime/AIModule/Classes/BehaviorTree/BlackboardComponent.h
--- a/Engine/Source/Runtime/AIModule/Classes/BehaviorTree/BlackboardComponent.h
+++ b/Engine/Source/Runtime/AIModule/Classes/BehaviorTree/BlackboardComponent.h
@@ -344,7 +344,7 @@ template<class TDataClass>
 bool UBlackboardComponent::SetValue(FBlackboard::FKey KeyID, typename TDataClass::FDataType Value)
 {
 	const FBlackboardEntry* EntryInfo = BlackboardAsset ? BlackboardAsset->GetKey(KeyID) : nullptr;
-	if ((EntryInfo == nullptr) || (EntryInfo->KeyType == nullptr))
+	if ((EntryInfo == nullptr) || (EntryInfo->KeyType == nullptr) || (EntryInfo->KeyType->GetClass() != TDataClass::StaticClass()))
 	{
 		return false;
 	}
@@ -393,7 +393,7 @@ template<class TDataClass>
 typename TDataClass::FDataType UBlackboardComponent::GetValue(FBlackboard::FKey KeyID) const
 {
 	const FBlackboardEntry* EntryInfo = BlackboardAsset ? BlackboardAsset->GetKey(KeyID) : nullptr;
-	if ((EntryInfo == nullptr) || (EntryInfo->KeyType == nullptr))
+	if ((EntryInfo == nullptr) || (EntryInfo->KeyType == nullptr) || (EntryInfo->KeyType->GetClass() != TDataClass::StaticClass()))
 	{
 		return TDataClass::InvalidValue;
 	}
diff --git a/Engine/Source/Runtime/AIModule/Private/BehaviorTree/BlackboardComponent.cpp b/Engine/Source/Runtime/AIModule/Private/BehaviorTree/BlackboardComponent.cpp
--- a/Engine/Source/Runtime/AIModule/Private/BehaviorTree/BlackboardComponent.cpp
+++ b/Engine/Source/Runtime/AIModule/Private/BehaviorTree/BlackboardComponent.cpp
@@ -564,13 +564,6 @@ void UBlackboardComponent::DescribeSelfToVisLog(FVisualLogEntry* Snapshot) const
 
 UObject* UBlackboardComponent::GetValueAsObject(const FName& KeyName) const
 {
-	// temporary fix, will be replaced with type safe entries soon
-	FBlackboard::FKey KeyID = GetKeyID(KeyName);
-	if (GetKeyType(KeyID) != UBlackboardKeyType_Object::StaticClass())
-	{
-		return UBlackboardKeyType_Object::InvalidValue;
-	}
-
 	return GetValue<UBlackboardKeyType_Object>(KeyName);
 }
 
@@ -581,13 +574,6 @@ UClass* UBlackboardComponent::GetValueAsClass(const FName& KeyName) const
 
 uint8 UBlackboardComponent::GetValueAsEnum(const FName& KeyName) const
 {
-	FBlackboard::FKey KeyID = GetKeyID(KeyName);
-	if (GetKeyType(KeyID) != UBlackboardKeyType_Enum::StaticClass() &&
-		GetKeyType(KeyID) != UBlackboardKeyType_NativeEnum::StaticClass())
-	{
-		return UBlackboardKeyType_Enum::InvalidValue;
-	}
-
 	return GetValue<UBlackboardKeyType_Enum>(KeyName);
 }
 
