Improved comments on BT composite nodes #UE4

github PR-1116 by JohnAlcatraz

diff --git a/Engine/Source/Runtime/AIModule/Classes/BehaviorTree/Composites/BTComposite_Selector.h b/Engine/Source/Runtime/AIModule/Classes/BehaviorTree/Composites/BTComposite_Selector.h
--- a/Engine/Source/Runtime/AIModule/Classes/BehaviorTree/Composites/BTComposite_Selector.h
+++ b/Engine/Source/Runtime/AIModule/Classes/BehaviorTree/Composites/BTComposite_Selector.h
@@ -6,7 +6,8 @@
 
 /** 
  * Selector composite node.
- * A selector node runs each child in turn until the first child node succeeds, in which case the selector fails.
+ * Selector Nodes execute their children from left to right, and will stop executing its children when one of their children succeeds.
+ * If a Selector's child succeeds, the Selector succeeds. If all the Selector's children fail, the Selector fails.
  */
 UCLASS()
 class AIMODULE_API UBTComposite_Selector: public UBTCompositeNode
diff --git a/Engine/Source/Runtime/AIModule/Classes/BehaviorTree/Composites/BTComposite_Sequence.h b/Engine/Source/Runtime/AIModule/Classes/BehaviorTree/Composites/BTComposite_Sequence.h
--- a/Engine/Source/Runtime/AIModule/Classes/BehaviorTree/Composites/BTComposite_Sequence.h
+++ b/Engine/Source/Runtime/AIModule/Classes/BehaviorTree/Composites/BTComposite_Sequence.h
@@ -4,9 +4,10 @@
 #include "BehaviorTree/BTCompositeNode.h"
 #include "BTComposite_Sequence.generated.h"
 
-/** 
+/**
  * Sequence composite node.
- * A sequence node runs each child in turn until one fails, in which case the sequence succeeds.
+ * Sequence Nodes execute their children from left to right, and will stop executing its children when one of their children fails.
+ * If a child fails, then the Sequence fails. If all the Sequence's children succeed, then the Sequence succeeds.
  */
 UCLASS()
 class AIMODULE_API UBTComposite_Sequence : public UBTCompositeNode
