Marked AIPerceptionProperties referring to specific senses deprecated #UE4

UE-14979

diff --git a/Engine/Source/Runtime/AIModule/Classes/Perception/AIPerceptionComponent.h b/Engine/Source/Runtime/AIModule/Classes/Perception/AIPerceptionComponent.h
--- a/Engine/Source/Runtime/AIModule/Classes/Perception/AIPerceptionComponent.h
+++ b/Engine/Source/Runtime/AIModule/Classes/Perception/AIPerceptionComponent.h
@@ -115,23 +115,28 @@ class AIMODULE_API UAIPerceptionComponent : public UActorComponent
 
 protected:
 	/** Max distance at which a makenoise(1.0) loudness sound can be heard, regardless of occlusion */
-	UPROPERTY(EditAnywhere, BlueprintReadOnly, Category=AI)
+	DEPRECATED(4.8, "This property is deprecated. Please use apropriate sencse config class instead")
+	UPROPERTY(VisibleDefaultsOnly, Category = AI)
 	float HearingRange;
 
 	/** Max distance at which a makenoise(1.0) loudness sound can be heard if unoccluded (LOSHearingThreshold should be > HearingThreshold) */
-	UPROPERTY(EditAnywhere, BlueprintReadOnly, Category=AI)
+	DEPRECATED(4.8, "This property is deprecated. Please use apropriate sencse config class instead")
+	UPROPERTY(VisibleDefaultsOnly, Category = AI)
 	float LoSHearingRange;
 
 	/** Maximum sight distance to notice a target. */
-	UPROPERTY(EditAnywhere, BlueprintReadOnly, Category=AI)
+	DEPRECATED(4.8, "This property is deprecated. Please use apropriate sencse config class instead")
+	UPROPERTY(VisibleDefaultsOnly, Category = AI)
 	float SightRadius;
 
 	/** Maximum sight distance to see target that has been already seen. */
-	UPROPERTY(EditAnywhere, BlueprintReadOnly, Category=AI)
+	DEPRECATED(4.8, "This property is deprecated. Please use apropriate sencse config class instead")
+	UPROPERTY(VisibleDefaultsOnly, Category = AI)
 	float LoseSightRadius;
 
 	/** How far to the side AI can see, in degrees. Use SetPeripheralVisionAngle to change the value at runtime. */
-	UPROPERTY(EditAnywhere, BlueprintReadOnly, Category=AI)
+	DEPRECATED(4.8, "This property is deprecated. Please use apropriate sencse config class instead")
+	UPROPERTY(VisibleDefaultsOnly, Category = AI)
 	float PeripheralVisionAngle;
 		
 	UPROPERTY(EditDefaultsOnly, Instanced, Category = "AI Perception")
