Added easy to use BP events for spawning Damage and Prediction AI events #UE4

UE-15090

diff --git a/Engine/Source/Runtime/AIModule/Classes/Perception/AISense_Damage.h b/Engine/Source/Runtime/AIModule/Classes/Perception/AISense_Damage.h
--- a/Engine/Source/Runtime/AIModule/Classes/Perception/AISense_Damage.h
+++ b/Engine/Source/Runtime/AIModule/Classes/Perception/AISense_Damage.h
@@ -44,6 +44,8 @@ struct AIMODULE_API FAIDamageEvent
 	{
 		return DamagedActor != nullptr;
 	}
+
+	IAIPerceptionListenerInterface* FAIDamageEvent::GetDamagedActorAsPerceptionListener() const;
 };
 
 UCLASS(ClassGroup=AI)
@@ -58,6 +60,10 @@ public:
 	void RegisterEvent(const FAIDamageEvent& Event);	
 	virtual void RegisterWrappedEvent(UAISenseEvent& PerceptionEvent) override;
 
+	/** EventLocation will be reported as Instigator's location at the moment of event happening*/
+	UFUNCTION(BlueprintCallable, Category = "AI|Perception", meta = (HidePin = "WorldContext", DefaultToSelf = "WorldContext", AdvancedDisplay="HitLocation"))
+	static void ReportDamageEvent(UObject* WorldContext, AActor* DamagedActor, AActor* Instigator, float DamageAmount, FVector EventLocation, FVector HitLocation);
+
 protected:
 	virtual float Update() override;
 };
diff --git a/Engine/Source/Runtime/AIModule/Classes/Perception/AISense_Prediction.h b/Engine/Source/Runtime/AIModule/Classes/Perception/AISense_Prediction.h
--- a/Engine/Source/Runtime/AIModule/Classes/Perception/AISense_Prediction.h
+++ b/Engine/Source/Runtime/AIModule/Classes/Perception/AISense_Prediction.h
@@ -42,6 +42,16 @@ class AIMODULE_API UAISense_Prediction : public UAISense
 public:		
 	void RegisterEvent(const FAIPredictionEvent& Event);	
 
+	/** Asks perception system to supply Requestor with PredictedActor's predicted location in PredictionTime seconds
+	 *	Location is being predicted based on PredicterActor's current location and velocity */
+	UFUNCTION(BlueprintCallable, Category = "AI|Perception")
+	static void RequestControllerPredictionEvent(AAIController* Requestor, AActor* PredictedActor, float PredictionTime);
+
+	/** Asks perception system to supply Requestor with PredictedActor's predicted location in PredictionTime seconds
+	 *	Location is being predicted based on PredicterActor's current location and velocity */
+	UFUNCTION(BlueprintCallable, Category = "AI|Perception")
+	static void RequestPawnPredictionEvent(APawn* Requestor, AActor* PredictedActor, float PredictionTime);
+
 protected:
 	virtual float Update() override;
 };
diff --git a/Engine/Source/Runtime/AIModule/Private/Perception/AISense_Damage.cpp b/Engine/Source/Runtime/AIModule/Private/Perception/AISense_Damage.cpp
--- a/Engine/Source/Runtime/AIModule/Private/Perception/AISense_Damage.cpp
+++ b/Engine/Source/Runtime/AIModule/Private/Perception/AISense_Damage.cpp
@@ -47,6 +47,24 @@ void FAIDamageEvent::Compile()
 	}
 
 }
+
+IAIPerceptionListenerInterface* FAIDamageEvent::GetDamagedActorAsPerceptionListener() const
+{
+	IAIPerceptionListenerInterface* Listener = nullptr;
+	if (DamagedActor)
+	{
+		Listener = Cast<IAIPerceptionListenerInterface>(DamagedActor);
+		if (Listener == nullptr)
+		{
+			APawn* ListenerAsPawn = Cast<APawn>(DamagedActor);
+			if (ListenerAsPawn)
+			{
+				Listener = Cast<IAIPerceptionListenerInterface>(ListenerAsPawn->GetController());
+			}
+		}
+	}
+	return Listener;
+}
 //----------------------------------------------------------------------//
 // 
 //----------------------------------------------------------------------//
@@ -65,11 +83,11 @@ float UAISense_Damage::Update()
 	{
 		const FAIDamageEvent& Event = RegisteredEvents[EventIndex];
 
-		IAIPerceptionListenerInterface* PerceptionListener = Cast<IAIPerceptionListenerInterface>(Event.DamagedActor);
-		if (PerceptionListener != NULL)
+		IAIPerceptionListenerInterface* PerceptionListener = Event.GetDamagedActorAsPerceptionListener();
+		if (PerceptionListener != nullptr)
 		{
 			UAIPerceptionComponent* PerceptionComponent = PerceptionListener->GetPerceptionComponent();
-			if (PerceptionComponent != NULL)
+			if (PerceptionComponent != nullptr)
 			{
 				// this has to succeed, will assert a failure
 				FPerceptionListener& Listener = ListenersMap[PerceptionComponent->GetListenerId()];
@@ -104,3 +122,13 @@ void UAISense_Damage::RegisterWrappedEvent(UAISenseEvent& PerceptionEvent)
 		RegisterEvent(DamageEvent->GetDamageEvent());
 	}
 }
+
+void UAISense_Damage::ReportDamageEvent(UObject* WorldContext, AActor* DamagedActor, AActor* Instigator, float DamageAmount, FVector EventLocation, FVector HitLocation)
+{
+	UAIPerceptionSystem* PerceptionSystem = UAIPerceptionSystem::GetCurrent(WorldContext);
+	if (PerceptionSystem)
+	{
+		FAIDamageEvent Event(DamagedActor, Instigator, DamageAmount, EventLocation, HitLocation);
+		PerceptionSystem->OnEvent(Event);
+	}
+}
diff --git a/Engine/Source/Runtime/AIModule/Private/Perception/AISense_Prediction.cpp b/Engine/Source/Runtime/AIModule/Private/Perception/AISense_Prediction.cpp
--- a/Engine/Source/Runtime/AIModule/Private/Perception/AISense_Prediction.cpp
+++ b/Engine/Source/Runtime/AIModule/Private/Perception/AISense_Prediction.cpp
@@ -48,4 +48,24 @@ void UAISense_Prediction::RegisterEvent(const FAIPredictionEvent& Event)
 {
 	RegisteredEvents.Add(Event);
 	RequestImmediateUpdate();
-}
\ No newline at end of file
+}
+
+void UAISense_Prediction::RequestControllerPredictionEvent(AAIController* Requestor, AActor* PredictedActor, float PredictionTime)
+{
+	UAIPerceptionSystem* PerceptionSystem = UAIPerceptionSystem::GetCurrent(Requestor);
+	if (PerceptionSystem)
+	{
+		FAIPredictionEvent Event(Requestor, PredictedActor, PredictionTime);
+		PerceptionSystem->OnEvent(Event);
+	}
+}
+
+void UAISense_Prediction::RequestPawnPredictionEvent(APawn* Requestor, AActor* PredictedActor, float PredictionTime)
+{
+	UAIPerceptionSystem* PerceptionSystem = UAIPerceptionSystem::GetCurrent(Requestor);
+	if (PerceptionSystem && Requestor->GetController())
+	{
+		FAIPredictionEvent Event(Requestor->GetController(), PredictedActor, PredictionTime);
+		PerceptionSystem->OnEvent(Event);
+	}
+}
