Fix CIS (extra qualifier).

diff --git a/Engine/Source/Runtime/AIModule/Classes/Perception/AISense_Damage.h b/Engine/Source/Runtime/AIModule/Classes/Perception/AISense_Damage.h
--- a/Engine/Source/Runtime/AIModule/Classes/Perception/AISense_Damage.h
+++ b/Engine/Source/Runtime/AIModule/Classes/Perception/AISense_Damage.h
@@ -45,7 +45,7 @@ struct AIMODULE_API FAIDamageEvent
 		return DamagedActor != nullptr;
 	}
 
-	IAIPerceptionListenerInterface* FAIDamageEvent::GetDamagedActorAsPerceptionListener() const;
+	IAIPerceptionListenerInterface* GetDamagedActorAsPerceptionListener() const;
 };
 
 UCLASS(ClassGroup=AI)
