Fixed AI not ignoring their own pawns for perception purposes #UE4

UE-15088

diff --git a/Engine/Source/Runtime/AIModule/Private/Perception/AISense_Sight.cpp b/Engine/Source/Runtime/AIModule/Private/Perception/AISense_Sight.cpp
--- a/Engine/Source/Runtime/AIModule/Private/Perception/AISense_Sight.cpp
+++ b/Engine/Source/Runtime/AIModule/Private/Perception/AISense_Sight.cpp
@@ -357,7 +357,7 @@ bool UAISense_Sight::RegisterTarget(AActor& TargetActor, FQueriesOperationPostPr
 		const FPerceptionListener& Listener = ItListener->Value;
 		const IGenericTeamAgentInterface* ListenersTeamAgent = Listener.GetTeamAgent();
 
-		if (Listener.HasSense(GetSenseID()))
+		if (Listener.HasSense(GetSenseID()) && Listener.GetBodyActor() != &TargetActor)
 		{
 			const FDigestedSightProperties& PropDigest = DigestedProperties[Listener.GetListenerID()];
 			if (FAISenseAffiliationFilter::ShouldSenseTeam(ListenersTeamAgent, TargetActor, PropDigest.AffiliationFlags))
@@ -397,12 +397,13 @@ void UAISense_Sight::GenerateQueriesForListener(const FPerceptionListener& Liste
 {
 	bool bNewQueriesAdded = false;
 	const IGenericTeamAgentInterface* ListenersTeamAgent = Listener.GetTeamAgent();
+	const AActor* Avatar = Listener.GetBodyActor();
 
 	// create sight queries with all legal targets
 	for (TMap<FName, FAISightTarget>::TConstIterator ItTarget(ObservedTargets); ItTarget; ++ItTarget)
 	{
 		const AActor* TargetActor = ItTarget->Value.GetTargetActor();
-		if (TargetActor == NULL)
+		if (TargetActor == NULL || TargetActor == Avatar)
 		{
 			continue;
 		}
