Changed when AIPerceptionSystem's stimuli ageing is being set up #UE4

This solves some project issues where stimuli never aged
UE-15335

diff --git a/Engine/Source/Runtime/AIModule/Private/Perception/AIPerceptionSystem.cpp b/Engine/Source/Runtime/AIModule/Private/Perception/AIPerceptionSystem.cpp
--- a/Engine/Source/Runtime/AIModule/Private/Perception/AIPerceptionSystem.cpp
+++ b/Engine/Source/Runtime/AIModule/Private/Perception/AIPerceptionSystem.cpp
@@ -100,15 +100,6 @@ UWorld* UAIPerceptionSystem::GetWorld() const
 void UAIPerceptionSystem::PostInitProperties() 
 {
 	Super::PostInitProperties();
-
-	if (HasAnyFlags(RF_ClassDefaultObject) == false)
-	{
-		UWorld* World = GEngine->GetWorldFromContextObject(GetOuter());
-		if (World)
-		{
-			World->GetTimerManager().SetTimer(AgeStimuliTimerHandle, this, &UAIPerceptionSystem::AgeStimuli, PerceptionAgingRate, /*inbLoop=*/true);
-		}
-	}
 }
 
 TStatId UAIPerceptionSystem::GetStatId() const
@@ -510,6 +501,12 @@ void UAIPerceptionSystem::StartPlay()
 			RegisterAllPawnsAsSourcesForSense(SenseID);
 		}
 	}
+
+	UWorld* World = GetWorld();
+	if (World)
+	{
+		World->GetTimerManager().SetTimer(AgeStimuliTimerHandle, this, &UAIPerceptionSystem::AgeStimuli, PerceptionAgingRate, /*inbLoop=*/true);
+	}
 }
 
 void UAIPerceptionSystem::RegisterAllPawnsAsSourcesForSense(FAISenseID SenseID)
