unified behavior tree's StopTree & StopLogic functions
#ue4 UE-14908

diff --git a/Engine/Source/Runtime/AIModule/Classes/BehaviorTree/BehaviorTreeComponent.h b/Engine/Source/Runtime/AIModule/Classes/BehaviorTree/BehaviorTreeComponent.h
--- a/Engine/Source/Runtime/AIModule/Classes/BehaviorTree/BehaviorTreeComponent.h
+++ b/Engine/Source/Runtime/AIModule/Classes/BehaviorTree/BehaviorTreeComponent.h
@@ -66,13 +66,15 @@ struct FBTPendingExecutionInfo
 	void Unlock() { bLocked = false; }
 };
 
-struct FBTPendingInitializeInfo
+struct FBTTreeStartInfo
 {
 	UBehaviorTree* Asset;
 	EBTExecutionMode::Type ExecuteMode;
+	uint8 bPendingInitialize : 1;
 
-	FBTPendingInitializeInfo() : Asset(nullptr), ExecuteMode(EBTExecutionMode::Looped) {}
+	FBTTreeStartInfo() : Asset(nullptr), ExecuteMode(EBTExecutionMode::Looped), bPendingInitialize(0) {}
 	bool IsSet() const { return Asset != nullptr; }
+	bool HasPendingInitialize() const { return bPendingInitialize && IsSet(); }
 };
 
 UCLASS()
@@ -226,7 +228,7 @@ protected:
 	FBTPendingExecutionInfo PendingExecution;
 
 	/** stored data for starting new tree, waits until previously running finishes aborting */
-	FBTPendingInitializeInfo PendingInitialize;
+	FBTTreeStartInfo TreeStartInfo;
 
 	/** message observers mapped by instance & execution index */
 	TMultiMap<FBTNodeIndex,FAIMessageObserverHandle> TaskMessageObservers;
diff --git a/Engine/Source/Runtime/AIModule/Private/BehaviorTree/BehaviorTreeComponent.cpp b/Engine/Source/Runtime/AIModule/Private/BehaviorTree/BehaviorTreeComponent.cpp
--- a/Engine/Source/Runtime/AIModule/Private/BehaviorTree/BehaviorTreeComponent.cpp
+++ b/Engine/Source/Runtime/AIModule/Private/BehaviorTree/BehaviorTreeComponent.cpp
@@ -55,15 +55,13 @@ void UBehaviorTreeComponent::UninitializeComponent()
 void UBehaviorTreeComponent::RestartLogic()
 {
 	UE_VLOG(GetOwner(), LogBehaviorTree, Log, TEXT("UBehaviorTreeComponent::RestartLogic"));
-
-	bIsRunning = true;
 	RestartTree();
 }
 
 void UBehaviorTreeComponent::StopLogic(const FString& Reason) 
 {
 	UE_VLOG(GetOwner(), LogBehaviorTree, Log, TEXT("Stopping BT, reason: \'%s\'"), *Reason);
-	bIsRunning = false;
+	StopTree(EBTStopMode::Safe);
 }
 
 void UBehaviorTreeComponent::PauseLogic(const FString& Reason)
@@ -136,8 +134,9 @@ void UBehaviorTreeComponent::StartTree(UBehaviorTree& Asset, EBTExecutionMode::T
 
 	StopTree(EBTStopMode::Safe);
 
-	PendingInitialize.Asset = &Asset;
-	PendingInitialize.ExecuteMode = ExecuteMode;
+	TreeStartInfo.Asset = &Asset;
+	TreeStartInfo.ExecuteMode = ExecuteMode;
+	TreeStartInfo.bPendingInitialize = true;
 
 	ProcessPendingInitialize();
 }
@@ -153,7 +152,7 @@ void UBehaviorTreeComponent::ProcessPendingInitialize()
 	// finish cleanup
 	RemoveAllInstances();
 
-	bLoopExecution = (PendingInitialize.ExecuteMode == EBTExecutionMode::Looped);
+	bLoopExecution = (TreeStartInfo.ExecuteMode == EBTExecutionMode::Looped);
 	bIsRunning = true;
 
 #if USE_BEHAVIORTREE_DEBUGGER
@@ -167,9 +166,8 @@ void UBehaviorTreeComponent::ProcessPendingInitialize()
 	}
 
 	// push new instance
-	const bool bPushed = PushInstance(*PendingInitialize.Asset);
-
-	PendingInitialize = FBTPendingInitializeInfo();
+	const bool bPushed = PushInstance(*TreeStartInfo.Asset);
+	TreeStartInfo.bPendingInitialize = false;
 }
 
 void UBehaviorTreeComponent::StopTree(EBTStopMode::Type StopMode)
@@ -271,13 +269,26 @@ void UBehaviorTreeComponent::StopTree(EBTStopMode::Type StopMode)
 	// make sure to allow new execution requests
 	bRequestedFlowUpdate = false;
 	bRequestedStop = false;
+	bIsRunning = false;
 }
 
 void UBehaviorTreeComponent::RestartTree()
 {
 	UE_VLOG(GetOwner(), LogBehaviorTree, Log, TEXT("UBehaviorTreeComponent::RestartTree"));
 	
-	if (InstanceStack.Num())
+	if (!bIsRunning)
+	{
+		if (TreeStartInfo.IsSet())
+		{
+			TreeStartInfo.bPendingInitialize = true;
+			ProcessPendingInitialize();
+		}
+	}
+	else if (bRequestedStop)
+	{
+		TreeStartInfo.bPendingInitialize = true;
+	}
+	else if (InstanceStack.Num())
 	{
 		FBehaviorTreeInstance& TopInstance = InstanceStack[0];
 		RequestExecution(TopInstance.RootNode, 0, TopInstance.RootNode, -1, EBTNodeResult::Aborted);
@@ -364,7 +375,7 @@ void UBehaviorTreeComponent::OnTaskFinished(const UBTTaskNode* TaskNode, EBTNode
 		UpdateAbortingTasks();
 	}
 
-	if (PendingInitialize.IsSet())
+	if (TreeStartInfo.HasPendingInitialize())
 	{
 		ProcessPendingInitialize();
 	}
