Fixed issue with "skipping" an item when scoring it that its TestValue was being set to zero rather than UEnvQueryTypes::SkippedItemValue, which meant it was not displaying "SKIP" in debug display except if the item was set to "Filter Only".

#UE4 #DebuggingBugFix

diff --git a/Engine/Source/Runtime/AIModule/Private/EnvironmentQuery/EnvQueryTest.cpp b/Engine/Source/Runtime/AIModule/Private/EnvironmentQuery/EnvQueryTest.cpp
--- a/Engine/Source/Runtime/AIModule/Private/EnvironmentQuery/EnvQueryTest.cpp
+++ b/Engine/Source/Runtime/AIModule/Private/EnvironmentQuery/EnvQueryTest.cpp
@@ -140,7 +140,8 @@ void UEnvQueryTest::NormalizeItemScores(FEnvQueryInstance& QueryInstance)
 			}
 			else
 			{
-				TestValue = 0.0f;
+				// Do NOT clear TestValue to 0, because the SkippedItemValue is used to display "SKIP" when debugging.
+				// TestValue = 0.0f;
 				WeightedScore = 0.0f;
 			}
 
