CIS Fixes

diff --git a/Engine/Source/Runtime/AIModule/Classes/Perception/AISense_Damage.h b/Engine/Source/Runtime/AIModule/Classes/Perception/AISense_Damage.h
--- a/Engine/Source/Runtime/AIModule/Classes/Perception/AISense_Damage.h
+++ b/Engine/Source/Runtime/AIModule/Classes/Perception/AISense_Damage.h
@@ -4,6 +4,7 @@
 
 #include "Perception/AIPerceptionSystem.h"
 #include "Perception/AISense.h"
+#include "Perception/AIPerceptionListenerInterface.h"
 #include "AISense_Damage.generated.h"
 
 USTRUCT()
diff --git a/Engine/Source/Runtime/Engine/Private/Animation/BonePose.cpp b/Engine/Source/Runtime/Engine/Private/Animation/BonePose.cpp
--- a/Engine/Source/Runtime/Engine/Private/Animation/BonePose.cpp
+++ b/Engine/Source/Runtime/Engine/Private/Animation/BonePose.cpp
@@ -2,6 +2,7 @@
 
 #include "EnginePrivate.h"
 #include "BonePose.h"
+#include "AnimationRuntime.h"
 
 void FMeshPose::ResetToRefPose()
 {
