Back out revision 10 from //depot/UE4/Engine/Source/Runtime/AIModule/Classes/Perception/AISense_Damage.h #UE4

Changed the fix to just theclare the missing type

diff --git a/Engine/Source/Runtime/AIModule/Classes/Perception/AISense_Damage.h b/Engine/Source/Runtime/AIModule/Classes/Perception/AISense_Damage.h
--- a/Engine/Source/Runtime/AIModule/Classes/Perception/AISense_Damage.h
+++ b/Engine/Source/Runtime/AIModule/Classes/Perception/AISense_Damage.h
@@ -4,9 +4,10 @@
 
 #include "Perception/AIPerceptionSystem.h"
 #include "Perception/AISense.h"
-#include "Perception/AIPerceptionListenerInterface.h"
 #include "AISense_Damage.generated.h"
 
+class IAIPerceptionListenerInterface;
+
 USTRUCT()
 struct AIMODULE_API FAIDamageEvent
 {	
