fixed behavior tree decorators losing blackboard observers
#ue4 UE-16112

diff --git a/Engine/Source/Runtime/AIModule/Private/BehaviorTree/Decorators/BTDecorator_Blackboard.cpp b/Engine/Source/Runtime/AIModule/Private/BehaviorTree/Decorators/BTDecorator_Blackboard.cpp
--- a/Engine/Source/Runtime/AIModule/Private/BehaviorTree/Decorators/BTDecorator_Blackboard.cpp
+++ b/Engine/Source/Runtime/AIModule/Private/BehaviorTree/Decorators/BTDecorator_Blackboard.cpp
@@ -58,14 +58,11 @@ bool UBTDecorator_Blackboard::EvaluateOnBlackboard(const UBlackboardComponent& B
 EBlackboardNotificationResult UBTDecorator_Blackboard::OnBlackboardKeyValueChange(const UBlackboardComponent& Blackboard, FBlackboard::FKey ChangedKeyID)
 {
 	UBehaviorTreeComponent* BehaviorComp = (UBehaviorTreeComponent*)Blackboard.GetBrainComponent();
-
 	if (BehaviorComp == nullptr)
 	{
 		return EBlackboardNotificationResult::RemoveObserver;
 	}
 
-	EBlackboardNotificationResult Result = EBlackboardNotificationResult::ContinueObserving;
-
 	if (BlackboardKey.GetSelectedKeyID() == ChangedKeyID && GetFlowAbortMode() != EBTFlowAbortMode::None)
 	{
 		if (NotifyObserver == EBTBlackboardRestart::ResultChange)
@@ -82,13 +79,11 @@ EBlackboardNotificationResult UBTDecorator_Blackboard::OnBlackboardKeyValueChang
 			if ((bIsExecutingBranch && !bPass) || (!bIsExecutingBranch && bPass))
 			{
 				BehaviorComp->RequestExecution(this);
-				Result = EBlackboardNotificationResult::RemoveObserver;
 			}
 			else if (!bIsExecutingBranch && !bPass && GetParentNode() && GetParentNode()->Children.IsValidIndex(GetChildIndex()))
 			{
 				const UBTCompositeNode* BranchRoot = GetParentNode()->Children[GetChildIndex()].ChildComposite;
 				BehaviorComp->UnregisterAuxNodesInBranch(BranchRoot);
-				Result = EBlackboardNotificationResult::RemoveObserver;
 			}
 		}
 		else
@@ -98,11 +93,10 @@ EBlackboardNotificationResult UBTDecorator_Blackboard::OnBlackboardKeyValueChang
 				*Blackboard.GetKeyName(ChangedKeyID).ToString());
 
 			BehaviorComp->RequestExecution(this);
-			Result = EBlackboardNotificationResult::RemoveObserver;
 		}
 	}
 
-	return Result;
+	return EBlackboardNotificationResult::ContinueObserving;
 }
 
 void UBTDecorator_Blackboard::DescribeRuntimeValues(const UBehaviorTreeComponent& OwnerComp, uint8* NodeMemory, EBTDescriptionVerbosity::Type Verbosity, TArray<FString>& Values) const
