fixed behavior tree restarts from blackboard decorator in OnValueChange mode
#ue4 UE-16106

diff --git a/Engine/Source/Runtime/AIModule/Private/BehaviorTree/Decorators/BTDecorator_Blackboard.cpp b/Engine/Source/Runtime/AIModule/Private/BehaviorTree/Decorators/BTDecorator_Blackboard.cpp
--- a/Engine/Source/Runtime/AIModule/Private/BehaviorTree/Decorators/BTDecorator_Blackboard.cpp
+++ b/Engine/Source/Runtime/AIModule/Private/BehaviorTree/Decorators/BTDecorator_Blackboard.cpp
@@ -65,34 +65,38 @@ EBlackboardNotificationResult UBTDecorator_Blackboard::OnBlackboardKeyValueChang
 
 	if (BlackboardKey.GetSelectedKeyID() == ChangedKeyID && GetFlowAbortMode() != EBTFlowAbortMode::None)
 	{
-		if (NotifyObserver == EBTBlackboardRestart::ResultChange)
-		{
-			const bool bIsExecutingBranch = BehaviorComp->IsExecutingBranch(this, GetChildIndex());
-			const bool bPass = EvaluateOnBlackboard(Blackboard);
+		const bool bIsExecutingBranch = BehaviorComp->IsExecutingBranch(this, GetChildIndex());
+		const bool bPass = EvaluateOnBlackboard(Blackboard);
 
-			UE_VLOG(BehaviorComp->GetOwner(), LogBehaviorTree, Verbose, TEXT("%s, OnBlackboardKeyValueChange[%s] pass:%d executing:%d => %s"),
-				*UBehaviorTreeTypes::DescribeNodeHelper(this),
-				*Blackboard.GetKeyName(ChangedKeyID).ToString(), bPass, bIsExecutingBranch,
-				(bIsExecutingBranch && !bPass) || (!bIsExecutingBranch && bPass) ? TEXT("restart") : TEXT("skip"));
+		UE_VLOG(BehaviorComp->GetOwner(), LogBehaviorTree, Verbose, TEXT("%s, OnBlackboardChange[%s] pass:%d executing:%d => %s"),
+			*UBehaviorTreeTypes::DescribeNodeHelper(this),
+			*Blackboard.GetKeyName(ChangedKeyID).ToString(), bPass, bIsExecutingBranch,
+			(bIsExecutingBranch && !bPass) || (!bIsExecutingBranch && bPass) ? TEXT("restart") : TEXT("skip"));
 
-			// this is basically a XOR, but written down this way for readability's sake
-			if ((bIsExecutingBranch && !bPass) || (!bIsExecutingBranch && bPass))
-			{
-				BehaviorComp->RequestExecution(this);
-			}
-			else if (!bIsExecutingBranch && !bPass && GetParentNode() && GetParentNode()->Children.IsValidIndex(GetChildIndex()))
-			{
-				const UBTCompositeNode* BranchRoot = GetParentNode()->Children[GetChildIndex()].ChildComposite;
-				BehaviorComp->UnregisterAuxNodesInBranch(BranchRoot);
-			}
+		// this is basically a XOR, but written down this way for readability's sake
+		if ((bIsExecutingBranch && !bPass) || (!bIsExecutingBranch && bPass))
+		{
+			BehaviorComp->RequestExecution(this);
+		}
+		else if (!bIsExecutingBranch && !bPass && GetParentNode() && GetParentNode()->Children.IsValidIndex(GetChildIndex()))
+		{
+			const UBTCompositeNode* BranchRoot = GetParentNode()->Children[GetChildIndex()].ChildComposite;
+			BehaviorComp->UnregisterAuxNodesInBranch(BranchRoot);
 		}
-		else
+		else if (bIsExecutingBranch && bPass && NotifyObserver == EBTBlackboardRestart::ValueChange)
 		{
-			UE_VLOG(BehaviorComp->GetOwner(), LogBehaviorTree, Verbose, TEXT("%s, OnBlackboardKeyValueChange[%s] => restart"),
+			UE_VLOG(BehaviorComp->GetOwner(), LogBehaviorTree, Verbose, TEXT("%s, OnBlackboardChange[%s] => restart"),
 				*UBehaviorTreeTypes::DescribeNodeHelper(this),
 				*Blackboard.GetKeyName(ChangedKeyID).ToString());
 
-			BehaviorComp->RequestExecution(this);
+			// force result Aborted to restart from this decorator
+			// can't use helper function
+
+			const int32 InstanceIdx = BehaviorComp->FindInstanceContainingNode((UBTNode*)GetParentNode());
+			if (InstanceIdx != INDEX_NONE)
+			{
+				BehaviorComp->RequestExecution(GetParentNode(), InstanceIdx, this, GetChildIndex(), EBTNodeResult::Aborted);
+			}
 		}
 	}
 
