added InvalidOrientation constant to AISytem
#ue4

diff --git a/Engine/Source/Runtime/AIModule/Classes/AITypes.h b/Engine/Source/Runtime/AIModule/Classes/AITypes.h
--- a/Engine/Source/Runtime/AIModule/Classes/AITypes.h
+++ b/Engine/Source/Runtime/AIModule/Classes/AITypes.h
@@ -10,6 +10,7 @@
 namespace FAISystem
 {
 	static const FRotator InvalidRotation = FRotator(FLT_MAX);
+	static const FQuat InvalidOrientation = FQuat(FLT_MAX, FLT_MAX, FLT_MAX, FLT_MAX);
 	static const FVector InvalidLocation = FVector(FLT_MAX);
 	static const FVector InvalidDirection = FVector::ZeroVector; 
 	static const float InvalidRange = -1.f;
@@ -31,6 +32,11 @@ namespace FAISystem
 	{
 		return TestRotation != InvalidRotation;
 	}
+
+	FORCEINLINE bool IsValidOrientation(const FQuat& TestOrientation)
+	{
+		return TestOrientation != InvalidOrientation;
+	}
 }
 
 UENUM()
