added missing accessors to CrowdFollowingComponent
#ue4

diff --git a/Engine/Source/Runtime/AIModule/Classes/Navigation/CrowdFollowingComponent.h b/Engine/Source/Runtime/AIModule/Classes/Navigation/CrowdFollowingComponent.h
--- a/Engine/Source/Runtime/AIModule/Classes/Navigation/CrowdFollowingComponent.h
+++ b/Engine/Source/Runtime/AIModule/Classes/Navigation/CrowdFollowingComponent.h
@@ -75,12 +75,15 @@ class AIMODULE_API UCrowdFollowingComponent : public UPathFollowingComponent, pu
 	void SetCrowdSeparation(bool bEnable, bool bUpdateAgent = true);
 	void SetCrowdOptimizeVisibility(bool bEnable, bool bUpdateAgent = true);
 	void SetCrowdOptimizeTopology(bool bEnable, bool bUpdateAgent = true);
+	void SetCrowdPathOffset(bool bEnable, bool bUpdateAgent = true);
 	void SetCrowdSlowdownAtGoal(bool bEnable, bool bUpdateAgent = true);
 	void SetCrowdSeparationWeight(float Weight, bool bUpdateAgent = true);
 	void SetCrowdCollisionQueryRange(float Range, bool bUpdateAgent = true);
 	void SetCrowdPathOptimizationRange(float Range, bool bUpdateAgent = true);
 	void SetCrowdAvoidanceQuality(ECrowdAvoidanceQuality::Type Quality, bool bUpdateAgent = true);
 	void SetCrowdAvoidanceRangeMultiplier(float Multipler, bool bUpdateAgent = true);
+	void SetCrowdAffectFallingVelocity(bool bEnable);
+	void SetCrowdRotateToVelocity(bool bEnable);
 
 	FORCEINLINE bool IsCrowdSimulationEnabled() const { return bEnableCrowdSimulation; }
 	FORCEINLINE bool IsCrowdSimulatioSuspended() const { return bSuspendCrowdSimulation; }
@@ -91,6 +94,8 @@ class AIMODULE_API UCrowdFollowingComponent : public UPathFollowingComponent, pu
 	FORCEINLINE bool IsCrowdOptimizeTopologyEnabled() const { return bEnableOptimizeTopology; }
 	FORCEINLINE bool IsCrowdPathOffsetEnabled() const { return bEnablePathOffset; }
 	FORCEINLINE bool IsCrowdSlowdownAtGoalEnabled() const { return bEnableSlowdownAtGoal; }
+	FORCEINLINE bool IsCrowdAffectFallingVelocityEnabled() const { return bAffectFallingVelocity; }
+	FORCEINLINE bool IsCrowdRotateToVelocityEnabled() const { return bRotateToVelocity; }
 
 	FORCEINLINE bool IsCrowdSimulationActive() const { return IsCrowdSimulationEnabled() && !IsCrowdSimulatioSuspended(); }
 	/** checks if bEnableAnticipateTurns is set to true, and if crowd simulation is not suspended */
diff --git a/Engine/Source/Runtime/AIModule/Private/Navigation/CrowdFollowingComponent.cpp b/Engine/Source/Runtime/AIModule/Private/Navigation/CrowdFollowingComponent.cpp
--- a/Engine/Source/Runtime/AIModule/Private/Navigation/CrowdFollowingComponent.cpp
+++ b/Engine/Source/Runtime/AIModule/Private/Navigation/CrowdFollowingComponent.cpp
@@ -192,6 +192,29 @@ void UCrowdFollowingComponent::SetCrowdAvoidanceQuality(ECrowdAvoidanceQuality::
 	}
 }
 
+void UCrowdFollowingComponent::SetCrowdPathOffset(bool bEnable, bool bUpdateAgent)
+{
+	if (bEnablePathOffset != bEnable)
+	{
+		bEnablePathOffset = bEnable;
+
+		if (bUpdateAgent)
+		{
+			UpdateCrowdAgentParams();
+		}
+	}
+}
+
+void UCrowdFollowingComponent::SetCrowdAffectFallingVelocity(bool bEnable)
+{
+	bAffectFallingVelocity = bEnable;
+}
+
+void UCrowdFollowingComponent::SetCrowdRotateToVelocity(bool bEnable)
+{
+	bRotateToVelocity = bEnable;
+}
+
 void UCrowdFollowingComponent::SuspendCrowdSteering(bool bSuspend)
 {
 	if (bSuspendCrowdSimulation != bSuspend)
