Made BlackboardData a Blueprint type so that users can have blueprint variables pointing at BB assets #UE4

github PR#1200 by zorbathut

diff --git a/Engine/Source/Runtime/AIModule/Classes/BehaviorTree/BlackboardData.h b/Engine/Source/Runtime/AIModule/Classes/BehaviorTree/BlackboardData.h
--- a/Engine/Source/Runtime/AIModule/Classes/BehaviorTree/BlackboardData.h
+++ b/Engine/Source/Runtime/AIModule/Classes/BehaviorTree/BlackboardData.h
@@ -32,7 +32,7 @@ struct FBlackboardEntry
 	bool operator==(const FBlackboardEntry& Other) const;
 };
 
-UCLASS(AutoExpandCategories=(Blackboard))
+UCLASS(BlueprintType, AutoExpandCategories=(Blackboard))
 class AIMODULE_API UBlackboardData : public UDataAsset
 {
 	GENERATED_UCLASS_BODY()
