Some more dev->main intergation fall off fixes (related to GameplayTasks) #UE4

diff --git a/Engine/Source/Runtime/AIModule/Classes/AIController.h b/Engine/Source/Runtime/AIModule/Classes/AIController.h
--- a/Engine/Source/Runtime/AIModule/Classes/AIController.h
+++ b/Engine/Source/Runtime/AIModule/Classes/AIController.h
@@ -351,6 +351,7 @@ public:
 	virtual UGameplayTasksComponent* GetGameplayTasksComponent() override { return CachedGameplayTasksComponent; }
 	virtual void TaskStarted(UGameplayTask& NewTask) {}
 	virtual void TaskEnded(UGameplayTask& Task) {}
+	virtual void OnTaskInitialized(UGameplayTask& Task) override {}
 	virtual AActor* GetOwnerActor() const { return const_cast<AAIController*>(this); }
 	virtual AActor* GetAvatarActor() const { return GetPawn(); }
 
