CIS fix #UE4

diff --git a/Engine/Source/Runtime/AIModule/Private/AIModulePrivate.h b/Engine/Source/Runtime/AIModule/Private/AIModulePrivate.h
--- a/Engine/Source/Runtime/AIModule/Private/AIModulePrivate.h
+++ b/Engine/Source/Runtime/AIModule/Private/AIModulePrivate.h
@@ -70,6 +70,8 @@
 #include "BlueprintNodeHelpers.h"
 #include "BehaviorTree/BTFunctionLibrary.h"
 
+#include "AIResources.h"
+
 DECLARE_LOG_CATEGORY_EXTERN(LogAINavigation, Warning, All);
 DECLARE_LOG_CATEGORY_EXTERN(LogBehaviorTree, Display, All);
 DECLARE_LOG_CATEGORY_EXTERN(LogCrowdFollowing, Warning, All);
\ No newline at end of file
