Moved AI asset types into a separate AI category in content browser's asset creation menu #UE4

diff --git a/Engine/Source/Editor/BehaviorTreeEditor/Private/AssetTypeActions_BehaviorTree.cpp b/Engine/Source/Editor/BehaviorTreeEditor/Private/AssetTypeActions_BehaviorTree.cpp
--- a/Engine/Source/Editor/BehaviorTreeEditor/Private/AssetTypeActions_BehaviorTree.cpp
+++ b/Engine/Source/Editor/BehaviorTreeEditor/Private/AssetTypeActions_BehaviorTree.cpp
@@ -11,8 +11,16 @@
 #include "AssetTypeActions_BehaviorTree.h"
 #include "SBehaviorTreeDiff.h"
 
+#include "AIModule.h"
+
 #define LOCTEXT_NAMESPACE "AssetTypeActions"
 
+uint32 FAssetTypeActions_BehaviorTree::GetCategories() 
+{ 
+	IAIModule& AIModule = FModuleManager::GetModuleChecked<IAIModule>("AIModule").Get();
+	return AIModule.GetAIAssetCategoryBit();
+}
+
 void FAssetTypeActions_BehaviorTree::OpenAssetEditor( const TArray<UObject*>& InObjects, TSharedPtr<IToolkitHost> EditWithinLevelEditor )
 {
 	EToolkitMode::Type Mode = EditWithinLevelEditor.IsValid() ? EToolkitMode::WorldCentric : EToolkitMode::Standalone;
diff --git a/Engine/Source/Editor/BehaviorTreeEditor/Private/AssetTypeActions_BehaviorTree.h b/Engine/Source/Editor/BehaviorTreeEditor/Private/AssetTypeActions_BehaviorTree.h
--- a/Engine/Source/Editor/BehaviorTreeEditor/Private/AssetTypeActions_BehaviorTree.h
+++ b/Engine/Source/Editor/BehaviorTreeEditor/Private/AssetTypeActions_BehaviorTree.h
@@ -11,7 +11,7 @@ public:
 	virtual FColor GetTypeColor() const override { return FColor(149,70,255); }
 	virtual UClass* GetSupportedClass() const override;
 	virtual void OpenAssetEditor( const TArray<UObject*>& InObjects, TSharedPtr<class IToolkitHost> EditWithinLevelEditor = TSharedPtr<IToolkitHost>() ) override;
-	virtual uint32 GetCategories() override { return EAssetTypeCategories::Misc; }
+	virtual uint32 GetCategories() override;
 	virtual void PerformAssetDiff(UObject* Asset1, UObject* Asset2, const struct FRevisionInfo& OldRevision, const struct FRevisionInfo& NewRevision) const override;
 
 private:
diff --git a/Engine/Source/Editor/BehaviorTreeEditor/Private/AssetTypeActions_Blackboard.cpp b/Engine/Source/Editor/BehaviorTreeEditor/Private/AssetTypeActions_Blackboard.cpp
--- a/Engine/Source/Editor/BehaviorTreeEditor/Private/AssetTypeActions_Blackboard.cpp
+++ b/Engine/Source/Editor/BehaviorTreeEditor/Private/AssetTypeActions_Blackboard.cpp
@@ -10,6 +10,8 @@
 #include "AssetTypeActions_Blackboard.h"
 #include "AssetTypeActions_BehaviorTree.h"
 
+#include "AIModule.h"
+
 #define LOCTEXT_NAMESPACE "AssetTypeActions"
 
 UClass* FAssetTypeActions_Blackboard::GetSupportedClass() const
@@ -32,4 +34,10 @@ void FAssetTypeActions_Blackboard::OpenAssetEditor( const TArray<UObject*>& InOb
 	}	
 }
 
+uint32 FAssetTypeActions_Blackboard::GetCategories()
+{ 
+	IAIModule& AIModule = FModuleManager::GetModuleChecked<IAIModule>("AIModule").Get();
+	return AIModule.GetAIAssetCategoryBit();
+}
+
 #undef LOCTEXT_NAMESPACE
diff --git a/Engine/Source/Editor/BehaviorTreeEditor/Private/AssetTypeActions_Blackboard.h b/Engine/Source/Editor/BehaviorTreeEditor/Private/AssetTypeActions_Blackboard.h
--- a/Engine/Source/Editor/BehaviorTreeEditor/Private/AssetTypeActions_Blackboard.h
+++ b/Engine/Source/Editor/BehaviorTreeEditor/Private/AssetTypeActions_Blackboard.h
@@ -11,5 +11,5 @@ public:
 	virtual FColor GetTypeColor() const override { return FColor(201, 29, 85); }
 	virtual UClass* GetSupportedClass() const override;
 	virtual void OpenAssetEditor( const TArray<UObject*>& InObjects, TSharedPtr<class IToolkitHost> EditWithinLevelEditor = TSharedPtr<IToolkitHost>() ) override;
-	virtual uint32 GetCategories() override { return EAssetTypeCategories::Misc; }
+	virtual uint32 GetCategories() override;
 };
\ No newline at end of file
diff --git a/Engine/Source/Editor/EnvironmentQueryEditor/Private/AssetTypeActions_EnvironmentQuery.cpp b/Engine/Source/Editor/EnvironmentQueryEditor/Private/AssetTypeActions_EnvironmentQuery.cpp
--- a/Engine/Source/Editor/EnvironmentQueryEditor/Private/AssetTypeActions_EnvironmentQuery.cpp
+++ b/Engine/Source/Editor/EnvironmentQueryEditor/Private/AssetTypeActions_EnvironmentQuery.cpp
@@ -7,8 +7,10 @@
 #include "Editor/EnvironmentQueryEditor/Public/IEnvironmentQueryEditor.h"
 
 #include "EnvironmentQuery/EnvQuery.h"
-
 #include "AssetTypeActions_EnvironmentQuery.h"
+
+#include "AIModule.h"
+
 #define LOCTEXT_NAMESPACE "AssetTypeActions"
 
 void FAssetTypeActions_EnvironmentQuery::OpenAssetEditor( const TArray<UObject*>& InObjects, TSharedPtr<IToolkitHost> EditWithinLevelEditor )
@@ -31,5 +33,9 @@ UClass* FAssetTypeActions_EnvironmentQuery::GetSupportedClass() const
 	return UEnvQuery::StaticClass(); 
 }
 
-
+uint32 FAssetTypeActions_EnvironmentQuery::GetCategories()
+{
+	IAIModule& AIModule = FModuleManager::GetModuleChecked<IAIModule>("AIModule").Get();
+	return AIModule.GetAIAssetCategoryBit();
+}
 #undef LOCTEXT_NAMESPACE
diff --git a/Engine/Source/Editor/EnvironmentQueryEditor/Private/AssetTypeActions_EnvironmentQuery.h b/Engine/Source/Editor/EnvironmentQueryEditor/Private/AssetTypeActions_EnvironmentQuery.h
--- a/Engine/Source/Editor/EnvironmentQueryEditor/Private/AssetTypeActions_EnvironmentQuery.h
+++ b/Engine/Source/Editor/EnvironmentQueryEditor/Private/AssetTypeActions_EnvironmentQuery.h
@@ -11,5 +11,5 @@ public:
 	virtual FColor GetTypeColor() const override { return FColor(129,50,255); }
 	virtual UClass* GetSupportedClass() const override;
 	virtual void OpenAssetEditor( const TArray<UObject*>& InObjects, TSharedPtr<class IToolkitHost> EditWithinLevelEditor = TSharedPtr<IToolkitHost>() ) override;
-	virtual uint32 GetCategories() override { return EAssetTypeCategories::Misc; }
+	virtual uint32 GetCategories() override;
 };
\ No newline at end of file
diff --git a/Engine/Source/Runtime/AIModule/Private/AIModule.cpp b/Engine/Source/Runtime/AIModule/Private/AIModule.cpp
--- a/Engine/Source/Runtime/AIModule/Private/AIModule.cpp
+++ b/Engine/Source/Runtime/AIModule/Private/AIModule.cpp
@@ -1,11 +1,19 @@
 // Copyright 1998-2015 Epic Games, Inc. All Rights Reserved.
 #include "AIModulePrivate.h"
 #include "AISystem.h"
-#if WITH_EDITOR && ENABLE_VISUAL_LOG
-#	include "VisualLoggerExtension.h"
+
+#if WITH_EDITOR
+#include "Developer/AssetTools/Public/IAssetTools.h"
+#include "Developer/AssetTools/Public/AssetToolsModule.h"
+#if ENABLE_VISUAL_LOG
+	#include "VisualLoggerExtension.h"
+#endif // ENABLE_VISUAL_LOG
 #endif
+
 #include "AIModule.h"
 
+#define LOCTEXT_NAMESPACE "AIModule"
+
 DEFINE_LOG_CATEGORY_STATIC(LogAIModule, Log, All);
 
 class FAIModule : public IAIModule
@@ -16,10 +24,13 @@ class FAIModule : public IAIModule
 	
 	virtual UAISystemBase* CreateAISystemInstance(UWorld* World) override;
 	// End IModuleInterface
-
+#if WITH_EDITOR
+	virtual EAssetTypeCategories::Type GetAIAssetCategoryBit() const override { return AIAssetCategoryBit; }
 protected:
-#if WITH_EDITOR && ENABLE_VISUAL_LOG
+	EAssetTypeCategories::Type AIAssetCategoryBit;
+#if ENABLE_VISUAL_LOG
 	FVisualLoggerExtension	VisualLoggerExtension;
+#endif // ENABLE_VISUAL_LOG
 #endif
 };
 
@@ -34,6 +45,15 @@ void FAIModule::StartupModule()
 
 #if WITH_EDITOR 
 	FModuleManager::LoadModulePtr< IModuleInterface >("AITestSuite");
+
+	if (GIsEditor)
+	{
+		// Register asset types
+		IAssetTools& AssetTools = FModuleManager::LoadModuleChecked<FAssetToolsModule>("AssetTools").Get();
+
+		// register AI category so that AI assets can register to it
+		AIAssetCategoryBit = AssetTools.RegisterAdvancedAssetCategory(FName(TEXT("AI")), LOCTEXT("AIAssetCategory", "Artificial Inteligence"));
+	}
 #endif // WITH_EDITOR 
 
 #if WITH_EDITOR && ENABLE_VISUAL_LOG
@@ -57,3 +77,5 @@ UAISystemBase* FAIModule::CreateAISystemInstance(UWorld* World)
 	TSubclassOf<UAISystemBase> AISystemClass = LoadClass<UAISystemBase>(NULL, *UAISystem::GetAISystemClassName().ToString(), NULL, LOAD_None, NULL);
 	return NewObject<UAISystemBase>(World, AISystemClass);
 }
+
+#undef LOCTEXT_NAMESPACE
diff --git a/Engine/Source/Runtime/AIModule/Public/AIModule.h b/Engine/Source/Runtime/AIModule/Public/AIModule.h
--- a/Engine/Source/Runtime/AIModule/Public/AIModule.h
+++ b/Engine/Source/Runtime/AIModule/Public/AIModule.h
@@ -4,6 +4,11 @@
 
 #include "AI/AISystemBase.h"
 
+namespace EAssetTypeCategories
+{
+	enum Type;
+}
+
 /**
  * The public interface to this module
  */
@@ -32,5 +37,9 @@ public:
 	{
 		return FModuleManager::Get().IsModuleLoaded( "AIModule" );
 	}
+
+#if WITH_EDITOR
+	virtual EAssetTypeCategories::Type GetAIAssetCategoryBit() const = 0;
+#endif // WITH_EDITOR
 };
 
