Removed enum forward declaration from AIModule.h #UE4

diff --git a/Engine/Source/Runtime/AIModule/Public/AIModule.h b/Engine/Source/Runtime/AIModule/Public/AIModule.h
--- a/Engine/Source/Runtime/AIModule/Public/AIModule.h
+++ b/Engine/Source/Runtime/AIModule/Public/AIModule.h
@@ -3,11 +3,9 @@
 #pragma once
 
 #include "AI/AISystemBase.h"
-
-namespace EAssetTypeCategories
-{
-	enum Type;
-}
+#if WITH_EDITOR
+#include "Developer/AssetTools/Public/IAssetTools.h"
+#endif // WITH_EDITOR
 
 /**
  * The public interface to this module
