Added more debug data about EQS to Visual Logs
Fixed issue with visual logger colors - something was broken with them for filters/categories.

diff --git a/Engine/Source/Developer/LogVisualizer/Private/LogVisualizerStyle.cpp b/Engine/Source/Developer/LogVisualizer/Private/LogVisualizerStyle.cpp
--- a/Engine/Source/Developer/LogVisualizer/Private/LogVisualizerStyle.cpp
+++ b/Engine/Source/Developer/LogVisualizer/Private/LogVisualizerStyle.cpp
@@ -54,6 +54,11 @@ TSharedRef< FSlateStyleSet > FLogVisualizerStyle::Create()
 		.SetHighlightColor(FLinearColor(0.02f, 0.3f, 0.0f))
 		.SetHighlightShape(BOX_BRUSH("Common/TextBlockHighlightShape", FMargin(3.f / 8.f)));
 
+	const FTextBlockStyle LogsText = FTextBlockStyle(NormalText)
+		.SetFont(TTF_CORE_FONT("Fonts/DroidSansMono", 9));
+
+	Style.Set("TextLogs.Text", LogsText);
+
 	{
 		Style.Set("LogVisualizerApp.TabIcon", new IMAGE_BRUSH("Icons/icon_tab_DebugTools_40x", Icon16x16));
 		
diff --git a/Engine/Source/Developer/LogVisualizer/Private/SVisualLoggerFilters.cpp b/Engine/Source/Developer/LogVisualizer/Private/SVisualLoggerFilters.cpp
--- a/Engine/Source/Developer/LogVisualizer/Private/SVisualLoggerFilters.cpp
+++ b/Engine/Source/Developer/LogVisualizer/Private/SVisualLoggerFilters.cpp
@@ -285,7 +285,7 @@ void SVisualLoggerFilters::AddFilter(const FString& InFilterName)
 		}
 	}
 
-	const FLinearColor Color = FLogVisualizer::Get().GetColorForCategory(Filters.Num());
+	const FLinearColor Color = FLogVisualizer::Get().GetColorForCategory(InFilterName);
 	TSharedRef<SFilterWidget> NewFilter =
 		SNew(SFilterWidget)
 		.FilterName(*InFilterName)
diff --git a/Engine/Source/Developer/LogVisualizer/Private/SVisualLoggerLogsList.cpp b/Engine/Source/Developer/LogVisualizer/Private/SVisualLoggerLogsList.cpp
--- a/Engine/Source/Developer/LogVisualizer/Private/SVisualLoggerLogsList.cpp
+++ b/Engine/Source/Developer/LogVisualizer/Private/SVisualLoggerLogsList.cpp
@@ -105,6 +105,7 @@ TSharedRef<ITableRow> SVisualLoggerLogsList::LogEntryLinesGenerateRow(TSharedPtr
 						.ColorAndOpacity(FSlateColor(Item->Verbosity == ELogVerbosity::Error ? FLinearColor::Red : (Item->Verbosity == ELogVerbosity::Warning ? FLinearColor::Yellow : FLinearColor::Gray)))
 						.Text(FText::FromString(Item->Line))
 						.HighlightText(this, &SVisualLoggerLogsList::GetFilterText)
+						.TextStyle(FLogVisualizerStyle::Get(), TEXT("TextLogs.Text"))
 					]
 			];
 
diff --git a/Engine/Source/Runtime/AIModule/Classes/EnvironmentQuery/EnvQueryDebugHelpers.h b/Engine/Source/Runtime/AIModule/Classes/EnvironmentQuery/EnvQueryDebugHelpers.h
--- a/Engine/Source/Runtime/AIModule/Classes/EnvironmentQuery/EnvQueryDebugHelpers.h
+++ b/Engine/Source/Runtime/AIModule/Classes/EnvironmentQuery/EnvQueryDebugHelpers.h
@@ -155,6 +155,7 @@ public:
 #if USE_EQS_DEBUGGER
 	static void QueryToDebugData(FEnvQueryInstance& Query, EQSDebug::FQueryData& EQSLocalData);
 	static void QueryToBlobArray(FEnvQueryInstance& Query, TArray<uint8>& BlobArray, bool bUseCompression = false);
+	static void DebugDataToBlobArray(EQSDebug::FQueryData& QueryData, TArray<uint8>& BlobArray, bool bUseCompression = false);
 	static void BlobArrayToDebugData(const TArray<uint8>& BlobArray, EQSDebug::FQueryData& EQSLocalData, bool bUseCompression = false);
 #endif
 
diff --git a/Engine/Source/Runtime/AIModule/Private/EnvironmentQuery/EnvQueryDebugHelpers.cpp b/Engine/Source/Runtime/AIModule/Private/EnvironmentQuery/EnvQueryDebugHelpers.cpp
--- a/Engine/Source/Runtime/AIModule/Private/EnvironmentQuery/EnvQueryDebugHelpers.cpp
+++ b/Engine/Source/Runtime/AIModule/Private/EnvironmentQuery/EnvQueryDebugHelpers.cpp
@@ -12,6 +12,11 @@ void UEnvQueryDebugHelpers::QueryToBlobArray(FEnvQueryInstance& Query, TArray<ui
 	EQSDebug::FQueryData EQSLocalData;
 	UEnvQueryDebugHelpers::QueryToDebugData(Query, EQSLocalData);
 
+	DebugDataToBlobArray(EQSLocalData, BlobArray, bUseCompression);
+}
+
+void UEnvQueryDebugHelpers::DebugDataToBlobArray(EQSDebug::FQueryData& EQSLocalData, TArray<uint8>& BlobArray, bool bUseCompression)
+{
 	if (!bUseCompression)
 	{
 		FMemoryWriter ArWriter(BlobArray);
@@ -132,17 +137,94 @@ void  UEnvQueryDebugHelpers::BlobArrayToDebugData(const TArray<uint8>& BlobArray
 #endif //ENABLE_VISUAL_LOG
 
 #if ENABLE_VISUAL_LOG && USE_EQS_DEBUGGER
+
+#define PRINT_TABLE_ROW(__src__, __dst__, __len__) \
+	{ \
+		const FString DataToPrint = __src__; \
+		__dst__ += DataToPrint; \
+		for (int32 i = DataToPrint.Len(); i < __len__; i++) __dst__ += FString::Printf(TEXT(" ")); \
+	}
+
+namespace
+{
+	static void DrawEQSItemDetails(const EQSDebug::FQueryData& EQSLocalData, int32 ItemIdx, FVisualLogLine& LogLine)
+	{
+		const EQSDebug::FItemData& ItemData = EQSLocalData.Items[ItemIdx];
+
+		PRINT_TABLE_ROW(ItemData.Desc, LogLine.Line, 27);
+
+		PRINT_TABLE_ROW(FString::Printf(TEXT(" | %.2f"), ItemData.TotalScore), LogLine.Line, 8);
+
+		const int32 NumTests = ItemData.TestScores.Num();
+		float TotalWeightedScore = 0.0f;
+		for (int32 Idx = 0; Idx < NumTests; Idx++)
+		{
+			TotalWeightedScore += ItemData.TestScores[Idx];
+		}
+
+		for (int32 Idx = 0; Idx < NumTests; Idx++)
+		{
+			const float ScoreW = ItemData.TestScores[Idx];
+			const float ScoreN = ItemData.TestValues[Idx];
+			FString DescScoreW = FString::Printf(TEXT("%.2f"), ScoreW);
+			FString DescScoreN = (ScoreN == UEnvQueryTypes::SkippedItemValue) ? TEXT("SKIP") : FString::Printf(TEXT("%.2f"), ScoreN);
+
+			PRINT_TABLE_ROW(FString(TEXT(" | ")) + DescScoreW + FString(" ") + DescScoreN, LogLine.Line, 15);
+		}
+	}
+}
+
 void UEnvQueryDebugHelpers::LogQueryInternal(FEnvQueryInstance& Query, const FLogCategoryBase& Category, ELogVerbosity::Type Verbosity, float TimeSeconds, FVisualLogEntry *CurrentEntry)
 {
 	const int32 UniqueId = FVisualLogger::Get().GetUniqueId(TimeSeconds);
+
+	EQSDebug::FQueryData EQSLocalData;
+	UEnvQueryDebugHelpers::QueryToDebugData(Query, EQSLocalData);
+
 	TArray<uint8> BlobArray;
-	UEnvQueryDebugHelpers::QueryToBlobArray(Query, BlobArray);
+	DebugDataToBlobArray(EQSLocalData, BlobArray);
+
 	const FString AdditionalLogInfo = FString::Printf(TEXT("Executed EQS: \n - Name: '%s' (id=%d, option=%d),\n - All Items: %d,\n - ValidItems: %d"), *Query.QueryName, Query.QueryID, Query.OptionIndex, Query.ItemDetails.Num(), Query.NumValidItems);
 	FVisualLogLine Line(Category.GetCategoryName(), Verbosity, AdditionalLogInfo, Query.QueryID);
 	Line.TagName = *EVisLogTags::TAG_EQS;
 	Line.UniqueId = UniqueId;
 
-	CurrentEntry->LogLines.Add(Line);
 	CurrentEntry->AddDataBlock(EVisLogTags::TAG_EQS, BlobArray, Category.GetCategoryName(), Verbosity).UniqueId = UniqueId;
+
+
+	if (EQSLocalData.NumValidItems > 0)
+	{
+		const int32 NumTests = EQSLocalData.Tests.Num();
+		Line.Line += FString(TEXT("\n"));
+
+		// draw test weights for best X items
+		const int32 NumItems = EQSLocalData.Items.Num();
+
+		// table header		
+		{
+			FString HeaderString;
+			PRINT_TABLE_ROW(FString::Printf(TEXT("Item "), EQSLocalData.NumValidItems), HeaderString, 27);
+
+			PRINT_TABLE_ROW(FString::Printf(TEXT(" | Score"), EQSLocalData.NumValidItems), HeaderString, 8);
+
+			for (int32 TestIdx = 0; TestIdx < NumTests; TestIdx++)
+			{
+				PRINT_TABLE_ROW(FString::Printf(TEXT(" | Test %d  "), TestIdx), HeaderString, 15);
+			}
+
+			Line.Line += HeaderString;
+		}
+		Line.Line += FString(TEXT("\n"));
+
+		// valid items
+		for (int32 Idx = 0; Idx < NumItems; Idx++)
+		{
+			DrawEQSItemDetails(EQSLocalData, Idx, Line);
+			Line.Line += FString(TEXT("\n"));
+		}
+	}
+	CurrentEntry->LogLines.Add(Line);
 }
+
+#undef PRINT_TABLE_ROW
 #endif //ENABLE_VISUAL_LOG && USE_EQS_DEBUGGER
