Added a dedicated broadcast delegate to notify blueprints about single perception target's single sense stimulus #UE4

diff --git a/Engine/Source/Runtime/AIModule/Classes/Perception/AIPerceptionComponent.h b/Engine/Source/Runtime/AIModule/Classes/Perception/AIPerceptionComponent.h
--- a/Engine/Source/Runtime/AIModule/Classes/Perception/AIPerceptionComponent.h
+++ b/Engine/Source/Runtime/AIModule/Classes/Perception/AIPerceptionComponent.h
@@ -15,6 +15,7 @@ class UAIPerceptionSystem;
 class UAISenseConfig;
 
 DECLARE_DYNAMIC_MULTICAST_DELEGATE_OneParam(FPerceptionUpdatedDelegate, TArray<AActor*>, UpdatedActors);
+DECLARE_DYNAMIC_MULTICAST_DELEGATE_TwoParams(FActorPerceptionUpdatedDelegate, AActor*, Actor, FAIStimulus, Stimulus);
 
 struct AIMODULE_API FActorPerceptionInfo
 {
@@ -22,7 +23,7 @@ struct AIMODULE_API FActorPerceptionInfo
 
 	TArray<FAIStimulus> LastSensedStimuli;
 
-	/** if != MAX indicated the sense that takes precedense over other senses when it comes
+	/** if != MAX indicates the sense that takes precedense over other senses when it comes
 		to determining last stimulus location */
 	FAISenseID DominantSense;
 
@@ -256,6 +257,9 @@ public:
 	UPROPERTY(BlueprintAssignable)
 	FPerceptionUpdatedDelegate OnPerceptionUpdated;
 
+	UPROPERTY(BlueprintAssignable)
+	FActorPerceptionUpdatedDelegate OnTargetPerceptionUpdated;
+
 protected:
 
 	void UpdatePerceptionFilter(FAISenseID Channel, bool bNewValue);
diff --git a/Engine/Source/Runtime/AIModule/Classes/Perception/AISenseConfig.h b/Engine/Source/Runtime/AIModule/Classes/Perception/AISenseConfig.h
--- a/Engine/Source/Runtime/AIModule/Classes/Perception/AISenseConfig.h
+++ b/Engine/Source/Runtime/AIModule/Classes/Perception/AISenseConfig.h
@@ -26,7 +26,8 @@ protected:
 public:
 	UAISenseConfig(const FObjectInitializer& ObjectInitializer = FObjectInitializer::Get());
 
-	virtual TSubclassOf<UAISense> GetSenseImplementation() const { return UAISense::StaticClass(); }
+	virtual TSubclassOf<UAISense> GetSenseImplementation() const;
+
 	FAISenseID GetSenseID() const;
 	
 	float GetMaxAge() const { return MaxAge; }
diff --git a/Engine/Source/Runtime/AIModule/Private/Perception/AIPerceptionComponent.cpp b/Engine/Source/Runtime/AIModule/Private/Perception/AIPerceptionComponent.cpp
--- a/Engine/Source/Runtime/AIModule/Private/Perception/AIPerceptionComponent.cpp
+++ b/Engine/Source/Runtime/AIModule/Private/Perception/AIPerceptionComponent.cpp
@@ -379,6 +379,8 @@ void UAIPerceptionComponent::ProcessStimuli()
 		UE_VLOG(GetOwner(), LogAIPerception, Warning, TEXT("UAIPerceptionComponent::ProcessStimuli called without any Stimuli to process"));
 		return;
 	}
+
+	const bool bBroadcastEveryTargetUpdate = OnTargetPerceptionUpdated.IsBound();
 	
 	FStimulusToProcess* SourcedStimulus = StimuliToProcess.GetData();
 	TArray<AActor*> UpdatedActors;
@@ -419,12 +421,7 @@ void UAIPerceptionComponent::ProcessStimuli()
 		check(SourcedStimulus->Stimulus.Type.IsValid());
 
 		FAIStimulus& StimulusStore = PerceptualInfo->LastSensedStimuli[SourcedStimulus->Stimulus.Type];
-
-		// if the new stimulus is "valid" or it's info that "no longer sensed" and it used to be sensed successfully
-		if (SourcedStimulus->Stimulus.WantsToNotifyOnlyOnPerceptionChange() == false || SourcedStimulus->Stimulus.WasSuccessfullySensed() != StimulusStore.WasSuccessfullySensed())
-		{
-			UpdatedActors.AddUnique(SourcedStimulus->Source);
-		}
+		const bool bActorInfoUpdated = SourcedStimulus->Stimulus.WantsToNotifyOnlyOnPerceptionChange() == false || SourcedStimulus->Stimulus.WasSuccessfullySensed() != StimulusStore.WasSuccessfullySensed();
 
 		if (SourcedStimulus->Stimulus.WasSuccessfullySensed())
 		{
@@ -440,6 +437,16 @@ void UAIPerceptionComponent::ProcessStimuli()
 			// may be useful in future
 			StimulusStore.MarkNoLongerSensed();
 		}
+
+		// if the new stimulus is "valid" or it's info that "no longer sensed" and it used to be sensed successfully
+		if (bActorInfoUpdated)
+		{
+			UpdatedActors.AddUnique(SourcedStimulus->Source);
+			if (bBroadcastEveryTargetUpdate)
+			{
+				OnTargetPerceptionUpdated.Broadcast(SourcedStimulus->Source, StimulusStore);
+			}
+		}
 	}
 
 	StimuliToProcess.Reset();
diff --git a/Engine/Source/Runtime/AIModule/Private/Perception/AISense.cpp b/Engine/Source/Runtime/AIModule/Private/Perception/AISense.cpp
--- a/Engine/Source/Runtime/AIModule/Private/Perception/AISense.cpp
+++ b/Engine/Source/Runtime/AIModule/Private/Perception/AISense.cpp
@@ -111,6 +111,11 @@ UAISenseConfig::UAISenseConfig(const FObjectInitializer& ObjectInitializer)
 {
 }
 
+TSubclassOf<UAISense> UAISenseConfig::GetSenseImplementation() const 
+{ 
+	return UAISense::StaticClass(); 
+}
+
 #if !UE_BUILD_SHIPPING
 void UAISenseConfig::DrawDebugInfo(UCanvas& Canvas, UAIPerceptionComponent& PerceptionComponent) const
 {
