Made GetSenseClassForStimulus's WorldContext param default to self #UE4

diff --git a/Engine/Source/Runtime/AIModule/Classes/Perception/AIPerceptionSystem.h b/Engine/Source/Runtime/AIModule/Classes/Perception/AIPerceptionSystem.h
--- a/Engine/Source/Runtime/AIModule/Classes/Perception/AIPerceptionSystem.h
+++ b/Engine/Source/Runtime/AIModule/Classes/Perception/AIPerceptionSystem.h
@@ -144,7 +144,7 @@ public:
 
 	FAISenseID RegisterSenseClass(TSubclassOf<UAISense> SenseClass);
 
-	UFUNCTION(BlueprintCallable, Category = "AI|Perception")
+	UFUNCTION(BlueprintCallable, Category = "AI|Perception", meta = (HidePin = "WorldContext", DefaultToSelf = "WorldContext"))
 	static TSubclassOf<UAISense> GetSenseClassForStimulus(UObject* WorldContext, const FAIStimulus& Stimulus);
 	
 protected:
