fixed rare crash on aborting behavior tree task connected to parallel node
#ue4 UE-16959

diff --git a/Engine/Source/Runtime/AIModule/Private/BehaviorTree/BehaviorTreeComponent.cpp b/Engine/Source/Runtime/AIModule/Private/BehaviorTree/BehaviorTreeComponent.cpp
--- a/Engine/Source/Runtime/AIModule/Private/BehaviorTree/BehaviorTreeComponent.cpp
+++ b/Engine/Source/Runtime/AIModule/Private/BehaviorTree/BehaviorTreeComponent.cpp
@@ -915,14 +915,22 @@ void UBehaviorTreeComponent::ApplySearchUpdates(const TArray<FBehaviorTreeSearch
 					*UBehaviorTreeTypes::DescribeNodeHelper(UpdateInfo.TaskNode),
 					(NodeResult == EBTNodeResult::InProgress) ? TEXT("in progress") : TEXT("instant"));
 
-				// mark as pending abort
-				if (NodeResult == EBTNodeResult::InProgress)
+				// check if task node is still valid, could've received LatentAbortFinished during AbortTask call
+				const bool bStillValid = InstanceStack.IsValidIndex(UpdateInfo.InstanceIndex) &&
+					InstanceStack[UpdateInfo.InstanceIndex].ParallelTasks.IsValidIndex(ParallelTaskIdx) &&
+					InstanceStack[UpdateInfo.InstanceIndex].ParallelTasks[ParallelTaskIdx] == UpdateInfo.TaskNode;
+				
+				if (bStillValid)
 				{
-					UpdateInstance.ParallelTasks[ParallelTaskIdx].Status = EBTTaskStatus::Aborting;
-					bWaitingForAbortingTasks = true;
-				}
+					// mark as pending abort
+					if (NodeResult == EBTNodeResult::InProgress)
+					{
+						UpdateInstance.ParallelTasks[ParallelTaskIdx].Status = EBTTaskStatus::Aborting;
+						bWaitingForAbortingTasks = true;
+					}
 
-				OnTaskFinished(UpdateInfo.TaskNode, NodeResult);
+					OnTaskFinished(UpdateInfo.TaskNode, NodeResult);
+				}
 			}
 			else
 			{
