Made AIController derive from IGenericTeamAgentInterface #UE4

diff --git a/Engine/Source/Runtime/AIModule/Classes/AIController.h b/Engine/Source/Runtime/AIModule/Classes/AIController.h
--- a/Engine/Source/Runtime/AIModule/Classes/AIController.h
+++ b/Engine/Source/Runtime/AIModule/Classes/AIController.h
@@ -12,6 +12,7 @@
 #include "Perception/AIPerceptionListenerInterface.h"
 #include "BehaviorTree/BehaviorTreeTypes.h"
 #include "GameplayTaskOwnerInterface.h"
+#include "GenericTeamAgentInterface.h"
 #include "AIController.generated.h"
 
 class APawn;
@@ -74,7 +75,7 @@ struct FFocusKnowledge
  */
 
 UCLASS(ClassGroup = AI, BlueprintType, Blueprintable)
-class AIMODULE_API AAIController : public AController, public IAIPerceptionListenerInterface, public IGameplayTaskOwnerInterface
+class AIMODULE_API AAIController : public AController, public IAIPerceptionListenerInterface, public IGameplayTaskOwnerInterface, public IGenericTeamAgentInterface
 {
 	GENERATED_BODY()
 
@@ -345,6 +346,15 @@ public:
 	//----------------------------------------------------------------------//
 	virtual bool IsFollowingAPath() const override;
 
+	//----------------------------------------------------------------------//
+	// IGenericTeamAgentInterface
+	//----------------------------------------------------------------------//
+private:
+	FGenericTeamId TeamID;
+public:
+	virtual void SetGenericTeamId(const FGenericTeamId& NewTeamID) override;
+	virtual FGenericTeamId GetGenericTeamId() const override { return TeamID; }
+
 	//----------------------------------------------------------------------//
 	// IGameplayTaskOwnerInterface
 	//----------------------------------------------------------------------//
diff --git a/Engine/Source/Runtime/AIModule/Private/AIController.cpp b/Engine/Source/Runtime/AIModule/Private/AIController.cpp
--- a/Engine/Source/Runtime/AIModule/Private/AIController.cpp
+++ b/Engine/Source/Runtime/AIModule/Private/AIController.cpp
@@ -45,6 +45,7 @@ AAIController::AAIController(const FObjectInitializer& ObjectInitializer)
 
 	bSkipExtraLOSChecks = true;
 	bWantsPlayerState = false;
+	TeamID = FGenericTeamId::NoTeam;
 }
 
 void AAIController::Tick(float DeltaTime)
@@ -932,4 +933,16 @@ UAIPerceptionComponent* AAIController::GetAIPerceptionComponent()
 const UAIPerceptionComponent* AAIController::GetAIPerceptionComponent() const 
 {
 	return PerceptionComponent;
+}
+
+//----------------------------------------------------------------------//
+// IGenericTeamAgentInterface
+//----------------------------------------------------------------------//
+void AAIController::SetGenericTeamId(const FGenericTeamId& NewTeamID)
+{
+	if (TeamID != NewTeamID)
+	{
+		TeamID = NewTeamID;
+		// @todo notify perception system that a controller changed team ID
+	}
 }
\ No newline at end of file
