bookkeeping merge

diff --git a/Engine/Source/Editor/FoliageEdit/Private/SFoliagePalette.h b/Engine/Source/Editor/FoliageEdit/Private/SFoliagePalette.h
--- a/Engine/Source/Editor/FoliageEdit/Private/SFoliagePalette.h
+++ b/Engine/Source/Editor/FoliageEdit/Private/SFoliagePalette.h
@@ -3,6 +3,7 @@
 #pragma once
 
 #include "Misc/TextFilter.h"
+#include "FoliageEdMode.h"
 
 class SFoliagePaletteTile;
 class FEdModeFoliage;
diff --git a/Engine/Source/Runtime/AIModule/Classes/Perception/AIPerceptionSystem.h b/Engine/Source/Runtime/AIModule/Classes/Perception/AIPerceptionSystem.h
--- a/Engine/Source/Runtime/AIModule/Classes/Perception/AIPerceptionSystem.h
+++ b/Engine/Source/Runtime/AIModule/Classes/Perception/AIPerceptionSystem.h
@@ -4,6 +4,7 @@
 
 #include "Tickable.h"
 #include "AIPerceptionTypes.h"
+#include "AISense.h"
 #include "AIPerceptionSystem.generated.h"
 
 DECLARE_LOG_CATEGORY_EXTERN(LogAIPerception, Warning, All);
