Fixed AI perception crashing during editor-time changes to UAIPerceptionStimuliSourceComponent::RegisterAsSourceForSenses #UE4

It turns out UAIPerceptionStimuliSourceComponent::RegisterWithPerceptionSystem is being called just after a user adds an entry to UAIPerceptionStimuliSourceComponent::RegisterAsSourceForSenses, but before the value of the new entry gets set. The code assumes null values were illegal while it turned out to be a common case.

diff --git a/Engine/Source/Runtime/AIModule/Private/Perception/AIPerceptionStimuliSourceComponent.cpp b/Engine/Source/Runtime/AIModule/Private/Perception/AIPerceptionStimuliSourceComponent.cpp
--- a/Engine/Source/Runtime/AIModule/Private/Perception/AIPerceptionStimuliSourceComponent.cpp
+++ b/Engine/Source/Runtime/AIModule/Private/Perception/AIPerceptionStimuliSourceComponent.cpp
@@ -49,9 +49,12 @@ void UAIPerceptionStimuliSourceComponent::RegisterWithPerceptionSystem()
 		{
 			for (auto& SenseClass : RegisterAsSourceForSenses)
 			{
-				check(SenseClass);
-				PerceptionSystem->RegisterSourceForSenseClass(SenseClass, *OwnerActor);
-				bSuccessfullyRegistered = true;
+				if(SenseClass)
+				{
+					PerceptionSystem->RegisterSourceForSenseClass(SenseClass, *OwnerActor);
+					bSuccessfullyRegistered = true;
+				}
+				// we just ignore the empty entries
 			}
 		}
 	}
