#ue4 - Add MaxRange and a Tag (FName) to noise parameters.

- MaxRange is only used if > 0. It is modified by the loudness. Sensors must still pass their own range checks as well.

UE-13844

diff --git a/Engine/Source/Runtime/AIModule/Classes/Perception/AIPerceptionSystem.h b/Engine/Source/Runtime/AIModule/Classes/Perception/AIPerceptionSystem.h
--- a/Engine/Source/Runtime/AIModule/Classes/Perception/AIPerceptionSystem.h
+++ b/Engine/Source/Runtime/AIModule/Classes/Perception/AIPerceptionSystem.h
@@ -138,7 +138,7 @@ public:
 
 	static UAIPerceptionSystem* GetCurrent(UObject* WorldContextObject);
 
-	static void MakeNoiseImpl(AActor* NoiseMaker, float Loudness, APawn* NoiseInstigator, const FVector& NoiseLocation);
+	static void MakeNoiseImpl(AActor* NoiseMaker, float Loudness, APawn* NoiseInstigator, const FVector& NoiseLocation, float MaxRange, FName Tag);
 
 	UFUNCTION(BlueprintCallable, Category = "AI|Perception", meta = (HidePin = "WorldContext", DefaultToSelf = "WorldContext"))
 	static bool RegisterPerceptionStimuliSource(UObject* WorldContext, TSubclassOf<UAISense> Sense, AActor* Target);
diff --git a/Engine/Source/Runtime/AIModule/Classes/Perception/AIPerceptionTypes.h b/Engine/Source/Runtime/AIModule/Classes/Perception/AIPerceptionTypes.h
--- a/Engine/Source/Runtime/AIModule/Classes/Perception/AIPerceptionTypes.h
+++ b/Engine/Source/Runtime/AIModule/Classes/Perception/AIPerceptionTypes.h
@@ -155,6 +155,8 @@ public:
 	FVector StimulusLocation;
 	UPROPERTY(BlueprintReadWrite, Category = "AI|Perception")
 	FVector ReceiverLocation;
+	UPROPERTY(BlueprintReadWrite, Category = "AI|Perception")
+	FName Tag;
 
 	FAISenseID Type;
 
@@ -171,7 +173,7 @@ protected:
 public:
 	
 	/** this is the recommended constructor. Use others if you know what you're doing. */
-	FAIStimulus(const UAISense& Sense, float StimulusStrength, const FVector& InStimulusLocation, const FVector& InReceiverLocation, FResult Result = SensingSucceeded);
+	FAIStimulus(const UAISense& Sense, float StimulusStrength, const FVector& InStimulusLocation, const FVector& InReceiverLocation, FResult Result = SensingSucceeded, FName InStimulusTag = NAME_None);
 
 	/*FAIStimulus(FAISenseID SenseType, float StimulusStrength, const FVector& InStimulusLocation, const FVector& InReceiverLocation, FResult Result = SensingSucceeded)
 		: Age(0.f), ExpirationAge(NeverHappenedAge)
@@ -184,7 +186,7 @@ public:
 	// default constructor
 	FAIStimulus()
 		: Age(NeverHappenedAge), ExpirationAge(NeverHappenedAge), Strength(-1.f), StimulusLocation(FAISystem::InvalidLocation)
-		, ReceiverLocation(FAISystem::InvalidLocation), Type(FAISenseID::InvalidID()), bWantsToNotifyOnlyOnValueChange(false)
+		, ReceiverLocation(FAISystem::InvalidLocation), Tag(NAME_None), Type(FAISenseID::InvalidID()), bWantsToNotifyOnlyOnValueChange(false)
 		, bSuccessfullySensed(false), bExpired(false)
 	{}
 
diff --git a/Engine/Source/Runtime/AIModule/Classes/Perception/AISense_Hearing.h b/Engine/Source/Runtime/AIModule/Classes/Perception/AISense_Hearing.h
--- a/Engine/Source/Runtime/AIModule/Classes/Perception/AISense_Hearing.h
+++ b/Engine/Source/Runtime/AIModule/Classes/Perception/AISense_Hearing.h
@@ -21,16 +21,37 @@ struct AIMODULE_API FAINoiseEvent
 	UPROPERTY(EditAnywhere, BlueprintReadWrite, Category = "Sense")
 	FVector NoiseLocation;
 
+	/**
+	 * Loudness modifier of the sound.
+	 * If MaxRange is non-zero, this modifies the range (by multiplication).
+	 * If there is no MaxRange, then if Square(DistanceToSound) <= Square(HearingRange) * Loudness, the sound is heard, false otherwise.
+	 */
 	UPROPERTY(EditAnywhere, BlueprintReadWrite, Category = "Sense", meta = (UIMin = 0, ClampMin = 0))
 	float Loudness;
+
+	/**
+	 * Max range at which the sound can be heard. Multiplied by Loudness.
+	 * A value of 0 indicates that there is no range limit, though listeners are still limited by their own hearing range.
+	 */
+	UPROPERTY(EditAnywhere, BlueprintReadWrite, Category = "Sense", meta = (UIMin = 0, ClampMin = 0))
+	float MaxRange;
 	
+	/**
+	 * Actor triggering the sound.
+	 */
 	UPROPERTY(EditAnywhere, BlueprintReadWrite, Category = "Sense")
 	AActor* Instigator;
 
+	/**
+	 * Named identifier for the noise.
+	 */
+	UPROPERTY(EditAnywhere, BlueprintReadWrite, Category = "Sense")
+	FName Tag;
+
 	FGenericTeamId TeamIdentifier;
 		
 	FAINoiseEvent();
-	FAINoiseEvent(AActor* InInstigator, const FVector& InNoiseLocation, float InLoudness = 1.f);
+	FAINoiseEvent(AActor* InInstigator, const FVector& InNoiseLocation, float InLoudness = 1.f, float InMaxRange = 0.f, FName Tag = NAME_None);
 
 	/** Verifies and calculates derived data */
 	void Compile();
@@ -67,8 +88,17 @@ public:
 	// part of BP interface. Translates PerceptionEvent to FAINoiseEvent and call RegisterEvent(const FAINoiseEvent& Event)
 	virtual void RegisterWrappedEvent(UAISenseEvent& PerceptionEvent) override;
 
+	/**
+	 * Report a noise event.
+	 * 
+	 * @param NoiseLocation Location of the noise.
+	 * @param Loudness Loudness of the noise. If MaxRange is non-zero, modifies MaxRange, otherwise modifies the squared distance of the sensor's range.
+	 * @param Instigator Actor that triggered the noise.
+	 * @param MaxRange Max range at which the sound can be heard, multiplied by Loudness. Values <= 0 mean no limit (still limited by listener's range however).
+	 * @param Tag Identifier for the event.
+	 */
 	UFUNCTION(BlueprintCallable, Category = "AI|Perception", meta = (HidePin = "WorldContext", DefaultToSelf = "WorldContext"))
-	static void ReportNoiseEvent(UObject* WorldContext, FVector NoiseLocation, float Loudness = 1.f, AActor* Instigator = nullptr);
+	static void ReportNoiseEvent(UObject* WorldContext, FVector NoiseLocation, float Loudness = 1.f, AActor* Instigator = nullptr, float MaxRange = 0.f, FName Tag = NAME_None);
 
 protected:
 	virtual float Update() override;
diff --git a/Engine/Source/Runtime/AIModule/Private/Perception/AIPerceptionSystem.cpp b/Engine/Source/Runtime/AIModule/Private/Perception/AIPerceptionSystem.cpp
--- a/Engine/Source/Runtime/AIModule/Private/Perception/AIPerceptionSystem.cpp
+++ b/Engine/Source/Runtime/AIModule/Private/Perception/AIPerceptionSystem.cpp
@@ -453,7 +453,7 @@ bool UAIPerceptionSystem::DeliverDelayedStimuli(UAIPerceptionSystem::EDelayedSti
 	return Index > 0;
 }
 
-void UAIPerceptionSystem::MakeNoiseImpl(AActor* NoiseMaker, float Loudness, APawn* NoiseInstigator, const FVector& NoiseLocation)
+void UAIPerceptionSystem::MakeNoiseImpl(AActor* NoiseMaker, float Loudness, APawn* NoiseInstigator, const FVector& NoiseLocation, float MaxRange, FName Tag)
 {
 	UE_CLOG(NoiseMaker == nullptr && NoiseInstigator == nullptr, LogAIPerception, Warning
 		, TEXT("UAIPerceptionSystem::MakeNoiseImpl called with both NoiseMaker and NoiseInstigator being null. Unable to resolve UWorld context!"));
@@ -464,7 +464,9 @@ void UAIPerceptionSystem::MakeNoiseImpl(AActor* NoiseMaker, float Loudness, APaw
 	{
 		UAIPerceptionSystem::OnEvent(World, FAINoiseEvent(NoiseInstigator ? NoiseInstigator : NoiseMaker
 			, NoiseLocation
-			, Loudness));
+			, Loudness
+			, MaxRange
+			, Tag));
 	}
 }
 
diff --git a/Engine/Source/Runtime/AIModule/Private/Perception/AIPerceptionTypes.cpp b/Engine/Source/Runtime/AIModule/Private/Perception/AIPerceptionTypes.cpp
--- a/Engine/Source/Runtime/AIModule/Private/Perception/AIPerceptionTypes.cpp
+++ b/Engine/Source/Runtime/AIModule/Private/Perception/AIPerceptionTypes.cpp
@@ -10,10 +10,11 @@
 // @todo maybe should be a function (IsValidAge)
 const float FAIStimulus::NeverHappenedAge = FLT_MAX;
 
-FAIStimulus::FAIStimulus(const UAISense& Sense, float StimulusStrength, const FVector& InStimulusLocation, const FVector& InReceiverLocation, FResult Result)
+FAIStimulus::FAIStimulus(const UAISense& Sense, float StimulusStrength, const FVector& InStimulusLocation, const FVector& InReceiverLocation, FResult Result, FName InStimulusTag)
 	: Age(0.f), Strength(Result == SensingSucceeded ? StimulusStrength : -1.f)
 	, StimulusLocation(InStimulusLocation)
 	, ReceiverLocation(InReceiverLocation)
+	, Tag(InStimulusTag)
 	, bWantsToNotifyOnlyOnValueChange(Sense.WantsUpdateOnlyOnPerceptionValueChange())
 	, bSuccessfullySensed(Result == SensingSucceeded)
 	, bExpired(false)
diff --git a/Engine/Source/Runtime/AIModule/Private/Perception/AISense_Hearing.cpp b/Engine/Source/Runtime/AIModule/Private/Perception/AISense_Hearing.cpp
--- a/Engine/Source/Runtime/AIModule/Private/Perception/AISense_Hearing.cpp
+++ b/Engine/Source/Runtime/AIModule/Private/Perception/AISense_Hearing.cpp
@@ -9,14 +9,14 @@
 // FAINoiseEvent
 //----------------------------------------------------------------------//
 FAINoiseEvent::FAINoiseEvent()
-	: Age(0.f), NoiseLocation(FAISystem::InvalidLocation), Loudness(1.f)
-	, Instigator(nullptr), TeamIdentifier(FGenericTeamId::NoTeam)
+	: Age(0.f), NoiseLocation(FAISystem::InvalidLocation), Loudness(1.f), MaxRange(0.f)
+	, Instigator(nullptr), Tag(NAME_None), TeamIdentifier(FGenericTeamId::NoTeam)
 {
 }
 
-FAINoiseEvent::FAINoiseEvent(AActor* InInstigator, const FVector& InNoiseLocation, float InLoudness)
-	: Age(0.f), NoiseLocation(InNoiseLocation), Loudness(InLoudness)
-	, Instigator(InInstigator), TeamIdentifier(FGenericTeamId::NoTeam)
+FAINoiseEvent::FAINoiseEvent(AActor* InInstigator, const FVector& InNoiseLocation, float InLoudness, float InMaxRange, FName InTag)
+	: Age(0.f), NoiseLocation(InNoiseLocation), Loudness(InLoudness), MaxRange(InMaxRange)
+	, Instigator(InInstigator), Tag(InTag), TeamIdentifier(FGenericTeamId::NoTeam)
 {
 	Compile();
 }
@@ -71,12 +71,12 @@ UAISense_Hearing::UAISense_Hearing(const FObjectInitializer& ObjectInitializer)
 	}
 }
 
-void UAISense_Hearing::ReportNoiseEvent(UObject* WorldContext, FVector NoiseLocation, float Loudness, AActor* Instigator)
+void UAISense_Hearing::ReportNoiseEvent(UObject* WorldContext, FVector NoiseLocation, float Loudness, AActor* Instigator, float MaxRange, FName Tag)
 {
 	UAIPerceptionSystem* PerceptionSystem = UAIPerceptionSystem::GetCurrent(WorldContext);
 	if (PerceptionSystem)
 	{
-		FAINoiseEvent Event(Instigator, NoiseLocation, Loudness);
+		FAINoiseEvent Event(Instigator, NoiseLocation, Loudness, MaxRange, Tag);
 		PerceptionSystem->OnEvent(Event);
 	}
 }
@@ -133,9 +133,28 @@ float UAISense_Hearing::Update()
 		for (int32 EventIndex = 0; EventIndex < NoiseEvents.Num(); ++EventIndex)
 		{
 			const FAINoiseEvent& Event = NoiseEvents[EventIndex];
-		
-			if (FVector::DistSquared(Event.NoiseLocation, Listener.CachedLocation) > PropDigest.HearingRangeSq * Event.Loudness
-				|| FAISenseAffiliationFilter::ShouldSenseTeam(Listener.TeamIdentifier, Event.TeamIdentifier, PropDigest.AffiliationFlags) == false)
+			const float ClampedLoudness = FMath::Max(0.f, Event.Loudness);
+			const float DistToSoundSquared = FVector::DistSquared(Event.NoiseLocation, Listener.CachedLocation);
+			
+			if (Event.MaxRange > 0.f)
+			{
+				// Limit by range
+				if (DistToSoundSquared > FMath::Square(Event.MaxRange * ClampedLoudness) ||
+					DistToSoundSquared > PropDigest.HearingRangeSq)
+				{
+					continue;
+				}
+			}
+			else
+			{
+				// Limit by loudness modified squared range (this is the old behavior)
+				if (DistToSoundSquared > PropDigest.HearingRangeSq * ClampedLoudness)
+				{
+					continue;
+				}
+			}
+							
+			if (FAISenseAffiliationFilter::ShouldSenseTeam(Listener.TeamIdentifier, Event.TeamIdentifier, PropDigest.AffiliationFlags) == false)
 			{
 				continue;
 			}
@@ -143,7 +162,7 @@ float UAISense_Hearing::Update()
 			const float Delay = SpeedOfSoundSq > 0.f ? FVector::DistSquared(Event.NoiseLocation, Listener.CachedLocation) / SpeedOfSoundSq : 0;
 			// pass over to listener to process 			
 			PerseptionSys->RegisterDelayedStimulus(Listener.GetListenerID(), Delay, Event.Instigator
-				, FAIStimulus(*this, Event.Loudness, Event.NoiseLocation, Listener.CachedLocation) );
+				, FAIStimulus(*this, ClampedLoudness, Event.NoiseLocation, Listener.CachedLocation, FAIStimulus::SensingSucceeded, Event.Tag) );
 		}
 	}
 
diff --git a/Engine/Source/Runtime/Engine/Classes/GameFramework/Actor.h b/Engine/Source/Runtime/Engine/Classes/GameFramework/Actor.h
--- a/Engine/Source/Runtime/Engine/Classes/GameFramework/Actor.h
+++ b/Engine/Source/Runtime/Engine/Classes/GameFramework/Actor.h
@@ -38,7 +38,7 @@ DECLARE_DYNAMIC_MULTICAST_DELEGATE_OneParam( FActorEndTouchOverSignature, ETouch
 DECLARE_DYNAMIC_MULTICAST_DELEGATE(FActorDestroyedSignature);
 DECLARE_DYNAMIC_MULTICAST_DELEGATE_OneParam(FActorEndPlaySignature, EEndPlayReason::Type, EndPlayReason);
 
-DECLARE_DELEGATE_FourParams(FMakeNoiseDelegate, AActor*, float, class APawn*, const FVector&);
+DECLARE_DELEGATE_SixParams(FMakeNoiseDelegate, AActor*, float /*Loudness*/, class APawn*, const FVector&, float /*MaxRange*/, FName /*Tag*/);
 
 #if !UE_BUILD_SHIPPING
 DECLARE_DELEGATE_RetVal_ThreeParams(bool, FOnProcessEvent, AActor*, UFunction*, void*);
@@ -1013,12 +1013,14 @@ public:
 	 * Note that the NoiseInstigator Pawn MUST have a PawnNoiseEmitterComponent for the noise to be detected by a PawnSensingComponent.
 	 * Senders of MakeNoise should have an Instigator if they are not pawns, or pass a NoiseInstigator.
 	 *
-	 * @param Loudness - is the relative loudness of this noise (range 0.0 to 1.0).  Directly affects the hearing range specified by the SensingComponent's HearingThreshold.
-	 * @param NoiseInstigator - Pawn responsible for this noise.  Uses the actor's Instigator if NoiseInstigator=NULL
-	 * @param NoiseLocation - Position of noise source.  If zero vector, use the actor's location.
-	*/
+	 * @param Loudness The relative loudness of this noise. Usual range is 0 (no noise) to 1 (full volume). If MaxRange is used, this scales the max range, otherwise it affects the hearing range specified by the sensor.
+	 * @param NoiseInstigator Pawn responsible for this noise.  Uses the actor's Instigator if NoiseInstigator=NULL
+	 * @param NoiseLocation Position of noise source.  If zero vector, use the actor's location.
+	 * @param MaxRange Max range at which the sound may be heard. A value of 0 indicates no max range (though perception may have its own range). Loudness scales the range. (Note: not supported for legacy PawnSensingComponent, only for AIPerception)
+	 * @param Tag Identifier for the noise.
+	 */
 	UFUNCTION(BlueprintCallable, BlueprintAuthorityOnly, Category="AI", meta=(BlueprintProtected = "true"))
-	void MakeNoise(float Loudness=1.f, APawn* NoiseInstigator=NULL, FVector NoiseLocation=FVector::ZeroVector);
+	void MakeNoise(float Loudness=1.f, APawn* NoiseInstigator=NULL, FVector NoiseLocation=FVector::ZeroVector, float MaxRange = 0.f, FName Tag = NAME_None);
 
 	//=============================================================================
 	// Blueprint
@@ -2435,7 +2437,7 @@ public:
 	static FOnProcessEvent ProcessEventDelegate;
 #endif
 
-	static void MakeNoiseImpl(AActor* NoiseMaker, float Loudness, APawn* NoiseInstigator, const FVector& NoiseLocation);
+	static void MakeNoiseImpl(AActor* NoiseMaker, float Loudness, APawn* NoiseInstigator, const FVector& NoiseLocation, float MaxRange, FName Tag);
 	static void SetMakeNoiseDelegate(const FMakeNoiseDelegate& NewDelegate);
 
 	/** A fence to track when the primitive is detached from the scene in the rendering thread. */
diff --git a/Engine/Source/Runtime/Engine/Private/Actor.cpp b/Engine/Source/Runtime/Engine/Private/Actor.cpp
--- a/Engine/Source/Runtime/Engine/Private/Actor.cpp
+++ b/Engine/Source/Runtime/Engine/Private/Actor.cpp
@@ -1716,23 +1716,26 @@ void AActor::FellOutOfWorld(const UDamageType& dmgType)
 	Destroy();
 }
 
-void AActor::MakeNoise(float Loudness, APawn* NoiseInstigator, FVector NoiseLocation)
+void AActor::MakeNoise(float Loudness, APawn* NoiseInstigator, FVector NoiseLocation, float MaxRange, FName Tag)
 {
 	NoiseInstigator = NoiseInstigator ? NoiseInstigator : Instigator;
 	if ((GetNetMode() != NM_Client) && NoiseInstigator)
 	{
 		AActor::MakeNoiseDelegate.Execute(this, Loudness, NoiseInstigator
-			, NoiseLocation.IsZero() ? GetActorLocation() : NoiseLocation);
+			, NoiseLocation.IsZero() ? GetActorLocation() : NoiseLocation
+			, MaxRange
+			, Tag);
 	}
 }
 
-void AActor::MakeNoiseImpl(AActor* NoiseMaker, float Loudness, APawn* NoiseInstigator, const FVector& NoiseLocation)
+void AActor::MakeNoiseImpl(AActor* NoiseMaker, float Loudness, APawn* NoiseInstigator, const FVector& NoiseLocation, float MaxRange, FName Tag)
 {
 	check(NoiseMaker);
 
 	UPawnNoiseEmitterComponent* NoiseEmitterComponent = NoiseInstigator->GetPawnNoiseEmitterComponent();
 	if (NoiseEmitterComponent)
 	{
+		// Note: MaxRange and Tag are not supported for this legacy component. Use AISense_Hearing instead.
 		NoiseEmitterComponent->MakeNoise( NoiseMaker, Loudness, NoiseLocation );
 	}
 }
