replaced AI message's sender with weak object pointer
#ue4 UE-17281

diff --git a/Engine/Source/Runtime/AIModule/Classes/BrainComponent.h b/Engine/Source/Runtime/AIModule/Classes/BrainComponent.h
--- a/Engine/Source/Runtime/AIModule/Classes/BrainComponent.h
+++ b/Engine/Source/Runtime/AIModule/Classes/BrainComponent.h
@@ -29,7 +29,7 @@ struct AIMODULE_API FAIMessage
 	FName MessageName;
 
 	/** message source */
-	UObject* Sender;
+	FWeakObjectPtr Sender;
 
 	/** message param: ID */
 	FAIRequestID RequestID;
