Made BTService_DefaultFocus internals protected rather than private (which was a bug) #UE4

diff --git a/Engine/Source/Runtime/AIModule/Classes/BehaviorTree/Services/BTService_DefaultFocus.h b/Engine/Source/Runtime/AIModule/Classes/BehaviorTree/Services/BTService_DefaultFocus.h
--- a/Engine/Source/Runtime/AIModule/Classes/BehaviorTree/Services/BTService_DefaultFocus.h
+++ b/Engine/Source/Runtime/AIModule/Classes/BehaviorTree/Services/BTService_DefaultFocus.h
@@ -26,7 +26,8 @@ UCLASS(hidecategories=(Service))
 class AIMODULE_API UBTService_DefaultFocus : public UBTService_BlackboardBase
 {
 	GENERATED_BODY()
-	
+
+protected:
 	// not exposed to users on purpose. Here to make reusing focus-setting mechanics by derived classes possible
 	UPROPERTY()
 	uint8 FocusPriority;
