Fixed the default AI stimulus MaxAge being interpreted as "forget instantly" rather than "never forget" #UE4

github PR#1312 by witkowski-seb

diff --git a/Engine/Source/Runtime/AIModule/Classes/Perception/AISenseConfig.h b/Engine/Source/Runtime/AIModule/Classes/Perception/AISenseConfig.h
--- a/Engine/Source/Runtime/AIModule/Classes/Perception/AISenseConfig.h
+++ b/Engine/Source/Runtime/AIModule/Classes/Perception/AISenseConfig.h
@@ -30,7 +30,7 @@ public:
 
 	FAISenseID GetSenseID() const;
 	
-	float GetMaxAge() const { return MaxAge; }
+	float GetMaxAge() const { return MaxAge == 0.f ? FAIStimulus::NeverHappenedAge : MaxAge; }
 	bool IsEnabled() const { return bStartsEnabled; }
 
 #if !UE_BUILD_SHIPPING
