Exposing AIController::PathFollowingComponent and AIController::PerceptionComponent to class' default properties #UE4

github PR#1329 by slonopotamus

diff --git a/Engine/Source/Runtime/AIModule/Classes/AIController.h b/Engine/Source/Runtime/AIModule/Classes/AIController.h
--- a/Engine/Source/Runtime/AIModule/Classes/AIController.h
+++ b/Engine/Source/Runtime/AIModule/Classes/AIController.h
@@ -105,7 +105,7 @@ private_subobject:
 
 	/** Component used for moving along a path. */
 	DEPRECATED_FORGAME(4.6, "PathFollowingComponent should not be accessed directly, please use GetPathFollowingComponent() function instead. PathFollowingComponent will soon be private and your code will not compile.")
-	UPROPERTY()
+	UPROPERTY(VisibleDefaultsOnly, Category = AI)
 	UPathFollowingComponent* PathFollowingComponent;
 
 public:
@@ -115,7 +115,7 @@ public:
 	UBrainComponent* BrainComponent;
 
 	DEPRECATED_FORGAME(4.8, "PerceptionComponent should not be accessed directly, please use GetAIPerceptionComponent() function instead. PerceptionComponent will soon be private and your code will not compile.")
-	UPROPERTY()
+	UPROPERTY(VisibleDefaultsOnly, Category = AI)
 	UAIPerceptionComponent* PerceptionComponent;
 
 private_subobject:
