Moved BT's support for gameplay tasks up from BT task node to generic BT node #UE4

diff --git a/Engine/Source/Runtime/AIModule/Classes/BehaviorTree/BTNode.h b/Engine/Source/Runtime/AIModule/Classes/BehaviorTree/BTNode.h
--- a/Engine/Source/Runtime/AIModule/Classes/BehaviorTree/BTNode.h
+++ b/Engine/Source/Runtime/AIModule/Classes/BehaviorTree/BTNode.h
@@ -4,6 +4,7 @@
 
 #include "BehaviorTreeComponent.h"
 #include "BehaviorTreeTypes.h"
+#include "GameplayTaskOwnerInterface.h"
 #include "BTNode.generated.h"
 
 class AAIController;
@@ -20,7 +21,7 @@ struct FBTInstancedNodeMemory
 };
 
 UCLASS(Abstract,config=Game)
-class AIMODULE_API UBTNode : public UObject
+class AIMODULE_API UBTNode : public UObject, public IGameplayTaskOwnerInterface
 {
 	GENERATED_UCLASS_BODY()
 
@@ -140,6 +141,31 @@ class AIMODULE_API UBTNode : public UObject
 	/** Gets called only for instanced nodes(bCreateNodeInstance == true). In practive overriden by BP-implemented BT nodes */
 	virtual void SetOwner(AActor* ActorOwner) {}
 
+	//----------------------------------------------------------------------//
+	// IGameplayTaskOwnerInterface & other gameplay task stuff
+	//----------------------------------------------------------------------//
+	virtual UGameplayTasksComponent* GetGameplayTasksComponent(const UGameplayTask& Task) const override;
+	virtual void OnTaskActivated(UGameplayTask& Task) override;
+	virtual void OnTaskDeactivated(UGameplayTask& Task) override;
+	virtual void OnTaskInitialized(UGameplayTask& Task) override;
+	virtual AActor* GetOwnerActor(const UGameplayTask* Task) const;
+	virtual AActor* GetAvatarActor(const UGameplayTask* Task) const;
+	virtual uint8 GetDefaultPriority() const;
+
+	UBehaviorTreeComponent* GetBTComponentForTask(UGameplayTask& Task) const;
+	
+	template <class T>
+	T* NewBTAITask(UBehaviorTreeComponent& BTComponent)
+	{
+		T* NewAITask = NewObject<T>();
+		AAIController* AIController = BTComponent.GetAIOwner();
+		check(AIController && "Can\'t spawn an AI task without AI controller!");
+		NewAITask->InitAITask(*AIController, *this, GetDefaultPriority());
+		bOwnsGameplayTasks = true;
+
+		return NewAITask;
+	}
+
 	/** node name */
 	UPROPERTY(Category=Description, EditAnywhere)
 	FString NodeName;
@@ -179,6 +205,10 @@ protected:
 	/** if set, node is injected by subtree */
 	uint8 bIsInjected : 1;
 
+	/** set to true if task owns any GameplayTasks. Note this requires tasks to be created via NewBTAITask
+	*	Otherwise specific BT task node class is responsible for ending the gameplay tasks on node finish */
+	uint8 bOwnsGameplayTasks : 1;
+
 	//----------------------------------------------------------------------//
 	// DEPRECATED
 	//----------------------------------------------------------------------//
diff --git a/Engine/Source/Runtime/AIModule/Classes/BehaviorTree/BTTaskNode.h b/Engine/Source/Runtime/AIModule/Classes/BehaviorTree/BTTaskNode.h
--- a/Engine/Source/Runtime/AIModule/Classes/BehaviorTree/BTTaskNode.h
+++ b/Engine/Source/Runtime/AIModule/Classes/BehaviorTree/BTTaskNode.h
@@ -3,7 +3,6 @@
 #pragma once
 
 #include "BehaviorTree/BTNode.h"
-#include "GameplayTaskOwnerInterface.h"
 #include "BTTaskNode.generated.h"
 
 class UBehaviorTreeComponent;
@@ -25,7 +24,7 @@ struct FAIMessage;
  */
 
 UCLASS(Abstract)
-class AIMODULE_API UBTTaskNode : public UBTNode, public IGameplayTaskOwnerInterface
+class AIMODULE_API UBTTaskNode : public UBTNode
 {
 	GENERATED_UCLASS_BODY()
 
@@ -66,6 +65,11 @@ public:
 	/** helper function: finishes latent aborting */
 	void FinishLatentAbort(UBehaviorTreeComponent& OwnerComp) const;
 
+	//----------------------------------------------------------------------//
+	// UBTTaskNode IGameplayTaskOwnerInterface
+	//----------------------------------------------------------------------//
+	virtual void OnTaskDeactivated(UGameplayTask& Task) override;
+
 protected:
 
 	/** if set, TickTask will be called */
@@ -73,11 +77,7 @@ protected:
 
 	/** if set, OnTaskFinished will be called */
 	uint8 bNotifyTaskFinished : 1;
-
-	/** set to true if task owns any GameplayTasks. Note this requires tasks to be created via NewBTAITask
-	 *	Otherwise specific BT task node class is responsible for ending the gameplay tasks on node finish */
-	uint8 bOwnsGameplayTasks : 1;
-
+	
 	/** ticks this task 
 	 * this function should be considered as const (don't modify state of object) if node is not instanced! */
 	virtual void TickTask(UBehaviorTreeComponent& OwnerComp, uint8* NodeMemory, float DeltaSeconds);
@@ -96,32 +96,7 @@ protected:
 	
 	/** unregister message observers */
 	void StopWaitingForMessages(UBehaviorTreeComponent& OwnerComp) const;
-
-	template <class T>
-	T* NewBTAITask(UBehaviorTreeComponent& BTComponent)
-	{
-		T* NewAITask = NewObject<T>();
-		AAIController* AIController = BTComponent.GetAIOwner();
-		check(AIController && "Can\'t spawn an AI task without AI controller!");
-		NewAITask->InitAITask(*AIController, *this, GetDefaultPriority());
-		bOwnsGameplayTasks = true;
-		
-		return NewAITask;
-	}
-
-	//----------------------------------------------------------------------//
-	// IGameplayTaskOwnerInterface
-	//----------------------------------------------------------------------//
-	virtual UGameplayTasksComponent* GetGameplayTasksComponent(const UGameplayTask& Task) const override;
-	virtual void OnTaskActivated(UGameplayTask& Task) override;
-	virtual void OnTaskDeactivated(UGameplayTask& Task) override;
-	virtual void OnTaskInitialized(UGameplayTask& Task) override;
-	virtual AActor* GetOwnerActor(const UGameplayTask* Task) const;
-	virtual AActor* GetAvatarActor(const UGameplayTask* Task) const;
-	virtual uint8 GetDefaultPriority() const;
-
-	UBehaviorTreeComponent* GetBTComponentForTask(UGameplayTask& Task) const;
-
+	
 	//----------------------------------------------------------------------//
 	// DEPRECATED
 	//----------------------------------------------------------------------//
diff --git a/Engine/Source/Runtime/AIModule/Private/BehaviorTree/BTDecorator.cpp b/Engine/Source/Runtime/AIModule/Private/BehaviorTree/BTDecorator.cpp
--- a/Engine/Source/Runtime/AIModule/Private/BehaviorTree/BTDecorator.cpp
+++ b/Engine/Source/Runtime/AIModule/Private/BehaviorTree/BTDecorator.cpp
@@ -6,8 +6,6 @@
 
 UBTDecorator::UBTDecorator(const FObjectInitializer& ObjectInitializer) : Super(ObjectInitializer)
 {
-	NodeName = "UnknownDecorator";
-
 	FlowAbortMode = EBTFlowAbortMode::None;
 	bAllowAbortNone = true;
 	bAllowAbortLowerPri = true;
diff --git a/Engine/Source/Runtime/AIModule/Private/BehaviorTree/BTNode.cpp b/Engine/Source/Runtime/AIModule/Private/BehaviorTree/BTNode.cpp
--- a/Engine/Source/Runtime/AIModule/Private/BehaviorTree/BTNode.cpp
+++ b/Engine/Source/Runtime/AIModule/Private/BehaviorTree/BTNode.cpp
@@ -259,4 +259,111 @@ FString UBTNode::GetRuntimeDescription(const UBehaviorTreeComponent* OwnerComp,
 		return GetRuntimeDescription(*OwnerComp, NodeMemory, Verbosity);
 	}
 	return TEXT("");
-}
\ No newline at end of file
+}
+
+
+//----------------------------------------------------------------------//
+// UBTNode IGameplayTaskOwnerInterface
+//----------------------------------------------------------------------//
+UGameplayTasksComponent* UBTNode::GetGameplayTasksComponent(const UGameplayTask& Task) const
+{
+	const UAITask* AsAITask = Cast<const UAITask>(&Task);
+	if (AsAITask)
+	{
+		return AsAITask->GetAIController() ? AsAITask->GetAIController()->GetGameplayTasksComponent(Task) : nullptr;
+	}
+
+	return Task.GetGameplayTasksComponent();
+}
+
+void UBTNode::OnTaskInitialized(UGameplayTask& Task)
+{
+	// validate the task
+	UAITask* AsAITask = Cast<UAITask>(&Task);
+	if (AsAITask != nullptr && AsAITask->GetAIController() == nullptr)
+	{
+		// this means that the task has either been created without specifying 
+		// UAITAsk::OwnerController's value (like via BP's Construct Object node)
+		// or it has been created in C++ with inappropriate function
+		UE_LOG(LogBehaviorTree, Error, TEXT("Missing AIController in AITask %s"), *AsAITask->GetName());
+	}
+}
+
+UBehaviorTreeComponent* UBTNode::GetBTComponentForTask(UGameplayTask& Task) const
+{
+	UAITask* AsAITask = Cast<UAITask>(&Task);
+	return AsAITask && AsAITask->GetAIController() ? Cast<UBehaviorTreeComponent>(AsAITask->GetAIController()->BrainComponent) : nullptr;
+}
+
+void UBTNode::OnTaskActivated(UGameplayTask& Task)
+{
+	ensure(Task.GetTaskOwner() == this);
+}
+
+void UBTNode::OnTaskDeactivated(UGameplayTask& Task)
+{
+	ensure(Task.GetTaskOwner() == this);
+}
+
+AActor* UBTNode::GetOwnerActor(const UGameplayTask* Task) const
+{
+	if (Task == nullptr)
+	{
+		if (IsInstanced())
+		{
+			const UBehaviorTreeComponent* BTComponent = Cast<const UBehaviorTreeComponent>(GetOuter());
+			//not having BT component for an instanced BT node is invalid!
+			check(BTComponent);
+			return BTComponent->GetAIOwner();
+		}
+		else
+		{
+			UE_LOG(LogBehaviorTree, Warning, TEXT("%s: Unable to determine Owner Actor for a null GameplayTask"), *GetName());
+			return nullptr;
+		}
+	}
+
+	const UAITask* AsAITask = Cast<const UAITask>(Task);
+	if (AsAITask)
+	{
+		return AsAITask->GetAIController();
+	}
+
+	const UGameplayTasksComponent* GTComponent = Task->GetGameplayTasksComponent();
+
+	return GTComponent ? GTComponent->GetOwnerActor(Task) : nullptr;
+}
+
+AActor* UBTNode::GetAvatarActor(const UGameplayTask* Task) const
+{
+	if (Task == nullptr)
+	{
+		if (IsInstanced())
+		{
+			const UBehaviorTreeComponent* BTComponent = Cast<const UBehaviorTreeComponent>(GetOuter());
+			//not having BT component for an instanced BT node is invalid!
+			check(BTComponent);
+			return BTComponent->GetAIOwner() ? BTComponent->GetAIOwner()->GetPawn() : nullptr;
+		}
+		else
+		{
+			UE_LOG(LogBehaviorTree, Warning, TEXT("%s: Unable to determine Avatar Actor for a null GameplayTask"), *GetName());
+			return nullptr;
+		}
+	}
+
+	const UAITask* AsAITask = Cast<const UAITask>(Task);
+	if (AsAITask)
+	{
+		return AsAITask->GetAIController() ? AsAITask->GetAIController()->GetPawn() : nullptr;
+	}
+
+	const UGameplayTasksComponent* GTComponent = Task->GetGameplayTasksComponent();
+
+	return GTComponent ? GTComponent->GetOwnerActor(Task) : nullptr;
+}
+
+uint8 UBTNode::GetDefaultPriority() const
+{
+	return uint8(EAITaskPriority::AutonomousAI);
+}
diff --git a/Engine/Source/Runtime/AIModule/Private/BehaviorTree/BTTaskNode.cpp b/Engine/Source/Runtime/AIModule/Private/BehaviorTree/BTTaskNode.cpp
--- a/Engine/Source/Runtime/AIModule/Private/BehaviorTree/BTTaskNode.cpp
+++ b/Engine/Source/Runtime/AIModule/Private/BehaviorTree/BTTaskNode.cpp
@@ -6,7 +6,7 @@
 
 UBTTaskNode::UBTTaskNode(const FObjectInitializer& ObjectInitializer) : Super(ObjectInitializer)
 {
-	NodeName = "UnknownTask";
+	NodeName = GetClass()->GetName();
 	bNotifyTick = false;
 	bNotifyTaskFinished = false;
 }
@@ -165,41 +165,6 @@ FName UBTTaskNode::GetNodeIconName() const
 //----------------------------------------------------------------------//
 // UBTTaskNode IGameplayTaskOwnerInterface
 //----------------------------------------------------------------------//
-UGameplayTasksComponent* UBTTaskNode::GetGameplayTasksComponent(const UGameplayTask& Task) const 
-{
-	const UAITask* AsAITask = Cast<const UAITask>(&Task);
-	if (AsAITask)
-	{
-		return AsAITask->GetAIController() ? AsAITask->GetAIController()->GetGameplayTasksComponent(Task) : nullptr;
-	}
-
-	return Task.GetGameplayTasksComponent();
-}
-
-void UBTTaskNode::OnTaskInitialized(UGameplayTask& Task)
-{
-	// validate the task
-	UAITask* AsAITask = Cast<UAITask>(&Task);
-	if (AsAITask != nullptr && AsAITask->GetAIController() == nullptr)
-	{
-		// this means that the task has either been created without specifying 
-		// UAITAsk::OwnerController's value (like via BP's Construct Object node)
-		// or it has been created in C++ with inappropriate function
-		UE_LOG(LogBehaviorTree, Error, TEXT("Missing AIController in AITask %s"), *AsAITask->GetName());
-	}
-}
-
-UBehaviorTreeComponent* UBTTaskNode::GetBTComponentForTask(UGameplayTask& Task) const
-{
-	UAITask* AsAITask = Cast<UAITask>(&Task);
-	return AsAITask && AsAITask->GetAIController() ? Cast<UBehaviorTreeComponent>(AsAITask->GetAIController()->BrainComponent) : nullptr;
-}
-
-void UBTTaskNode::OnTaskActivated(UGameplayTask& Task) 
-{
-	ensure(Task.GetTaskOwner() == this);
-}
-
 void UBTTaskNode::OnTaskDeactivated(UGameplayTask& Task)
 {
 	ensure(Task.GetTaskOwner() == this);
@@ -212,69 +177,6 @@ void UBTTaskNode::OnTaskDeactivated(UGameplayTask& Task)
 	}
 }
 
-AActor* UBTTaskNode::GetOwnerActor(const UGameplayTask* Task) const
-{
-	if (Task == nullptr)
-	{
-		if (IsInstanced())
-		{
-			const UBehaviorTreeComponent* BTComponent = Cast<const UBehaviorTreeComponent>(GetOuter());
-			//not having BT component for an instanced BT node is invalid!
-			check(BTComponent);
-			return BTComponent->GetAIOwner();
-		}
-		else
-		{
-			UE_LOG(LogBehaviorTree, Warning, TEXT("%s: Unable to determine Owner Actor for a null GameplayTask"), *GetName());
-			return nullptr;
-		}
-	}
-
-	const UAITask* AsAITask = Cast<const UAITask>(Task);
-	if (AsAITask)
-	{
-		return AsAITask->GetAIController();
-	}
-
-	const UGameplayTasksComponent* GTComponent = Task->GetGameplayTasksComponent();
-
-	return GTComponent ? GTComponent->GetOwnerActor(Task) : nullptr; 
-}
-
-AActor* UBTTaskNode::GetAvatarActor(const UGameplayTask* Task) const
-{
-	if (Task == nullptr)
-	{
-		if (IsInstanced())
-		{
-			const UBehaviorTreeComponent* BTComponent = Cast<const UBehaviorTreeComponent>(GetOuter());
-			//not having BT component for an instanced BT node is invalid!
-			check(BTComponent); 
-			return BTComponent->GetAIOwner() ? BTComponent->GetAIOwner()->GetPawn() : nullptr;
-		}
-		else
-		{
-			UE_LOG(LogBehaviorTree, Warning, TEXT("%s: Unable to determine Avatar Actor for a null GameplayTask"), *GetName());
-			return nullptr;
-		}
-	}
-
-	const UAITask* AsAITask = Cast<const UAITask>(Task);
-	if (AsAITask)
-	{
-		return AsAITask->GetAIController() ? AsAITask->GetAIController()->GetPawn() : nullptr;
-	}
-
-	const UGameplayTasksComponent* GTComponent = Task->GetGameplayTasksComponent();
-
-	return GTComponent ? GTComponent->GetOwnerActor(Task) : nullptr;
-}
-
-uint8 UBTTaskNode::GetDefaultPriority() const 
-{ 
-	return uint8(EAITaskPriority::AutonomousAI); 
-}
-
 //----------------------------------------------------------------------//
 // DEPRECATED
 //----------------------------------------------------------------------//
