Non-unity build fix #UE4

diff --git a/Engine/Source/Runtime/AIModule/Private/BehaviorTree/BTNode.cpp b/Engine/Source/Runtime/AIModule/Private/BehaviorTree/BTNode.cpp
--- a/Engine/Source/Runtime/AIModule/Private/BehaviorTree/BTNode.cpp
+++ b/Engine/Source/Runtime/AIModule/Private/BehaviorTree/BTNode.cpp
@@ -4,6 +4,7 @@
 #include "BehaviorTree/BehaviorTree.h"
 #include "BehaviorTree/BTNode.h"
 #include "BehaviorTree/BTCompositeNode.h"
+#include "GameplayTasksComponent.h"
 
 //----------------------------------------------------------------------//
 // UBTNode
