added dynamic subtree injection support for behavior trees
#ue4 UE-14626

diff --git a/Engine/Source/Runtime/AIModule/Classes/BehaviorTree/BTFunctionLibrary.h b/Engine/Source/Runtime/AIModule/Classes/BehaviorTree/BTFunctionLibrary.h
--- a/Engine/Source/Runtime/AIModule/Classes/BehaviorTree/BTFunctionLibrary.h
+++ b/Engine/Source/Runtime/AIModule/Classes/BehaviorTree/BTFunctionLibrary.h
@@ -25,6 +25,9 @@ public:
 	UFUNCTION(BlueprintPure, Category="AI|BehaviorTree", Meta=(HidePin="NodeOwner", DefaultToSelf="NodeOwner"))
 	static UBlackboardComponent* GetOwnersBlackboard(UBTNode* NodeOwner);
 
+	UFUNCTION(BlueprintPure, Category = "AI|BehaviorTree", Meta=(HidePin="NodeOwner", DefaultToSelf="NodeOwner"))
+	static UBehaviorTreeComponent* GetOwnerComponent(UBTNode* NodeOwner);
+
 	UFUNCTION(BlueprintPure, Category="AI|BehaviorTree", Meta=(HidePin="NodeOwner", DefaultToSelf="NodeOwner"))
 	static UObject* GetBlackboardValueAsObject(UBTNode* NodeOwner, const FBlackboardKeySelector& Key);
 
diff --git a/Engine/Source/Runtime/AIModule/Classes/BehaviorTree/BehaviorTreeComponent.h b/Engine/Source/Runtime/AIModule/Classes/BehaviorTree/BehaviorTreeComponent.h
--- a/Engine/Source/Runtime/AIModule/Classes/BehaviorTree/BehaviorTreeComponent.h
+++ b/Engine/Source/Runtime/AIModule/Classes/BehaviorTree/BehaviorTreeComponent.h
@@ -13,6 +13,7 @@ class UBTCompositeNode;
 class UBTTaskNode;
 class UBTDecorator;
 class UBTTask_RunBehavior;
+class UBTTask_RunBehaviorDynamic;
 class FBehaviorTreeDebugger;
 class UBehaviorTree;
 class UBTAuxiliaryNode;
@@ -198,11 +199,17 @@ public:
 	virtual FString DescribeActiveTrees() const;
 
 	/** @return the cooldown tag end time, 0.0f if CooldownTag is not found */
+	UFUNCTION(BlueprintCallable, Category = "AI|Logic")
 	float GetTagCooldownEndTime(FGameplayTag CooldownTag) const;
 
 	/** add to the cooldown tag's duration */
+	UFUNCTION(BlueprintCallable, Category = "AI|Logic")
 	void AddCooldownTagDuration(FGameplayTag CooldownTag, float CooldownDuration, bool bAddToExistingDuration);
 
+	/** assign subtree to RunBehaviorDynamic task specified by tag */
+	UFUNCTION(BlueprintCallable, Category="AI|Logic")
+	virtual void SetDynamicSubtree(FGameplayTag InjectTag, UBehaviorTree* BehaviorAsset);
+
 #if ENABLE_VISUAL_LOG
 	virtual void DescribeSelfToVisLog(struct FVisualLogEntry* Snapshot) const override;
 #endif
@@ -352,6 +359,7 @@ protected:
 	friend UBTCompositeNode;
 	friend UBTTaskNode;
 	friend UBTTask_RunBehavior;
+	friend UBTTask_RunBehaviorDynamic;
 	friend FBehaviorTreeDebugger;
 	friend FBehaviorTreeInstance;
 };
diff --git a/Engine/Source/Runtime/AIModule/Classes/BehaviorTree/Tasks/BTTask_RunBehavior.h b/Engine/Source/Runtime/AIModule/Classes/BehaviorTree/Tasks/BTTask_RunBehavior.h
--- a/Engine/Source/Runtime/AIModule/Classes/BehaviorTree/Tasks/BTTask_RunBehavior.h
+++ b/Engine/Source/Runtime/AIModule/Classes/BehaviorTree/Tasks/BTTask_RunBehavior.h
@@ -13,9 +13,7 @@
  * which are injected into parent tree, and structure of running tree
  * cannot be modified in runtime (see: BTNode: ExecutionIndex, MemoryOffset)
  *
- * Dynamic subtrees can be implemented, but at the cost of some features:
- * - no root level decorators
- * - no or limited editor preview
+ * Use RunBehaviorDynamic task for subtrees that need to be changed in runtime.
  */
 
 UCLASS()
diff --git a/Engine/Source/Runtime/AIModule/Private/BehaviorTree/BTFunctionLibrary.cpp b/Engine/Source/Runtime/AIModule/Private/BehaviorTree/BTFunctionLibrary.cpp
--- a/Engine/Source/Runtime/AIModule/Private/BehaviorTree/BTFunctionLibrary.cpp
+++ b/Engine/Source/Runtime/AIModule/Private/BehaviorTree/BTFunctionLibrary.cpp
@@ -34,6 +34,18 @@ namespace
 	}
 }
 
+UBehaviorTreeComponent* UBTFunctionLibrary::GetOwnerComponent(UBTNode* NodeOwner)
+{
+	check(NodeOwner);
+	check(NodeOwner->GetClass()->HasAnyClassFlags(CLASS_CompiledFromBlueprint) && "This function call is valid only for BP-implemented BT nodes");
+	check(NodeOwner->GetOuter());
+
+	UBehaviorTreeComponent* OwnerComp = Cast<UBehaviorTreeComponent>(NodeOwner->GetOuter());
+	check(OwnerComp);
+
+	return OwnerComp;
+}
+
 UBlackboardComponent* UBTFunctionLibrary::GetOwnersBlackboard(UBTNode* NodeOwner)
 {
 	check(NodeOwner != NULL);
diff --git a/Engine/Source/Runtime/AIModule/Private/BehaviorTree/BehaviorTreeComponent.cpp b/Engine/Source/Runtime/AIModule/Private/BehaviorTree/BehaviorTreeComponent.cpp
--- a/Engine/Source/Runtime/AIModule/Private/BehaviorTree/BehaviorTreeComponent.cpp
+++ b/Engine/Source/Runtime/AIModule/Private/BehaviorTree/BehaviorTreeComponent.cpp
@@ -8,6 +8,7 @@
 #include "BehaviorTree/BehaviorTreeManager.h"
 #include "BehaviorTree/BlackboardComponent.h"
 #include "BehaviorTree/BehaviorTreeComponent.h"
+#include "BehaviorTree/Tasks/BTTask_RunBehaviorDynamic.h"
 
 #if USE_BEHAVIORTREE_DEBUGGER
 int32 UBehaviorTreeComponent::ActiveDebuggerCounter = 0;
@@ -1896,6 +1897,62 @@ void UBehaviorTreeComponent::AddCooldownTagDuration(FGameplayTag CooldownTag, fl
 	}
 }
 
+void SetDynamicSubtreeHelper(const UBTCompositeNode* TestComposite,
+	const FBehaviorTreeInstance& InstanceInfo, const UBehaviorTreeComponent* OwnerComp,
+	const FGameplayTag& InjectTag, UBehaviorTree* BehaviorAsset)
+{
+	for (int32 Idx = 0; Idx < TestComposite->Children.Num(); Idx++)
+	{
+		const FBTCompositeChild& ChildInfo = TestComposite->Children[Idx];
+		if (ChildInfo.ChildComposite)
+		{
+			SetDynamicSubtreeHelper(ChildInfo.ChildComposite, InstanceInfo, OwnerComp, InjectTag, BehaviorAsset);
+		}
+		else
+		{
+			UBTTask_RunBehaviorDynamic* SubtreeTask = Cast<UBTTask_RunBehaviorDynamic>(ChildInfo.ChildTask);
+			if (SubtreeTask && SubtreeTask->HasMatchingTag(InjectTag))
+			{
+				const uint8* NodeMemory = SubtreeTask->GetNodeMemory<uint8>(InstanceInfo);
+				UBTTask_RunBehaviorDynamic* InstancedNode = Cast<UBTTask_RunBehaviorDynamic>(SubtreeTask->GetNodeInstance(*OwnerComp, (uint8*)NodeMemory));
+				if (InstancedNode)
+				{
+					InstancedNode->SetBehaviorAsset(BehaviorAsset);
+					UE_VLOG(OwnerComp->GetOwner(), LogBehaviorTree, Log, TEXT("Replaced subtree in %s with %s (tag: %s)"),
+						*UBehaviorTreeTypes::DescribeNodeHelper(SubtreeTask), *GetNameSafe(BehaviorAsset), *InjectTag.ToString());
+				}
+			}
+		}
+	}
+}
+
+void UBehaviorTreeComponent::SetDynamicSubtree(FGameplayTag InjectTag, UBehaviorTree* BehaviorAsset)
+{
+	// replace at matching injection points
+	for (int32 InstanceIndex = 0; InstanceIndex < InstanceStack.Num(); InstanceIndex++)
+	{
+		const FBehaviorTreeInstance& InstanceInfo = InstanceStack[InstanceIndex];
+		SetDynamicSubtreeHelper(InstanceInfo.RootNode, InstanceInfo, this, InjectTag, BehaviorAsset);
+	}
+
+	// restart subtree if it was replaced
+	for (int32 InstanceIndex = 0; InstanceIndex < InstanceStack.Num(); InstanceIndex++)
+	{
+		const FBehaviorTreeInstance& InstanceInfo = InstanceStack[InstanceIndex];
+		if (InstanceInfo.ActiveNodeType == EBTActiveNode::ActiveTask)
+		{
+			const UBTTask_RunBehaviorDynamic* SubtreeTask = Cast<const UBTTask_RunBehaviorDynamic>(InstanceInfo.ActiveNode);
+			if (SubtreeTask && SubtreeTask->HasMatchingTag(InjectTag))
+			{
+				UBTCompositeNode* RestartNode = SubtreeTask->GetParentNode();
+				int32 RestartChildIdx = RestartNode->GetChildIndex(*SubtreeTask);
+
+				RequestExecution(RestartNode, InstanceIndex, SubtreeTask, RestartChildIdx, EBTNodeResult::Aborted);
+				break;
+			}
+		}
+	}
+}
 
 #if ENABLE_VISUAL_LOG
 void UBehaviorTreeComponent::DescribeSelfToVisLog(FVisualLogEntry* Snapshot) const
