missing files for behavior tree change
#ue4

diff --git a/Engine/Source/Runtime/AIModule/Classes/BehaviorTree/Tasks/BTTask_RunBehaviorDynamic.h b/Engine/Source/Runtime/AIModule/Classes/BehaviorTree/Tasks/BTTask_RunBehaviorDynamic.h
new file mode 100644
--- /dev/null
+++ b/Engine/Source/Runtime/AIModule/Classes/BehaviorTree/Tasks/BTTask_RunBehaviorDynamic.h
@@ -0,0 +1,54 @@
+// Copyright 1998-2015 Epic Games, Inc. All Rights Reserved.
+
+#pragma once
+#include "BehaviorTree/BehaviorTree.h"
+#include "BehaviorTree/BTTaskNode.h"
+#include "BTTask_RunBehaviorDynamic.generated.h"
+
+/**
+ * RunBehaviorDynamic task allows pushing subtrees on execution stack.
+ * Subtree asset can be assigned at runtime with SetDynamicSubtree function of BehaviorTreeComponent.
+ *
+ * Does NOT support subtree's root level decorators!
+ */
+
+UCLASS()
+class AIMODULE_API UBTTask_RunBehaviorDynamic : public UBTTaskNode
+{
+	GENERATED_UCLASS_BODY()
+
+	virtual EBTNodeResult::Type ExecuteTask(UBehaviorTreeComponent& OwnerComp, uint8* NodeMemory) override;
+	virtual void OnInstanceCreated(UBehaviorTreeComponent& OwnerComp) override;
+	virtual FString GetStaticDescription() const override;
+	virtual void DescribeRuntimeValues(const UBehaviorTreeComponent& OwnerComp, uint8* NodeMemory, EBTDescriptionVerbosity::Type Verbosity, TArray<FString>& Values) const override;
+
+#if WITH_EDITOR
+	virtual FName GetNodeIconName() const override;
+#endif // WITH_EDITOR
+
+	bool HasMatchingTag(const FGameplayTag& Tag) const;
+	void SetBehaviorAsset(UBehaviorTree* NewBehaviorAsset);
+
+protected:
+
+	/** Gameplay tag that will identify this task for subtree injection */
+	UPROPERTY(Category=Node, EditAnywhere)
+	FGameplayTag InjectionTag;
+
+	/** default behavior to run */
+	UPROPERTY(Category=Node, EditAnywhere)
+	UBehaviorTree* DefaultBehaviorAsset;
+
+	/** current subtree */
+	UPROPERTY()
+	UBehaviorTree* BehaviorAsset;
+};
+
+//////////////////////////////////////////////////////////////////////////
+// Inlines
+
+FORCEINLINE bool UBTTask_RunBehaviorDynamic::HasMatchingTag(const FGameplayTag& Tag) const
+{
+	return InjectionTag == Tag;
+}
+
diff --git a/Engine/Source/Runtime/AIModule/Private/BehaviorTree/Tasks/BTTask_RunBehaviorDynamic.cpp b/Engine/Source/Runtime/AIModule/Private/BehaviorTree/Tasks/BTTask_RunBehaviorDynamic.cpp
new file mode 100644
--- /dev/null
+++ b/Engine/Source/Runtime/AIModule/Private/BehaviorTree/Tasks/BTTask_RunBehaviorDynamic.cpp
@@ -0,0 +1,48 @@
+// Copyright 1998-2015 Epic Games, Inc. All Rights Reserved.
+
+#include "AIModulePrivate.h"
+#include "BehaviorTree/Tasks/BTTask_RunBehaviorDynamic.h"
+#include "BehaviorTree/BehaviorTreeManager.h"
+
+UBTTask_RunBehaviorDynamic::UBTTask_RunBehaviorDynamic(const FObjectInitializer& ObjectInitializer) : Super(ObjectInitializer)
+{
+	NodeName = "Run Behavior Dynamic";
+	bCreateNodeInstance = true;
+}
+
+EBTNodeResult::Type UBTTask_RunBehaviorDynamic::ExecuteTask(UBehaviorTreeComponent& OwnerComp, uint8* NodeMemory)
+{
+	const bool bPushed = BehaviorAsset != nullptr && OwnerComp.PushInstance(*BehaviorAsset);
+	return bPushed ? EBTNodeResult::InProgress : EBTNodeResult::Failed;
+}
+
+void UBTTask_RunBehaviorDynamic::OnInstanceCreated(UBehaviorTreeComponent& OwnerComp)
+{
+	Super::OnInstanceCreated(OwnerComp);
+	BehaviorAsset = DefaultBehaviorAsset;
+}
+
+FString UBTTask_RunBehaviorDynamic::GetStaticDescription() const
+{
+	return FString::Printf(TEXT("%s: %s"), *Super::GetStaticDescription(), *InjectionTag.ToString());
+}
+
+void UBTTask_RunBehaviorDynamic::DescribeRuntimeValues(const UBehaviorTreeComponent& OwnerComp, uint8* NodeMemory, EBTDescriptionVerbosity::Type Verbosity, TArray<FString>& Values) const
+{
+	Super::DescribeRuntimeValues(OwnerComp, NodeMemory, Verbosity, Values);
+	Values.Add(FString::Printf(TEXT("subtree: %s"), *GetNameSafe(BehaviorAsset)));
+}
+
+void UBTTask_RunBehaviorDynamic::SetBehaviorAsset(UBehaviorTree* NewBehaviorAsset)
+{
+	BehaviorAsset = NewBehaviorAsset;
+}
+
+#if WITH_EDITOR
+
+FName UBTTask_RunBehaviorDynamic::GetNodeIconName() const
+{
+	return FName("BTEditor.Graph.BTNode.Task.RunBehavior.Icon");
+}
+
+#endif	// WITH_EDITOR
