Fixed UE-6549 "EQSTestingPawn display queries in Components tab".

diff --git a/Engine/Source/Runtime/AIModule/Private/EnvironmentQuery/EQSRenderingComponent.cpp b/Engine/Source/Runtime/AIModule/Private/EnvironmentQuery/EQSRenderingComponent.cpp
--- a/Engine/Source/Runtime/AIModule/Private/EnvironmentQuery/EQSRenderingComponent.cpp
+++ b/Engine/Source/Runtime/AIModule/Private/EnvironmentQuery/EQSRenderingComponent.cpp
@@ -272,6 +272,12 @@ void FEQSSceneProxy::DrawDebugLabels(UCanvas* Canvas, APlayerController* PC)
 		return;
 	}
 
+	// little hacky test but it's the only way to remove text rendering from bad worlds, when using UDebugDrawService for it
+	if (Canvas && Canvas->SceneView && Canvas->SceneView->Family && Canvas->SceneView->Family->Scene && Canvas->SceneView->Family->Scene->GetWorld() != ActorOwner->GetWorld())
+	{
+		return;
+	}
+
 	FDebugRenderSceneProxy::DrawDebugLabels(Canvas, PC);
 }
 
