Fixed compilation isssue for monolitics and/or non unity builds.

diff --git a/Engine/Source/Runtime/AIModule/Private/EnvironmentQuery/EQSRenderingComponent.cpp b/Engine/Source/Runtime/AIModule/Private/EnvironmentQuery/EQSRenderingComponent.cpp
--- a/Engine/Source/Runtime/AIModule/Private/EnvironmentQuery/EQSRenderingComponent.cpp
+++ b/Engine/Source/Runtime/AIModule/Private/EnvironmentQuery/EQSRenderingComponent.cpp
@@ -1,6 +1,7 @@
 // Copyright 1998-2015 Epic Games, Inc. All Rights Reserved.
 
 #include "AIModulePrivate.h"
+#include "Engine/Canvas.h"
 #include "DebugRenderSceneProxy.h"
 #include "EnvironmentQuery/EQSRenderingComponent.h"
 #include "EnvironmentQuery/EQSQueryResultSourceInterface.h"
