Fixed typo

diff --git a/Engine/Source/Runtime/AIModule/Private/AIModule.cpp b/Engine/Source/Runtime/AIModule/Private/AIModule.cpp
--- a/Engine/Source/Runtime/AIModule/Private/AIModule.cpp
+++ b/Engine/Source/Runtime/AIModule/Private/AIModule.cpp
@@ -52,7 +52,7 @@ void FAIModule::StartupModule()
 		IAssetTools& AssetTools = FModuleManager::LoadModuleChecked<FAssetToolsModule>("AssetTools").Get();
 
 		// register AI category so that AI assets can register to it
-		AIAssetCategoryBit = AssetTools.RegisterAdvancedAssetCategory(FName(TEXT("AI")), LOCTEXT("AIAssetCategory", "Artificial Inteligence"));
+		AIAssetCategoryBit = AssetTools.RegisterAdvancedAssetCategory(FName(TEXT("AI")), LOCTEXT("AIAssetCategory", "Artificial Intelligence"));
 	}
 #endif // WITH_EDITOR 
 
