fixed multiple default subobjects with the same names created for behavior tree's blackboard filters
#ue4 UE-18476

diff --git a/Engine/Source/Runtime/AIModule/Private/EnvironmentQuery/Items/EnvQueryItemType.cpp b/Engine/Source/Runtime/AIModule/Private/EnvironmentQuery/Items/EnvQueryItemType.cpp
--- a/Engine/Source/Runtime/AIModule/Private/EnvironmentQuery/Items/EnvQueryItemType.cpp
+++ b/Engine/Source/Runtime/AIModule/Private/EnvironmentQuery/Items/EnvQueryItemType.cpp
@@ -16,7 +16,10 @@ UEnvQueryItemType::UEnvQueryItemType(const FObjectInitializer& ObjectInitializer
 void UEnvQueryItemType::FinishDestroy()
 {
 	// unregister from known types 
-	UEnvQueryManager::RegisteredItemTypes.RemoveSingleSwap(GetClass());
+	if (HasAnyFlags(RF_ClassDefaultObject))
+	{
+		UEnvQueryManager::RegisteredItemTypes.RemoveSingleSwap(GetClass());
+	}
 
 	Super::FinishDestroy();
 }
diff --git a/Engine/Source/Runtime/AIModule/Private/EnvironmentQuery/Items/EnvQueryItemType_ActorBase.cpp b/Engine/Source/Runtime/AIModule/Private/EnvironmentQuery/Items/EnvQueryItemType_ActorBase.cpp
--- a/Engine/Source/Runtime/AIModule/Private/EnvironmentQuery/Items/EnvQueryItemType_ActorBase.cpp
+++ b/Engine/Source/Runtime/AIModule/Private/EnvironmentQuery/Items/EnvQueryItemType_ActorBase.cpp
@@ -12,7 +12,7 @@ UEnvQueryItemType_ActorBase::UEnvQueryItemType_ActorBase(const FObjectInitialize
 void UEnvQueryItemType_ActorBase::AddBlackboardFilters(FBlackboardKeySelector& KeySelector, UObject* FilterOwner) const
 {
 	Super::AddBlackboardFilters(KeySelector, FilterOwner);
-	KeySelector.AddObjectFilter(FilterOwner, TEXT("EnvQueryActor"), AActor::StaticClass());
+	KeySelector.AddObjectFilter(FilterOwner, GetClass()->GetFName(), AActor::StaticClass());
 }
 
 bool UEnvQueryItemType_ActorBase::StoreInBlackboard(FBlackboardKeySelector& KeySelector, UBlackboardComponent* Blackboard, const uint8* RawData) const
diff --git a/Engine/Source/Runtime/AIModule/Private/EnvironmentQuery/Items/EnvQueryItemType_VectorBase.cpp b/Engine/Source/Runtime/AIModule/Private/EnvironmentQuery/Items/EnvQueryItemType_VectorBase.cpp
--- a/Engine/Source/Runtime/AIModule/Private/EnvironmentQuery/Items/EnvQueryItemType_VectorBase.cpp
+++ b/Engine/Source/Runtime/AIModule/Private/EnvironmentQuery/Items/EnvQueryItemType_VectorBase.cpp
@@ -11,7 +11,7 @@ UEnvQueryItemType_VectorBase::UEnvQueryItemType_VectorBase(const FObjectInitiali
 
 void UEnvQueryItemType_VectorBase::AddBlackboardFilters(FBlackboardKeySelector& KeySelector, UObject* FilterOwner) const
 {
-	KeySelector.AddVectorFilter(FilterOwner, TEXT("EnvQueryVector"));
+	KeySelector.AddVectorFilter(FilterOwner, GetClass()->GetFName());
 }
 
 bool UEnvQueryItemType_VectorBase::StoreInBlackboard(FBlackboardKeySelector& KeySelector, UBlackboardComponent* Blackboard, const uint8* RawData) const
