Fixed AI stimuli never expiring #UE4

This resulted in AI practically never forgetting last sensed location of given Actor, ignoring configured Max Age for given AI sense
UE-17052

diff --git a/Engine/Source/Runtime/AIModule/Private/Perception/AIPerceptionComponent.cpp b/Engine/Source/Runtime/AIModule/Private/Perception/AIPerceptionComponent.cpp
--- a/Engine/Source/Runtime/AIModule/Private/Perception/AIPerceptionComponent.cpp
+++ b/Engine/Source/Runtime/AIModule/Private/Perception/AIPerceptionComponent.cpp
@@ -432,13 +432,18 @@ void UAIPerceptionComponent::ProcessStimuli()
 			RefreshStimulus(StimulusStore, SourcedStimulus->Stimulus);
 		}
 		else if (StimulusStore.IsExpired() == false)
+		{	
+			if (bActorInfoUpdated)
+			{
+				// @note there some more valid info in SourcedStimulus->Stimulus regarding test that failed
+				// may be useful in future
+				StimulusStore.MarkNoLongerSensed();
+				StimulusStore.SetStimulusAge(0);
+			}
+		}
+		else if (StimulusStore.GetAge() != FAIStimulus::NeverHappenedAge)
 		{
-			//HandleExpiredStimulus(StimulusStore);
-
-			// @note there some more valid info in SourcedStimulus->Stimulus regarding test that failed
-			// may be useful in future
-			StimulusStore.MarkNoLongerSensed();
-			StimulusStore.SetStimulusAge(0);
+			HandleExpiredStimulus(StimulusStore);
 		}
 
 		// if the new stimulus is "valid" or it's info that "no longer sensed" and it used to be sensed successfully
@@ -480,6 +485,7 @@ void UAIPerceptionComponent::RefreshStimulus(FAIStimulus& StimulusStore, const F
 void UAIPerceptionComponent::HandleExpiredStimulus(FAIStimulus& StimulusStore)
 {
 	ensure(StimulusStore.IsExpired() == true);
+	StimulusStore = FAIStimulus();
 }
 
 bool UAIPerceptionComponent::AgeStimuli(const float ConstPerceptionAgingRate)
