Fixed AIController not unregistering from its pawn's GameplayTaskComponent notifications on Unpossess call #UE4

UE-18033

diff --git a/Engine/Source/Runtime/AIModule/Private/AIController.cpp b/Engine/Source/Runtime/AIModule/Private/AIController.cpp
--- a/Engine/Source/Runtime/AIModule/Private/AIController.cpp
+++ b/Engine/Source/Runtime/AIModule/Private/AIController.cpp
@@ -500,7 +500,11 @@ void AAIController::UnPossess()
 		BrainComponent->Cleanup();
 	}
 
-	CachedGameplayTasksComponent = nullptr;
+	if (CachedGameplayTasksComponent)
+	{
+		CachedGameplayTasksComponent->OnClaimedResourcesChange.RemoveDynamic(this, &AAIController::OnGameplayTaskResourcesClaimed);
+		CachedGameplayTasksComponent = nullptr;
+	}
 }
 
 void AAIController::SetPawn(APawn* InPawn)
