blueprint based behavior tree tasks will now ignore FinishExecute and FinishAbort calls depending on state of task, added another cleanup of all latent blueprint actions when task finishes execution
#ue4

diff --git a/Engine/Source/Runtime/AIModule/Classes/BehaviorTree/Tasks/BTTask_BlueprintBase.h b/Engine/Source/Runtime/AIModule/Classes/BehaviorTree/Tasks/BTTask_BlueprintBase.h
--- a/Engine/Source/Runtime/AIModule/Classes/BehaviorTree/Tasks/BTTask_BlueprintBase.h
+++ b/Engine/Source/Runtime/AIModule/Classes/BehaviorTree/Tasks/BTTask_BlueprintBase.h
@@ -26,6 +26,7 @@ class AIMODULE_API UBTTask_BlueprintBase : public UBTTaskNode
 	virtual FString GetStaticDescription() const override;
 	virtual void DescribeRuntimeValues(const UBehaviorTreeComponent& OwnerComp, uint8* NodeMemory, EBTDescriptionVerbosity::Type Verbosity, TArray<FString>& Values) const override;
 	virtual void OnInstanceDestroyed(UBehaviorTreeComponent& OwnerComp) override;
+	virtual void OnTaskFinished(UBehaviorTreeComponent& OwnerComp, uint8* NodeMemory, EBTNodeResult::Type TaskResult) override;
 
 	virtual void SetOwner(AActor* ActorOwner) override;
 
@@ -61,6 +62,9 @@ protected:
 	/** set if ReceiveAbort is implemented by blueprint */
 	uint32 ReceiveAbortImplementations : 2;
 
+	/** set when task enters Aborting state */
+	uint32 bIsAborting : 1;
+
 	/** if set, execution is inside blueprint's ReceiveExecute(Abort) event
 	  * FinishExecute(Abort) function should store their result in CurrentCallResult variable */
 	mutable uint32 bStoreFinishResult : 1;
@@ -124,6 +128,10 @@ protected:
 	UFUNCTION(BlueprintCallable, Category="AI|BehaviorTree")
 	bool IsTaskExecuting() const;
 	
+	/** check if task is currently being aborted */
+	UFUNCTION(BlueprintCallable, Category = "AI|BehaviorTree")
+	bool IsTaskAborting() const;
+
 	/** ticks this task */
 	virtual void TickTask(UBehaviorTreeComponent& OwnerComp, uint8* NodeMemory, float DeltaSeconds) override;
 
diff --git a/Engine/Source/Runtime/AIModule/Private/BehaviorTree/Tasks/BTTask_BlueprintBase.cpp b/Engine/Source/Runtime/AIModule/Private/BehaviorTree/Tasks/BTTask_BlueprintBase.cpp
--- a/Engine/Source/Runtime/AIModule/Private/BehaviorTree/Tasks/BTTask_BlueprintBase.cpp
+++ b/Engine/Source/Runtime/AIModule/Private/BehaviorTree/Tasks/BTTask_BlueprintBase.cpp
@@ -12,7 +12,9 @@ UBTTask_BlueprintBase::UBTTask_BlueprintBase(const FObjectInitializer& ObjectIni
 	ReceiveAbortImplementations = FBTNodeBPImplementationHelper::CheckEventImplementationVersion(TEXT("ReceiveAbort"), TEXT("ReceiveAbortAI"), this, StopAtClass);
 
 	bNotifyTick = ReceiveTickImplementations != FBTNodeBPImplementationHelper::NoImplementation;
+	bNotifyTaskFinished = true;
 	bShowPropertyDetails = true;
+	bIsAborting = false;
 
 	// all blueprint based nodes must create instances
 	bCreateNodeInstance = true;
@@ -33,6 +35,7 @@ EBTNodeResult::Type UBTTask_BlueprintBase::ExecuteTask(UBehaviorTreeComponent& O
 {
 	// fail when task doesn't react to execution (start or tick)
 	CurrentCallResult = (ReceiveExecuteImplementations != 0 || ReceiveTickImplementations != 0) ? EBTNodeResult::InProgress : EBTNodeResult::Failed;
+	bIsAborting = false;
 
 	if (ReceiveExecuteImplementations != FBTNodeBPImplementationHelper::NoImplementation)
 	{
@@ -60,6 +63,7 @@ EBTNodeResult::Type UBTTask_BlueprintBase::AbortTask(UBehaviorTreeComponent& Own
 	BlueprintNodeHelpers::AbortLatentActions(OwnerComp, *this);
 
 	CurrentCallResult = ReceiveAbortImplementations != 0 ? EBTNodeResult::InProgress : EBTNodeResult::Aborted;
+	bIsAborting = true;
 
 	if (ReceiveAbortImplementations != FBTNodeBPImplementationHelper::NoImplementation)
 	{
@@ -92,6 +96,16 @@ void UBTTask_BlueprintBase::TickTask(UBehaviorTreeComponent& OwnerComp, uint8* N
 	}
 }
 
+void UBTTask_BlueprintBase::OnTaskFinished(UBehaviorTreeComponent& OwnerComp, uint8* NodeMemory, EBTNodeResult::Type TaskResult)
+{
+	Super::OnTaskFinished(OwnerComp, NodeMemory, TaskResult);
+	
+	if (TaskResult != EBTNodeResult::InProgress)
+	{
+		BlueprintNodeHelpers::AbortLatentActions(OwnerComp, *this);
+	}
+}
+
 void UBTTask_BlueprintBase::FinishExecute(bool bSuccess)
 {
 	UBehaviorTreeComponent* OwnerComp = Cast<UBehaviorTreeComponent>(GetOuter());
@@ -101,7 +115,7 @@ void UBTTask_BlueprintBase::FinishExecute(bool bSuccess)
 	{
 		CurrentCallResult = NodeResult;
 	}
-	else if (OwnerComp)
+	else if (OwnerComp && !bIsAborting)
 	{
 		FinishLatentTask(*OwnerComp, NodeResult);
 	}
@@ -116,7 +130,7 @@ void UBTTask_BlueprintBase::FinishAbort()
 	{
 		CurrentCallResult = NodeResult;
 	}
-	else if (OwnerComp)
+	else if (OwnerComp && bIsAborting)
 	{
 		FinishLatentAbort(*OwnerComp);
 	}
@@ -130,6 +144,12 @@ bool UBTTask_BlueprintBase::IsTaskExecuting() const
 	return (TaskStatus == EBTTaskStatus::Active);
 }
 
+bool UBTTask_BlueprintBase::IsTaskAborting() const
+{
+	// use already cached data
+	return bIsAborting;
+}
+
 void UBTTask_BlueprintBase::SetFinishOnMessage(FName MessageName)
 {
 	UBehaviorTreeComponent* OwnerComp = Cast<UBehaviorTreeComponent>(GetOuter());
@@ -189,4 +209,4 @@ bool UBTTask_BlueprintBase::UsesBlueprint() const
 	return true;
 }
 
-#endif // WITH_EDITOR
\ No newline at end of file
+#endif // WITH_EDITOR
