Made AI MoveToActor requests safe when getting a null-GoalActor #UE4

UE-19364

diff --git a/Engine/Source/Runtime/AIModule/Classes/AITypes.h b/Engine/Source/Runtime/AIModule/Classes/AITypes.h
--- a/Engine/Source/Runtime/AIModule/Classes/AITypes.h
+++ b/Engine/Source/Runtime/AIModule/Classes/AITypes.h
@@ -471,8 +471,11 @@ struct AIMODULE_API FAIMoveRequest
 	FAIMoveRequest& SetAcceptanceRadius(float Radius) { AcceptanceRadius = Radius; return *this; }
 	FAIMoveRequest& SetUserData(FCustomMoveSharedPtr Data) { UserData = Data; return *this; }
 
-	bool HasGoalActor() const { return bHasGoalActor; }
-	AActor* GetGoalActor() const { return bHasGoalActor ? GoalActor : nullptr; }
+	/** the request should be either set up to move to a location, of go to a valid actor */
+	bool IsValid() const { return bMoveToActor == false || GoalActor != nullptr; }
+
+	bool IsMoveToActorRequest() const { return bMoveToActor; }
+	AActor* GetGoalActor() const { return bMoveToActor ? GoalActor : nullptr; }
 	FVector GetGoalLocation() const { return GoalLocation; }
 
 	bool IsUsingPathfinding() const { return bUsePathfinding; }
@@ -507,7 +510,7 @@ protected:
 	uint32 bInitialized : 1;
 
 	/** move goal is an actor */
-	uint32 bHasGoalActor : 1;
+	uint32 bMoveToActor : 1;
 
 	/** pathfinding: if set - regular pathfinding will be used, if not - direct path between two points */
 	uint32 bUsePathfinding : 1;
@@ -529,4 +532,9 @@ protected:
 
 	/** pathfollowing: custom user data */
 	FCustomMoveSharedPtr UserData;
+
+public:
+	// deprecated
+	DEPRECATED(4.8, "FAIMoveRequest::HasGoalActor is deprecated, please use FAIMoveRequest::IsMoveToActorRequest instead.")
+	bool HasGoalActor() const { return bMoveToActor; }
 };
diff --git a/Engine/Source/Runtime/AIModule/Private/AIController.cpp b/Engine/Source/Runtime/AIModule/Private/AIController.cpp
--- a/Engine/Source/Runtime/AIModule/Private/AIController.cpp
+++ b/Engine/Source/Runtime/AIModule/Private/AIController.cpp
@@ -562,11 +562,17 @@ EPathFollowingRequestResult::Type AAIController::MoveTo(const FAIMoveRequest& Mo
 	SCOPE_CYCLE_COUNTER(STAT_MoveTo);
 	UE_VLOG(this, LogAINavigation, Log, TEXT("MoveTo: %s"), *MoveRequest.ToString());
 
+	if (MoveRequest.IsValid() == false)
+	{
+		UE_VLOG(this, LogAINavigation, Error, TEXT("MoveTo request failed due MoveRequest not being valid. Most probably desireg Goal Actor not longer exists"), *MoveRequest.ToString());
+		return EPathFollowingRequestResult::Failed;
+	}
+
 	EPathFollowingRequestResult::Type Result = EPathFollowingRequestResult::Failed;
 	bool bCanRequestMove = true;
 	bool bAlreadyAtGoal = false;
 
-	if (!MoveRequest.HasGoalActor())
+	if (!MoveRequest.IsMoveToActorRequest())
 	{
 		if (MoveRequest.GetGoalLocation().ContainsNaN() || FAISystem::IsValidLocation(MoveRequest.GetGoalLocation()) == false)
 		{
@@ -593,7 +599,7 @@ EPathFollowingRequestResult::Type AAIController::MoveTo(const FAIMoveRequest& Mo
 		bAlreadyAtGoal = bCanRequestMove && PathFollowingComponent &&
 			PathFollowingComponent->HasReached(MoveRequest.GetGoalLocation(), MoveRequest.GetAcceptanceRadius(), !MoveRequest.CanStopOnOverlap());
 	}
-	else
+	else 
 	{
 		bAlreadyAtGoal = bCanRequestMove && PathFollowingComponent &&
 			PathFollowingComponent->HasReached(*MoveRequest.GetGoalActor(), MoveRequest.GetAcceptanceRadius(), !MoveRequest.CanStopOnOverlap());
@@ -700,7 +706,7 @@ bool AAIController::PreparePathfinding(const FAIMoveRequest& MoveRequest, FPathF
 		if (NavData)
 		{
 			FVector GoalLocation = MoveRequest.GetGoalLocation();
-			if (MoveRequest.HasGoalActor())
+			if (MoveRequest.IsMoveToActorRequest())
 			{
 				const INavAgentInterface* NavGoal = Cast<const INavAgentInterface>(MoveRequest.GetGoalActor());
 				if (NavGoal)
@@ -764,7 +770,7 @@ FAIRequestID AAIController::RequestPathAndMove(const FAIMoveRequest& MoveRequest
 		{
 			if (PathResult.IsSuccessful() && PathResult.Path.IsValid())
 			{
-				if (MoveRequest.HasGoalActor())
+				if (MoveRequest.IsMoveToActorRequest())
 				{
 					PathResult.Path->SetGoalActorObservation(*MoveRequest.GetGoalActor(), 100.0f);
 				}
diff --git a/Engine/Source/Runtime/AIModule/Private/AITypes.cpp b/Engine/Source/Runtime/AIModule/Private/AITypes.cpp
--- a/Engine/Source/Runtime/AIModule/Private/AITypes.cpp
+++ b/Engine/Source/Runtime/AIModule/Private/AITypes.cpp
@@ -156,7 +156,7 @@ const FAIRequestID FAIRequestID::CurrentRequest(FAIRequestID::AnyRequestID);
 
 FAIMoveRequest::FAIMoveRequest() : 
 	GoalActor(nullptr), GoalLocation(FAISystem::InvalidLocation), FilterClass(nullptr),
-	bInitialized(false), bHasGoalActor(false),
+	bInitialized(false), bMoveToActor(false),
 	bUsePathfinding(true), bAllowPartialPath(true), bProjectGoalOnNavigation(true),
 	bStopOnOverlap(true), bCanStrafe(false)
 {
@@ -164,8 +164,8 @@ FAIMoveRequest::FAIMoveRequest() :
 }
 
 FAIMoveRequest::FAIMoveRequest(const AActor* InGoalActor) :
-	GoalActor((AActor*)InGoalActor), GoalLocation(FAISystem::InvalidLocation), FilterClass(nullptr),
-	bInitialized(true), bHasGoalActor(true),
+	GoalActor(const_cast<AActor*>(InGoalActor)), GoalLocation(FAISystem::InvalidLocation), FilterClass(nullptr),
+	bInitialized(true), bMoveToActor(true),
 	bUsePathfinding(true), bAllowPartialPath(true), bProjectGoalOnNavigation(true),
 	bStopOnOverlap(true), bCanStrafe(false)
 {
@@ -174,7 +174,7 @@ FAIMoveRequest::FAIMoveRequest(const AActor* InGoalActor) :
 
 FAIMoveRequest::FAIMoveRequest(const FVector& InGoalLocation) :
 	GoalActor(nullptr), GoalLocation(InGoalLocation), FilterClass(nullptr),
-	bInitialized(true), bHasGoalActor(false),
+	bInitialized(true), bMoveToActor(false),
 	bUsePathfinding(true), bAllowPartialPath(true), bProjectGoalOnNavigation(true),
 	bStopOnOverlap(true), bCanStrafe(false)
 {
@@ -186,7 +186,7 @@ void FAIMoveRequest::SetGoalActor(const AActor* InGoalActor)
 	if (!bInitialized)
 	{
 		GoalActor = (AActor*)InGoalActor;
-		bHasGoalActor = true;
+		bMoveToActor = true;
 		bInitialized = true;
 	}
 }
@@ -202,7 +202,7 @@ void FAIMoveRequest::SetGoalLocation(const FVector& InGoalLocation)
 
 bool FAIMoveRequest::UpdateGoalLocation(const FVector& NewLocation) const
 {
-	if (!bHasGoalActor)
+	if (!bMoveToActor)
 	{
 		GoalLocation = NewLocation;
 		return true;
@@ -214,7 +214,7 @@ bool FAIMoveRequest::UpdateGoalLocation(const FVector& NewLocation) const
 FString FAIMoveRequest::ToString() const
 {
 	return FString::Printf(TEXT("%s(%s) Mode(%s) Filter(%s) AcceptanceRadius(%.1f%s)"),
-		bHasGoalActor ? TEXT("Actor") : TEXT("Location"), bHasGoalActor ? *GetNameSafe(GoalActor) : *GoalLocation.ToString(),
+		bMoveToActor ? TEXT("Actor") : TEXT("Location"), bMoveToActor ? *GetNameSafe(GoalActor) : *GoalLocation.ToString(),
 		bUsePathfinding ? (bAllowPartialPath ? TEXT("partial path") : TEXT("complete path")) : TEXT("direct"),
 		*GetNameSafe(FilterClass),
 		AcceptanceRadius, bStopOnOverlap ? TEXT(" + overlap") : TEXT("")
