Merging AI perception changes over from Orion branch #UE4

diff --git a/Engine/Source/Runtime/AIModule/Classes/Perception/AISenseConfig_Prediction.h b/Engine/Source/Runtime/AIModule/Classes/Perception/AISenseConfig_Prediction.h
--- a/Engine/Source/Runtime/AIModule/Classes/Perception/AISenseConfig_Prediction.h
+++ b/Engine/Source/Runtime/AIModule/Classes/Perception/AISenseConfig_Prediction.h
@@ -12,5 +12,5 @@ class AIMODULE_API UAISenseConfig_Prediction : public UAISenseConfig
 {
 	GENERATED_BODY()
 public:	
-	virtual TSubclassOf<UAISense> GetSenseImplementation() const override { return UAISense_Prediction::StaticClass(); }
+	virtual TSubclassOf<UAISense> GetSenseImplementation() const override;
 };
\ No newline at end of file
diff --git a/Engine/Source/Runtime/AIModule/Private/Perception/AISense.cpp b/Engine/Source/Runtime/AIModule/Private/Perception/AISense.cpp
--- a/Engine/Source/Runtime/AIModule/Private/Perception/AISense.cpp
+++ b/Engine/Source/Runtime/AIModule/Private/Perception/AISense.cpp
@@ -4,6 +4,8 @@
 #include "Perception/AISense.h"
 #include "Perception/AISenseConfig_Sight.h"
 #include "Perception/AISenseConfig_Hearing.h"
+#include "Perception/AISenseConfig_Prediction.h"
+#include "Perception/AISense_Prediction.h"
 
 #if !UE_BUILD_SHIPPING
 #include "DrawDebugHelpers.h"
@@ -181,3 +183,11 @@ void UAISenseConfig_Hearing::DrawDebugInfo(UCanvas& Canvas, UAIPerceptionCompone
 	}
 }
 #endif // !UE_BUILD_SHIPPING
+
+//----------------------------------------------------------------------//
+// UAISenseConfig_Prediction
+//----------------------------------------------------------------------//
+TSubclassOf<UAISense> UAISenseConfig_Prediction::GetSenseImplementation() const 
+{ 
+	return UAISense_Prediction::StaticClass(); 
+}
\ No newline at end of file
diff --git a/Engine/Source/Runtime/GameplayAbilities/Private/Abilities/GameplayAbility.cpp b/Engine/Source/Runtime/GameplayAbilities/Private/Abilities/GameplayAbility.cpp
--- a/Engine/Source/Runtime/GameplayAbilities/Private/Abilities/GameplayAbility.cpp
+++ b/Engine/Source/Runtime/GameplayAbilities/Private/Abilities/GameplayAbility.cpp
@@ -1473,3 +1473,8 @@ void UGameplayAbility::ConvertDeprecatedGameplayEffectReferencesToBlueprintRefer
 		MarkPackageDirty();
 	}
 }
+
+float UGameplayAbility::GetCooldownTimeRemaining() const
+{
+	return IsInstantiated() ? GetCooldownTimeRemaining(CurrentActorInfo) : 0.f;
+}
\ No newline at end of file
diff --git a/Engine/Source/Runtime/GameplayAbilities/Public/Abilities/GameplayAbility.h b/Engine/Source/Runtime/GameplayAbilities/Public/Abilities/GameplayAbility.h
--- a/Engine/Source/Runtime/GameplayAbilities/Public/Abilities/GameplayAbility.h
+++ b/Engine/Source/Runtime/GameplayAbilities/Public/Abilities/GameplayAbility.h
@@ -665,6 +665,9 @@ protected:
 	/** Deprecated. Use CooldownGameplayEffectClass instead */
 	UPROPERTY(VisibleDefaultsOnly, Category=Deprecated)
 	class UGameplayEffect* CooldownGameplayEffect;
+
+	UFUNCTION(BlueprintCallable, Category = Ability)
+	float GetCooldownTimeRemaining() const; 
 	
 	// ----------------------------------------------------------------------------------------------------------------
 	//
@@ -761,6 +764,7 @@ protected:
 	// 
 	// ----------------------------------------------------------------------------------------------------------------
 	
+public:
 	/** Returns current level of the Ability */
 	UFUNCTION(BlueprintCallable, Category = Ability)
 	int32 GetAbilityLevel() const;
