Made EQS queries BT types so that we can build parametrised EQS BT nodes #UE4

diff --git a/Engine/Source/Runtime/AIModule/Classes/EnvironmentQuery/EnvQuery.h b/Engine/Source/Runtime/AIModule/Classes/EnvironmentQuery/EnvQuery.h
--- a/Engine/Source/Runtime/AIModule/Classes/EnvironmentQuery/EnvQuery.h
+++ b/Engine/Source/Runtime/AIModule/Classes/EnvironmentQuery/EnvQuery.h
@@ -10,7 +10,7 @@ class UEnvQueryOption;
 class UEdGraph;
 #endif // WITH_EDITORONLY_DATA
 
-UCLASS()
+UCLASS(BlueprintType)
 class AIMODULE_API UEnvQuery : public UObject
 {
 	GENERATED_UCLASS_BODY()
