Fixed PawnAction vs BehaviorTree desync caused by path observer delegates from actions already removed from stack (with FailedToStart result). Needs better solution for FailedToStart in whatever in going to replace PawnActions soon.

diff --git a/Engine/Source/Runtime/AIModule/Classes/Actions/PawnAction.h b/Engine/Source/Runtime/AIModule/Classes/Actions/PawnAction.h
--- a/Engine/Source/Runtime/AIModule/Classes/Actions/PawnAction.h
+++ b/Engine/Source/Runtime/AIModule/Classes/Actions/PawnAction.h
@@ -103,6 +103,10 @@ protected:
 	UPROPERTY(Category = PawnAction, EditDefaultsOnly, BlueprintReadWrite)
 	uint32 bShouldPauseMovement : 1;
 
+	/** if set, action will call OnFinished notify even when ending as FailedToStart */
+	UPROPERTY(Category = PawnAction, EditDefaultsOnly, BlueprintReadWrite, AdvancedDisplay)
+	uint32 bAlwaysNotifyOnFinished : 1;
+
 private:
 	/** indicates whether action is in the process of abortion, and if so on what state */
 	EPawnActionAbortState::Type AbortState;
diff --git a/Engine/Source/Runtime/AIModule/Private/Actions/PawnAction.cpp b/Engine/Source/Runtime/AIModule/Private/Actions/PawnAction.cpp
--- a/Engine/Source/Runtime/AIModule/Private/Actions/PawnAction.cpp
+++ b/Engine/Source/Runtime/AIModule/Private/Actions/PawnAction.cpp
@@ -251,7 +251,7 @@ bool UPawnAction::Activate()
 void UPawnAction::OnPopped()
 {
 	// not calling OnFinish if action haven't actually started
-	if (bFailedToStart == false)
+	if (!bFailedToStart || bAlwaysNotifyOnFinished)
 	{
 		OnFinished(FinishResult);
 	}
diff --git a/Engine/Source/Runtime/AIModule/Private/Actions/PawnAction_Move.cpp b/Engine/Source/Runtime/AIModule/Private/Actions/PawnAction_Move.cpp
--- a/Engine/Source/Runtime/AIModule/Private/Actions/PawnAction_Move.cpp
+++ b/Engine/Source/Runtime/AIModule/Private/Actions/PawnAction_Move.cpp
@@ -15,6 +15,9 @@ UPawnAction_Move::UPawnAction_Move(const FObjectInitializer& ObjectInitializer)
 	, bAbortChildActionOnPathChange(false)
 {
 	bShouldPauseMovement = true;
+
+	// force using OnFinished notify to clear observer delegates from path when action leaves the stack
+	bAlwaysNotifyOnFinished = true;
 }
 
 void UPawnAction_Move::BeginDestroy()
