Fixed Parallel BT Composite node crash on stopping behavior tree #UE4

diff --git a/Engine/Source/Runtime/AIModule/Private/BehaviorTree/BehaviorTreeComponent.cpp b/Engine/Source/Runtime/AIModule/Private/BehaviorTree/BehaviorTreeComponent.cpp
--- a/Engine/Source/Runtime/AIModule/Private/BehaviorTree/BehaviorTreeComponent.cpp
+++ b/Engine/Source/Runtime/AIModule/Private/BehaviorTree/BehaviorTreeComponent.cpp
@@ -194,27 +194,27 @@ void UBehaviorTreeComponent::StopTree(EBTStopMode::Type StopMode)
 			// notify active parallel tasks
 			for (int32 ParallelIndex = 0; ParallelIndex < InstanceInfo.ParallelTasks.Num(); ParallelIndex++)
 			{
-				FBehaviorTreeParallelTask& ParalleInfo = InstanceInfo.ParallelTasks[ParallelIndex];
-				if (ParalleInfo.TaskNode && (ParalleInfo.Status == EBTTaskStatus::Active))
+				FBehaviorTreeParallelTask ParallelTaskInfo = InstanceInfo.ParallelTasks[ParallelIndex];
+				if (ParallelTaskInfo.TaskNode && (ParallelTaskInfo.Status == EBTTaskStatus::Active) && ParallelTaskInfo.TaskNode->IsPendingKill() == false)
 				{
 					// remove all message observers from node to abort to avoid calling OnTaskFinished from AbortTask
-					UnregisterMessageObserversFrom(ParalleInfo.TaskNode);
+					UnregisterMessageObserversFrom(ParallelTaskInfo.TaskNode);
 
-					uint8* NodeMemory = ParalleInfo.TaskNode->GetNodeMemory<uint8>(InstanceInfo);
-					EBTNodeResult::Type NodeResult = ParalleInfo.TaskNode->WrappedAbortTask(*this, NodeMemory);
+					uint8* NodeMemory = ParallelTaskInfo.TaskNode->GetNodeMemory<uint8>(InstanceInfo);
+					EBTNodeResult::Type NodeResult = ParallelTaskInfo.TaskNode->WrappedAbortTask(*this, NodeMemory);
 
 					UE_VLOG(GetOwner(), LogBehaviorTree, Log, TEXT("Parallel task aborted: %s (%s)"),
-						*UBehaviorTreeTypes::DescribeNodeHelper(ParalleInfo.TaskNode),
+						*UBehaviorTreeTypes::DescribeNodeHelper(ParallelTaskInfo.TaskNode),
 						(NodeResult == EBTNodeResult::InProgress) ? TEXT("in progress") : TEXT("instant"));
 
 					// mark as pending abort
 					if (NodeResult == EBTNodeResult::InProgress)
 					{
-						ParalleInfo.Status = EBTTaskStatus::Aborting;
+						ParallelTaskInfo.Status = EBTTaskStatus::Aborting;
 						bWaitingForAbortingTasks = true;
 					}
 
-					OnTaskFinished(ParalleInfo.TaskNode, NodeResult);
+					OnTaskFinished(ParallelTaskInfo.TaskNode, NodeResult);
 				}
 			}
 
