Fixed behavior tree decorators not observing in LowerPriority mode when search flow enters and leaves their branch without finding a task to run.

#ue4

diff --git a/Engine/Source/Runtime/AIModule/Classes/BehaviorTree/BehaviorTreeComponent.h b/Engine/Source/Runtime/AIModule/Classes/BehaviorTree/BehaviorTreeComponent.h
--- a/Engine/Source/Runtime/AIModule/Classes/BehaviorTree/BehaviorTreeComponent.h
+++ b/Engine/Source/Runtime/AIModule/Classes/BehaviorTree/BehaviorTreeComponent.h
@@ -190,6 +190,7 @@ public:
 
 	/** @return true if aux node is currently active */
 	bool IsAuxNodeActive(const UBTAuxiliaryNode* AuxNode) const;
+	bool IsAuxNodeActive(const UBTAuxiliaryNode* AuxNodeTemplate, int32 InstanceIdx) const;
 
 	/** @return status of speficied task */
 	EBTTaskStatus::Type GetTaskStatus(const UBTTaskNode* TaskNode) const;
diff --git a/Engine/Source/Runtime/AIModule/Private/BehaviorTree/BehaviorTreeComponent.cpp b/Engine/Source/Runtime/AIModule/Private/BehaviorTree/BehaviorTreeComponent.cpp
--- a/Engine/Source/Runtime/AIModule/Private/BehaviorTree/BehaviorTreeComponent.cpp
+++ b/Engine/Source/Runtime/AIModule/Private/BehaviorTree/BehaviorTreeComponent.cpp
@@ -471,6 +471,11 @@ bool UBehaviorTreeComponent::IsAuxNodeActive(const UBTAuxiliaryNode* AuxNode) co
 	return false;
 }
 
+bool UBehaviorTreeComponent::IsAuxNodeActive(const UBTAuxiliaryNode* AuxNodeTemplate, int32 InstanceIdx) const
+{
+	return InstanceStack.IsValidIndex(InstanceIdx) && InstanceStack[InstanceIdx].ActiveAuxNodes.Contains(AuxNodeTemplate);
+}
+
 EBTTaskStatus::Type UBehaviorTreeComponent::GetTaskStatus(const UBTTaskNode* TaskNode) const
 {
 	EBTTaskStatus::Type Status = EBTTaskStatus::Inactive;
diff --git a/Engine/Source/Runtime/AIModule/Private/BehaviorTree/BehaviorTreeTypes.cpp b/Engine/Source/Runtime/AIModule/Private/BehaviorTree/BehaviorTreeTypes.cpp
--- a/Engine/Source/Runtime/AIModule/Private/BehaviorTree/BehaviorTreeTypes.cpp
+++ b/Engine/Source/Runtime/AIModule/Private/BehaviorTree/BehaviorTreeTypes.cpp
@@ -200,12 +200,24 @@ void FBehaviorTreeSearchData::AddUniqueUpdate(const FBehaviorTreeSearchUpdate& U
 			PendingUpdates.RemoveAt(UpdateIndex, 1, false);
 		}
 	}
+	
+	// don't add Remove updates for inactive aux nodes, as they will block valid Add update coming later from the same search
+	// check only aux nodes, it happens due to UBTCompositeNode::NotifyDecoratorsOnActivation
+	if (!bSkipAdding && UpdateInfo.Mode == EBTNodeUpdateMode::Remove && UpdateInfo.AuxNode)
+	{
+		const bool bIsActive = OwnerComp.IsAuxNodeActive(UpdateInfo.AuxNode, UpdateInfo.InstanceIndex);
+		bSkipAdding = !bIsActive;
+	}
 
 	if (!bSkipAdding)
 	{
 		const int32 Idx = PendingUpdates.Add(UpdateInfo);
 		PendingUpdates[Idx].bPostUpdate = (UpdateInfo.Mode == EBTNodeUpdateMode::Add) && (Cast<UBTService>(UpdateInfo.AuxNode) != NULL);
 	}
+	else
+	{
+		UE_VLOG(OwnerComp.GetOwner(), LogBehaviorTree, Verbose, TEXT(">> or not, update skipped"));
+	}
 }
 
 void FBehaviorTreeSearchData::AssignSearchId()
