Typo fixes in comments.

diff --git a/Engine/Source/Runtime/AIModule/Classes/Perception/AIPerceptionComponent.h b/Engine/Source/Runtime/AIModule/Classes/Perception/AIPerceptionComponent.h
--- a/Engine/Source/Runtime/AIModule/Classes/Perception/AIPerceptionComponent.h
+++ b/Engine/Source/Runtime/AIModule/Classes/Perception/AIPerceptionComponent.h
@@ -116,27 +116,27 @@ class AIMODULE_API UAIPerceptionComponent : public UActorComponent
 
 protected:
 	/** Max distance at which a makenoise(1.0) loudness sound can be heard, regardless of occlusion */
-	DEPRECATED(4.8, "This property is deprecated. Please use apropriate sencse config class instead")
+	DEPRECATED(4.8, "This property is deprecated. Please use appropriate sense config class instead")
 	UPROPERTY(VisibleDefaultsOnly, Category = AI)
 	float HearingRange;
 
-	/** Max distance at which a makenoise(1.0) loudness sound can be heard if unoccluded (LOSHearingThreshold should be > HearingThreshold) */
-	DEPRECATED(4.8, "This property is deprecated. Please use apropriate sencse config class instead")
+	/** Max distance at which a makenoise(1.0) loudness sound can be heard if not occluded (LOSHearingThreshold should be > HearingThreshold) */
+	DEPRECATED(4.8, "This property is deprecated. Please use appropriate sense config class instead")
 	UPROPERTY(VisibleDefaultsOnly, Category = AI)
 	float LoSHearingRange;
 
 	/** Maximum sight distance to notice a target. */
-	DEPRECATED(4.8, "This property is deprecated. Please use apropriate sencse config class instead")
+	DEPRECATED(4.8, "This property is deprecated. Please use appropriate sense config class instead")
 	UPROPERTY(VisibleDefaultsOnly, Category = AI)
 	float SightRadius;
 
 	/** Maximum sight distance to see target that has been already seen. */
-	DEPRECATED(4.8, "This property is deprecated. Please use apropriate sencse config class instead")
+	DEPRECATED(4.8, "This property is deprecated. Please use appropriate sense config class instead")
 	UPROPERTY(VisibleDefaultsOnly, Category = AI)
 	float LoseSightRadius;
 
 	/** How far to the side AI can see, in degrees. Use SetPeripheralVisionAngle to change the value at runtime. */
-	DEPRECATED(4.8, "This property is deprecated. Please use apropriate sencse config class instead")
+	DEPRECATED(4.8, "This property is deprecated. Please use appropriate sense config class instead")
 	UPROPERTY(VisibleDefaultsOnly, Category = AI)
 	float PeripheralVisionAngle;
 		
@@ -144,7 +144,7 @@ protected:
 	TArray<UAISenseConfig*> SensesConfig;
 
 	/** Indicated sense that takes precedence over other senses when determining sensed actor's location. 
-	 *	Should be set to one of the sences configured in SensesConfig, or None. */
+	 *	Should be set to one of the senses configured in SensesConfig, or None. */
 	UPROPERTY(EditDefaultsOnly, Category = "AI Perception")
 	TSubclassOf<UAISense> DominantSense;
 	
@@ -205,7 +205,7 @@ public:
 
 	void GetHostileActors(TArray<AActor*>& OutActors) const;
 
-	// @note will stop on first age 0 stimulus
+	// @note Will stop on first age 0 stimulus
 	const FActorPerceptionInfo* GetFreshestTrace(const FAISenseID Sense) const;
 	
 	void SetDominantSense(TSubclassOf<UAISense> InDominantSense);
@@ -221,7 +221,7 @@ public:
 
 	void RegisterStimulus(AActor* Source, const FAIStimulus& Stimulus);
 	void ProcessStimuli();
-	/** returns true if, as result of stimuli aging, this listener needs an update (like if some stimuli expired)*/
+	/** Returns true if, as result of stimuli aging, this listener needs an update (like if some stimuli expired) */
 	bool AgeStimuli(const float ConstPerceptionAgingRate);
 	void ForgetActor(AActor* ActorToForget);
 
