fixed crash on stopping behavior tree with instantly aborting parallel tasks
#ue4

diff --git a/Engine/Source/Runtime/AIModule/Private/BehaviorTree/BehaviorTreeComponent.cpp b/Engine/Source/Runtime/AIModule/Private/BehaviorTree/BehaviorTreeComponent.cpp
--- a/Engine/Source/Runtime/AIModule/Private/BehaviorTree/BehaviorTreeComponent.cpp
+++ b/Engine/Source/Runtime/AIModule/Private/BehaviorTree/BehaviorTreeComponent.cpp
@@ -192,29 +192,51 @@ void UBehaviorTreeComponent::StopTree(EBTStopMode::Type StopMode)
 			InstanceInfo.ActiveAuxNodes.Reset();
 
 			// notify active parallel tasks
+			//
+			// calling OnTaskFinished with result other than InProgress will unregister parallel task,
+			// modifying array we're iterating on - iterator needs to be moved one step back in that case
+			//
 			for (int32 ParallelIndex = 0; ParallelIndex < InstanceInfo.ParallelTasks.Num(); ParallelIndex++)
 			{
-				FBehaviorTreeParallelTask ParallelTaskInfo = InstanceInfo.ParallelTasks[ParallelIndex];
-				if (ParallelTaskInfo.TaskNode && (ParallelTaskInfo.Status == EBTTaskStatus::Active) && ParallelTaskInfo.TaskNode->IsPendingKill() == false)
+				FBehaviorTreeParallelTask& ParallelTaskInfo = InstanceInfo.ParallelTasks[ParallelIndex];
+				const UBTTaskNode* CachedTaskNode = ParallelTaskInfo.TaskNode;
+
+				if (CachedTaskNode && (ParallelTaskInfo.Status == EBTTaskStatus::Active) && !CachedTaskNode->IsPendingKill())
 				{
-					// remove all message observers from node to abort to avoid calling OnTaskFinished from AbortTask
-					UnregisterMessageObserversFrom(ParallelTaskInfo.TaskNode);
+					// remove all message observers added by task execution, so they won't interfere with Abort call
+					UnregisterMessageObserversFrom(CachedTaskNode);
 
-					uint8* NodeMemory = ParallelTaskInfo.TaskNode->GetNodeMemory<uint8>(InstanceInfo);
-					EBTNodeResult::Type NodeResult = ParallelTaskInfo.TaskNode->WrappedAbortTask(*this, NodeMemory);
+					uint8* NodeMemory = CachedTaskNode->GetNodeMemory<uint8>(InstanceInfo);
+					EBTNodeResult::Type NodeResult = CachedTaskNode->WrappedAbortTask(*this, NodeMemory);
 
 					UE_VLOG(GetOwner(), LogBehaviorTree, Log, TEXT("Parallel task aborted: %s (%s)"),
-						*UBehaviorTreeTypes::DescribeNodeHelper(ParallelTaskInfo.TaskNode),
+						*UBehaviorTreeTypes::DescribeNodeHelper(CachedTaskNode),
 						(NodeResult == EBTNodeResult::InProgress) ? TEXT("in progress") : TEXT("instant"));
 
 					// mark as pending abort
 					if (NodeResult == EBTNodeResult::InProgress)
 					{
-						ParallelTaskInfo.Status = EBTTaskStatus::Aborting;
-						bWaitingForAbortingTasks = true;
+						const bool bIsValidForStatus = InstanceInfo.ParallelTasks.IsValidIndex(ParallelIndex) && (ParallelTaskInfo.TaskNode == CachedTaskNode);
+						if (bIsValidForStatus)
+						{
+							ParallelTaskInfo.Status = EBTTaskStatus::Aborting;
+							bWaitingForAbortingTasks = true;
+						}
+						else
+						{
+							UE_VLOG(GetOwner(), LogBehaviorTree, Warning, TEXT("Parallel task %s was unregistered before completing Abort state!"),
+								*UBehaviorTreeTypes::DescribeNodeHelper(CachedTaskNode));
+						}
 					}
 
-					OnTaskFinished(ParallelTaskInfo.TaskNode, NodeResult);
+					OnTaskFinished(CachedTaskNode, NodeResult);
+
+					const bool bIsValidAfterFinishing = InstanceInfo.ParallelTasks.IsValidIndex(ParallelIndex) && (ParallelTaskInfo.TaskNode == CachedTaskNode);
+					if (!bIsValidAfterFinishing)
+					{
+						// move iterator back if current task was unregistered
+						ParallelIndex--;
+					}
 				}
 			}
 
