Loosened constrains on AI sense listeners removal as part of listener updating #UE4

Merged over from 4.10 branch

diff --git a/Engine/Source/Runtime/AIModule/Private/Perception/AISense_Hearing.cpp b/Engine/Source/Runtime/AIModule/Private/Perception/AISense_Hearing.cpp
--- a/Engine/Source/Runtime/AIModule/Private/Perception/AISense_Hearing.cpp
+++ b/Engine/Source/Runtime/AIModule/Private/Perception/AISense_Hearing.cpp
@@ -104,7 +104,7 @@ void UAISense_Hearing::OnListenerUpdateImpl(const FPerceptionListener& UpdatedLi
 	}
 	else
 	{
-		DigestedProperties.FindAndRemoveChecked(ListenerID);
+		DigestedProperties.Remove(ListenerID);
 	}
 }
 
diff --git a/Engine/Source/Runtime/AIModule/Private/Perception/AISense_Sight.cpp b/Engine/Source/Runtime/AIModule/Private/Perception/AISense_Sight.cpp
--- a/Engine/Source/Runtime/AIModule/Private/Perception/AISense_Sight.cpp
+++ b/Engine/Source/Runtime/AIModule/Private/Perception/AISense_Sight.cpp
@@ -466,7 +466,7 @@ void UAISense_Sight::OnListenerUpdateImpl(const FPerceptionListener& UpdatedList
 	}
 	else
 	{
-		DigestedProperties.FindAndRemoveChecked(ListenerID);
+		DigestedProperties.Remove(ListenerID);
 	}
 }
 
