Fixed crashes related to out-of-order AI senses registration #UE4

Merged over from 4.10 branch

diff --git a/Engine/Source/Runtime/AIModule/Private/Perception/AIPerceptionSystem.cpp b/Engine/Source/Runtime/AIModule/Private/Perception/AIPerceptionSystem.cpp
--- a/Engine/Source/Runtime/AIModule/Private/Perception/AIPerceptionSystem.cpp
+++ b/Engine/Source/Runtime/AIModule/Private/Perception/AIPerceptionSystem.cpp
@@ -254,10 +254,10 @@ void UAIPerceptionSystem::Tick(float DeltaSeconds)
 			for (AIPerception::FListenerMap::TIterator ListenerIt(ListenerContainer); ListenerIt; ++ListenerIt)
 			{
 				check(ListenerIt->Value.Listener.IsValid())
-				if (ListenerIt->Value.HasAnyNewStimuli())
-				{
-					ListenerIt->Value.ProcessStimuli();
-				}
+					if (ListenerIt->Value.HasAnyNewStimuli())
+					{
+						ListenerIt->Value.ProcessStimuli();
+					}
 			}
 
 			bSomeListenersNeedUpdateDueToStimuliAging = false;
@@ -310,15 +310,15 @@ void UAIPerceptionSystem::UpdateListener(UAIPerceptionComponent& Listener)
 	SCOPE_CYCLE_COUNTER(STAT_AI_PerceptionSys);
 
 	if (Listener.IsPendingKill())
-    {
+	{
 		UnregisterListener(Listener);
-        return;
-    }
+		return;
+	}
 
 	const FPerceptionListenerID ListenerId = Listener.GetListenerId();
 
 	if (ListenerId != FPerceptionListenerID::InvalidID())
-    {
+	{
 		FPerceptionListener& ListenerEntry = ListenerContainer[ListenerId];
 		ListenerEntry.UpdateListenerProperties(Listener);
 		OnListenerUpdate(ListenerEntry);
@@ -331,7 +331,7 @@ void UAIPerceptionSystem::UpdateListener(UAIPerceptionComponent& Listener)
 		ListenerEntry.CacheLocation();
 				
 		OnNewListener(ListenerContainer[NewListenerId]);
-    }
+	}
 }
 
 void UAIPerceptionSystem::UnregisterListener(UAIPerceptionComponent& Listener)
@@ -479,6 +479,11 @@ void UAIPerceptionSystem::OnNewPawn(APawn& Pawn)
 
 	for (UAISense* Sense : Senses)
 	{
+		if (Sense == nullptr)
+		{
+			continue;
+		}
+
 		if (Sense->WantsNewPawnNotification())
 		{
 			Sense->OnNewPawn(Pawn);
@@ -497,7 +502,7 @@ void UAIPerceptionSystem::StartPlay()
 {
 	for (UAISense* Sense : Senses)
 	{
-		if (Sense->ShouldAutoRegisterAllPawnsAsSources())
+		if (Sense != nullptr && Sense->ShouldAutoRegisterAllPawnsAsSources())
 		{
 			FAISenseID SenseID = Sense->GetSenseID();
 			RegisterAllPawnsAsSourcesForSense(SenseID);
@@ -570,7 +575,7 @@ TSubclassOf<UAISense> UAIPerceptionSystem::GetSenseClassForStimulus(UObject* Wor
 {
 	TSubclassOf<UAISense> Result = nullptr;
 	UAIPerceptionSystem* PercSys = GetCurrent(WorldContext);
-	if (PercSys && PercSys->Senses.IsValidIndex(Stimulus.Type))
+	if (PercSys && PercSys->Senses.IsValidIndex(Stimulus.Type) && PercSys->Senses[Stimulus.Type] != nullptr)
 	{
 		Result = PercSys->Senses[Stimulus.Type]->GetClass();
 	}
@@ -611,11 +616,11 @@ void UAIPerceptionSystem::ReportPerceptionEvent(UObject* WorldContext, UAISenseE
 #if !UE_BUILD_SHIPPING
 FColor UAIPerceptionSystem::GetSenseDebugColor(FAISenseID SenseID) const
 {
-	return SenseID.IsValid() && Senses.IsValidIndex(SenseID) ? Senses[SenseID]->GetDebugColor() : FColor::Black;
+	return SenseID.IsValid() && Senses.IsValidIndex(SenseID) && Senses[SenseID] != nullptr ? Senses[SenseID]->GetDebugColor() : FColor::Black;
 }
 
 FString UAIPerceptionSystem::GetSenseName(FAISenseID SenseID) const
 {
-	return SenseID.IsValid() && Senses.IsValidIndex(SenseID) ? *Senses[SenseID]->GetDebugName() : TEXT("Invalid");
+	return SenseID.IsValid() && Senses.IsValidIndex(SenseID) && Senses[SenseID] != nullptr ? *Senses[SenseID]->GetDebugName() : TEXT("Invalid");
 }
 #endif 
