Fixed BT node instantiation code that would ocasionally produce UObject naming confilicts #UE4

Merged over from Orion branch
UE-20778

diff --git a/Engine/Source/Runtime/AIModule/Private/BehaviorTree/BTNode.cpp b/Engine/Source/Runtime/AIModule/Private/BehaviorTree/BTNode.cpp
--- a/Engine/Source/Runtime/AIModule/Private/BehaviorTree/BTNode.cpp
+++ b/Engine/Source/Runtime/AIModule/Private/BehaviorTree/BTNode.cpp
@@ -74,7 +74,7 @@ void UBTNode::InitializeInSubtree(UBehaviorTreeComponent& OwnerComp, uint8* Node
 		UBTNode* NodeInstance = OwnerComp.NodeInstances.IsValidIndex(NextInstancedIndex) ? OwnerComp.NodeInstances[NextInstancedIndex] : NULL;
 		if (NodeInstance == NULL)
 		{
-			NodeInstance = NewObject<UBTNode>(&OwnerComp, GetClass(), GetFName(), RF_NoFlags, (UObject*)(this));
+			NodeInstance = NewObject<UBTNode>(&OwnerComp, GetClass(), NAME_None, RF_NoFlags, (UObject*)(this));
 			NodeInstance->InitializeNode(GetParentNode(), GetExecutionIndex(), GetMemoryOffset(), GetTreeDepth());
 			NodeInstance->bIsInstanced = true;
 
