Added a way to remove all EQS queries tied to fiven Querier #UE4

Merged over from Orion branch.

diff --git a/Engine/Source/Runtime/AIModule/Classes/EnvironmentQuery/EnvQueryManager.h b/Engine/Source/Runtime/AIModule/Classes/EnvironmentQuery/EnvQueryManager.h
--- a/Engine/Source/Runtime/AIModule/Classes/EnvironmentQuery/EnvQueryManager.h
+++ b/Engine/Source/Runtime/AIModule/Classes/EnvironmentQuery/EnvQueryManager.h
@@ -129,6 +129,10 @@ class AIMODULE_API UEnvQueryManager : public UObject, public FTickableGameObject
 	int32 RunQuery(const FEnvQueryRequest& Request, EEnvQueryRunMode::Type RunMode, FQueryFinishedSignature const& FinishDelegate);
 	int32 RunQuery(TSharedPtr<FEnvQueryInstance> QueryInstance, FQueryFinishedSignature const& FinishDelegate);
 
+	/** Removed all active queries asked by Querier. No "on finished" notifications are being sent, call this function when
+	 *	you no longer care about Querier's queries, like when he's "dead" */
+	void SilentlyRemoveAllQueriesByQuerier(const UObject& Querier);
+
 	/** alternative way to run queries. Do not use for anything other then testing! (worse performance) */
 	TSharedPtr<FEnvQueryResult> RunInstantQuery(const FEnvQueryRequest& Request, EEnvQueryRunMode::Type RunMode);
 
diff --git a/Engine/Source/Runtime/AIModule/Private/EnvironmentQuery/EnvQueryManager.cpp b/Engine/Source/Runtime/AIModule/Private/EnvironmentQuery/EnvQueryManager.cpp
--- a/Engine/Source/Runtime/AIModule/Private/EnvironmentQuery/EnvQueryManager.cpp
+++ b/Engine/Source/Runtime/AIModule/Private/EnvironmentQuery/EnvQueryManager.cpp
@@ -203,6 +203,18 @@ TSharedPtr<FEnvQueryResult> UEnvQueryManager::RunInstantQuery(const FEnvQueryReq
 	return QueryInstance;
 }
 
+void UEnvQueryManager::SilentlyRemoveAllQueriesByQuerier(const UObject& Querier)
+{
+	for (int32 QueryIndex = RunningQueries.Num() - 1; QueryIndex >= 0; --QueryIndex)
+	{
+		const TSharedPtr<FEnvQueryInstance>& QueryInstance = RunningQueries[QueryIndex];
+		if (QueryInstance.IsValid() == false || QueryInstance->Owner.IsValid() == false || QueryInstance->Owner.Get() == &Querier)
+		{
+			RunningQueries.RemoveAt(QueryIndex, 1, /*bAllowShrinking=*/false);
+		}
+	}
+}
+
 TSharedPtr<FEnvQueryInstance> UEnvQueryManager::PrepareQueryInstance(const FEnvQueryRequest& Request, EEnvQueryRunMode::Type RunMode)
 {
 	TSharedPtr<FEnvQueryInstance> QueryInstance = CreateQueryInstance(Request.QueryTemplate, RunMode);
