Copying //UE4/Dev-Framework to Dev-Main (//UE4/Dev-Main)

diff --git a/Engine/Source/Runtime/AIModule/Private/BehaviorTree/BehaviorTreeComponent.cpp b/Engine/Source/Runtime/AIModule/Private/BehaviorTree/BehaviorTreeComponent.cpp
--- a/Engine/Source/Runtime/AIModule/Private/BehaviorTree/BehaviorTreeComponent.cpp
+++ b/Engine/Source/Runtime/AIModule/Private/BehaviorTree/BehaviorTreeComponent.cpp
@@ -363,9 +363,9 @@ void UBehaviorTreeComponent::OnTaskFinished(const UBTTaskNode* TaskNode, EBTNode
 		TaskNode->WrappedOnTaskFinished(*this, TaskMemory, TaskResult);
 
 		// update execution when active task is finished
-		FBehaviorTreeInstance& ActiveInstance = InstanceStack[ActiveInstanceIdx];
-		if (ActiveInstance.ActiveNode == TaskNode)
+		if (InstanceStack.IsValidIndex(ActiveInstanceIdx) && InstanceStack[ActiveInstanceIdx].ActiveNode == TaskNode)
 		{
+			FBehaviorTreeInstance& ActiveInstance = InstanceStack[ActiveInstanceIdx];
 			const bool bWasAborting = (ActiveInstance.ActiveNodeType == EBTActiveNode::AbortingTask);
 			ActiveInstance.ActiveNodeType = EBTActiveNode::InactiveTask;
 
@@ -375,7 +375,7 @@ void UBehaviorTreeComponent::OnTaskFinished(const UBTTaskNode* TaskNode, EBTNode
 				RequestExecution(TaskResult);
 			}
 		}
-		else if (TaskResult == EBTNodeResult::Aborted && InstanceStack[TaskInstanceIdx].ActiveNode == TaskNode)
+		else if (TaskResult == EBTNodeResult::Aborted && InstanceStack.IsValidIndex(TaskInstanceIdx) && InstanceStack[TaskInstanceIdx].ActiveNode == TaskNode)
 		{
 			// active instance may be already changed when getting back from AbortCurrentTask 
 			// (e.g. new task is higher on stack)
@@ -412,7 +412,7 @@ void UBehaviorTreeComponent::OnTreeFinished()
 	ActiveInstanceIdx = 0;
 	StoreDebuggerExecutionStep(EBTExecutionSnap::OutOfNodes);
 
-	if (bLoopExecution)
+	if (bLoopExecution && InstanceStack.Num())
 	{
 		// it should be already deactivated (including root)
 		// set active node to initial state: root activation
@@ -437,7 +437,7 @@ void UBehaviorTreeComponent::OnTreeFinished()
 bool UBehaviorTreeComponent::IsExecutingBranch(const UBTNode* Node, int32 ChildIndex) const
 {
 	const int32 TestInstanceIdx = FindInstanceContainingNode(Node);
-	if (TestInstanceIdx == INDEX_NONE || InstanceStack[TestInstanceIdx].ActiveNode == NULL)
+	if (!InstanceStack.IsValidIndex(TestInstanceIdx) || InstanceStack[TestInstanceIdx].ActiveNode == NULL)
 	{
 		return false;
 	}
@@ -503,7 +503,7 @@ EBTTaskStatus::Type UBehaviorTreeComponent::GetTaskStatus(const UBTTaskNode* Tas
 	EBTTaskStatus::Type Status = EBTTaskStatus::Inactive;
 	const int32 InstanceIdx = FindInstanceContainingNode(TaskNode);
 
-	if (InstanceIdx != INDEX_NONE)
+	if (InstanceStack.IsValidIndex(InstanceIdx))
 	{
 		const uint16 ExecutionIndex = TaskNode->GetExecutionIndex();
 		const FBehaviorTreeInstance& InstanceInfo = InstanceStack[InstanceIdx];
@@ -605,7 +605,7 @@ void UBehaviorTreeComponent::RequestExecution(EBTNodeResult::Type LastResult)
 	// task helpers can't continue with InProgress or Aborted result, it should be handled 
 	// either by decorator helper or regular RequestExecution() (6 param version)
 
-	if (LastResult != EBTNodeResult::Aborted && LastResult != EBTNodeResult::InProgress)
+	if (LastResult != EBTNodeResult::Aborted && LastResult != EBTNodeResult::InProgress && InstanceStack.IsValidIndex(ActiveInstanceIdx))
 	{
 		const FBehaviorTreeInstance& ActiveInstance = InstanceStack[ActiveInstanceIdx];
 		UBTCompositeNode* ExecuteParent = (ActiveInstance.ActiveNode == NULL) ? ActiveInstance.RootNode :
@@ -692,7 +692,7 @@ void UBehaviorTreeComponent::RequestExecution(UBTCompositeNode* RequestedOn, int
 		*UBehaviorTreeTypes::DescribeNodeHelper(RequestedBy),
 		*UBehaviorTreeTypes::DescribeNodeResult(ContinueWithResult));
 
-	if (!bIsRunning || (GetOwner() && GetOwner()->IsPendingKillPending()))
+	if (!bIsRunning || !InstanceStack.IsValidIndex(ActiveInstanceIdx) || (GetOwner() && GetOwner()->IsPendingKillPending()))
 	{
 		UE_VLOG(GetOwner(), LogBehaviorTree, Log, TEXT("> skip: tree is not running"));
 		return;
@@ -875,7 +875,11 @@ void UBehaviorTreeComponent::ApplySearchUpdates(const TArray<FBehaviorTreeSearch
 	for (int32 Index = 0; Index < UpdateList.Num(); Index++)
 	{
 		const FBehaviorTreeSearchUpdate& UpdateInfo = UpdateList[Index];
-		const UBTNode* UpdateNode = UpdateInfo.AuxNode ? (const UBTNode*)UpdateInfo.AuxNode : (const UBTNode*)UpdateInfo.TaskNode;
+		if (!InstanceStack.IsValidIndex(UpdateInfo.InstanceIndex))
+		{
+			continue;
+		}
+
 		FBehaviorTreeInstance& UpdateInstance = InstanceStack[UpdateInfo.InstanceIndex];
 		int32 ParallelTaskIdx = INDEX_NONE;
 		bool bIsActive = false;
@@ -890,6 +894,7 @@ void UBehaviorTreeComponent::ApplySearchUpdates(const TArray<FBehaviorTreeSearch
 			bIsActive = (ParallelTaskIdx != INDEX_NONE && UpdateInstance.ParallelTasks[ParallelTaskIdx].Status == EBTTaskStatus::Active);
 		}
 
+		const UBTNode* UpdateNode = UpdateInfo.AuxNode ? (const UBTNode*)UpdateInfo.AuxNode : (const UBTNode*)UpdateInfo.TaskNode;
 		if ((UpdateInfo.Mode == EBTNodeUpdateMode::Remove && !bIsActive) ||
 			(UpdateInfo.Mode == EBTNodeUpdateMode::Add && (bIsActive || UpdateNode->GetExecutionIndex() > NewNodeExecutionIndex)) ||
 			(UpdateInfo.bPostUpdate != bPostUpdate))
@@ -1053,12 +1058,24 @@ void UBehaviorTreeComponent::TickComponent(float DeltaTime, enum ELevelTick Tick
 			ActiveTask->WrappedTickTask(*this, NodeMemory, DeltaTime);
 		}
 	}
+
+	// tick aborting task from abandoned subtree
+	if (InstanceStack.IsValidIndex(ActiveInstanceIdx + 1))
+	{
+		FBehaviorTreeInstance& LastInstance = InstanceStack.Last();
+		if (LastInstance.ActiveNodeType == EBTActiveNode::AbortingTask)
+		{
+			UBTTaskNode* ActiveTask = (UBTTaskNode*)LastInstance.ActiveNode;
+			uint8* NodeMemory = ActiveTask->GetNodeMemory<uint8>(LastInstance);
+			ActiveTask->WrappedTickTask(*this, NodeMemory, DeltaTime);
+		}
+	}
 }
 
 void UBehaviorTreeComponent::ProcessExecutionRequest()
 {
 	bRequestedFlowUpdate = false;
-	if (!IsRegistered())
+	if (!IsRegistered() || !InstanceStack.IsValidIndex(ActiveInstanceIdx))
 	{
 		// it shouldn't be called, component is no longer valid
 		return;
@@ -1516,7 +1533,7 @@ void UBehaviorTreeComponent::UnregisterMessageObserversFrom(const FBTNodeIndex&
 
 void UBehaviorTreeComponent::UnregisterMessageObserversFrom(const UBTTaskNode* TaskNode)
 {
-	if (TaskNode)
+	if (TaskNode && InstanceStack.Num())
 	{
 		const FBehaviorTreeInstance& ActiveInstance = InstanceStack.Last();
 
@@ -1540,16 +1557,19 @@ void UBehaviorTreeComponent::UnregisterMessageObserversFrom(const UBTTaskNode* T
 
 void UBehaviorTreeComponent::RegisterParallelTask(const UBTTaskNode* TaskNode)
 {
-	FBehaviorTreeInstance& InstanceInfo = InstanceStack[ActiveInstanceIdx];
-	AUX_NODE_WRAPPER(InstanceInfo.ParallelTasks.Add( FBehaviorTreeParallelTask(TaskNode, EBTTaskStatus::Active) ));
-	
-	UE_VLOG(GetOwner(), LogBehaviorTree, Verbose, TEXT("Parallel task: %s added to active list"),
-		*UBehaviorTreeTypes::DescribeNodeHelper(TaskNode));
-
-	if (InstanceInfo.ActiveNode == TaskNode)
+	if (InstanceStack.IsValidIndex(ActiveInstanceIdx))
 	{
-		// switch to inactive state, so it could start background tree
-		InstanceInfo.ActiveNodeType = EBTActiveNode::InactiveTask;
+		FBehaviorTreeInstance& InstanceInfo = InstanceStack[ActiveInstanceIdx];
+		AUX_NODE_WRAPPER(InstanceInfo.ParallelTasks.Add(FBehaviorTreeParallelTask(TaskNode, EBTTaskStatus::Active)));
+
+		UE_VLOG(GetOwner(), LogBehaviorTree, Verbose, TEXT("Parallel task: %s added to active list"),
+			*UBehaviorTreeTypes::DescribeNodeHelper(TaskNode));
+
+		if (InstanceInfo.ActiveNode == TaskNode)
+		{
+			// switch to inactive state, so it could start background tree
+			InstanceInfo.ActiveNodeType = EBTActiveNode::InactiveTask;
+		}
 	}
 }
 
@@ -2290,7 +2310,7 @@ void UBehaviorTreeComponent::StoreDebuggerRuntimeValues(TArray<FString>& Runtime
 void UBehaviorTreeComponent::UpdateDebuggerAfterExecution(const UBTTaskNode* TaskNode, uint16 InstanceIdx) const
 {
 #if USE_BEHAVIORTREE_DEBUGGER
-	if (!IsDebuggerActive())
+	if (!IsDebuggerActive() || !InstanceStack.IsValidIndex(InstanceIdx))
 	{
 		return;
 	}
diff --git a/Engine/Source/Runtime/AIModule/Private/BehaviorTree/BehaviorTreeManager.cpp b/Engine/Source/Runtime/AIModule/Private/BehaviorTree/BehaviorTreeManager.cpp
--- a/Engine/Source/Runtime/AIModule/Private/BehaviorTree/BehaviorTreeManager.cpp
+++ b/Engine/Source/Runtime/AIModule/Private/BehaviorTree/BehaviorTreeManager.cpp
@@ -218,12 +218,25 @@ bool UBehaviorTreeManager::LoadTree(UBehaviorTree& Asset, UBTCompositeNode*& Roo
 	return false;
 }
 
-void UBehaviorTreeManager::InitializeMemoryHelper(const TArray<UBTDecorator*>& Nodes, TArray<uint16>& MemoryOffsets, int32& MemorySize)
+void UBehaviorTreeManager::InitializeMemoryHelper(const TArray<UBTDecorator*>& Nodes, TArray<uint16>& MemoryOffsets, int32& MemorySize, bool bForceInstancing)
 {
 	TArray<FNodeInitializationData> InitList;
 	for (int32 NodeIndex = 0; NodeIndex < Nodes.Num(); NodeIndex++)
 	{
+		UBTNode* Node = Nodes[NodeIndex];
+
+		const bool bUsesInstancing = Node->HasInstance();
+		if (bForceInstancing)
+		{
+			Node->ForceInstancing(true);
+		}
+
 		InitList.Add(FNodeInitializationData(Nodes[NodeIndex], NULL, 0, 0, Nodes[NodeIndex]->GetInstanceMemorySize(), Nodes[NodeIndex]->GetSpecialMemorySize()));
+
+		if (bForceInstancing && !bUsesInstancing)
+		{
+			Node->ForceInstancing(false);
+		}
 	}
 
 	InitList.Sort(FNodeInitializationData::FMemorySort());
diff --git a/Engine/Source/Runtime/AIModule/Private/BehaviorTree/BlueprintNodeHelpers.cpp b/Engine/Source/Runtime/AIModule/Private/BehaviorTree/BlueprintNodeHelpers.cpp
--- a/Engine/Source/Runtime/AIModule/Private/BehaviorTree/BlueprintNodeHelpers.cpp
+++ b/Engine/Source/Runtime/AIModule/Private/BehaviorTree/BlueprintNodeHelpers.cpp
@@ -264,4 +264,10 @@ namespace BlueprintNodeHelpers
 			AbortLatentActions(OwnerOb, Ob);
 		}
 	}
+
+	bool HasBlueprintFunction(FName FuncName, const UObject* Object, const UClass* StopAtClass)
+	{
+		return Object && StopAtClass && HasBlueprintFunction(FuncName, *Object, *StopAtClass);
+	}
 }
+
diff --git a/Engine/Source/Runtime/AIModule/Private/BehaviorTree/Decorators/BTDecorator_BlueprintBase.cpp b/Engine/Source/Runtime/AIModule/Private/BehaviorTree/Decorators/BTDecorator_BlueprintBase.cpp
--- a/Engine/Source/Runtime/AIModule/Private/BehaviorTree/Decorators/BTDecorator_BlueprintBase.cpp
+++ b/Engine/Source/Runtime/AIModule/Private/BehaviorTree/Decorators/BTDecorator_BlueprintBase.cpp
@@ -10,12 +10,12 @@ UBTDecorator_BlueprintBase::UBTDecorator_BlueprintBase(const FObjectInitializer&
 	: Super(ObjectInitializer)
 {
 	UClass* StopAtClass = UBTDecorator_BlueprintBase::StaticClass();
-	ReceiveTickImplementations = FBTNodeBPImplementationHelper::CheckEventImplementationVersion(TEXT("ReceiveTick"), TEXT("ReceiveTickAI"), this, StopAtClass);
-	ReceiveExecutionStartImplementations = FBTNodeBPImplementationHelper::CheckEventImplementationVersion(TEXT("ReceiveExecutionStart"), TEXT("ReceiveExecutionStartAI"), this, StopAtClass);
-	ReceiveExecutionFinishImplementations = FBTNodeBPImplementationHelper::CheckEventImplementationVersion(TEXT("ReceiveExecutionFinish"), TEXT("ReceiveExecutionFinishAI"), this, StopAtClass);
-	ReceiveObserverActivatedImplementations = FBTNodeBPImplementationHelper::CheckEventImplementationVersion(TEXT("ReceiveObserverActivated"), TEXT("ReceiveObserverActivatedAI"), this, StopAtClass);
-	ReceiveObserverDeactivatedImplementations = FBTNodeBPImplementationHelper::CheckEventImplementationVersion(TEXT("ReceiveObserverDeactivated"), TEXT("ReceiveObserverDeactivatedAI"), this, StopAtClass);
-	PerformConditionCheckImplementations = FBTNodeBPImplementationHelper::CheckEventImplementationVersion(TEXT("PerformConditionCheck"), TEXT("PerformConditionCheckAI"), this, StopAtClass);
+	ReceiveTickImplementations = FBTNodeBPImplementationHelper::CheckEventImplementationVersion(TEXT("ReceiveTick"), TEXT("ReceiveTickAI"), *this, *StopAtClass);
+	ReceiveExecutionStartImplementations = FBTNodeBPImplementationHelper::CheckEventImplementationVersion(TEXT("ReceiveExecutionStart"), TEXT("ReceiveExecutionStartAI"), *this, *StopAtClass);
+	ReceiveExecutionFinishImplementations = FBTNodeBPImplementationHelper::CheckEventImplementationVersion(TEXT("ReceiveExecutionFinish"), TEXT("ReceiveExecutionFinishAI"), *this, *StopAtClass);
+	ReceiveObserverActivatedImplementations = FBTNodeBPImplementationHelper::CheckEventImplementationVersion(TEXT("ReceiveObserverActivated"), TEXT("ReceiveObserverActivatedAI"), *this, *StopAtClass);
+	ReceiveObserverDeactivatedImplementations = FBTNodeBPImplementationHelper::CheckEventImplementationVersion(TEXT("ReceiveObserverDeactivated"), TEXT("ReceiveObserverDeactivatedAI"), *this, *StopAtClass);
+	PerformConditionCheckImplementations = FBTNodeBPImplementationHelper::CheckEventImplementationVersion(TEXT("PerformConditionCheck"), TEXT("PerformConditionCheckAI"), *this, *StopAtClass);
 
 	bNotifyBecomeRelevant = ReceiveObserverActivatedImplementations != 0;
 	bNotifyCeaseRelevant = ReceiveObserverDeactivatedImplementations != 0;
@@ -75,7 +75,7 @@ void UBTDecorator_BlueprintBase::SetOwner(AActor* InActorOwner)
 
 void UBTDecorator_BlueprintBase::OnBecomeRelevant(UBehaviorTreeComponent& OwnerComp, uint8* NodeMemory)
 {
-	if (AIOwner != nullptr && ReceiveObserverActivatedImplementations & FBTNodeBPImplementationHelper::AISpecific)
+	if (AIOwner != nullptr && (ReceiveObserverActivatedImplementations & FBTNodeBPImplementationHelper::AISpecific))
 	{
 		ReceiveObserverActivatedAI(AIOwner, AIOwner->GetPawn());
 	}
@@ -113,7 +113,7 @@ void UBTDecorator_BlueprintBase::OnCeaseRelevant(UBehaviorTreeComponent& OwnerCo
 		BlackboardComp->UnregisterObserversFrom(this);
 	}
 		
-	if (AIOwner != nullptr && ReceiveObserverDeactivatedImplementations & FBTNodeBPImplementationHelper::AISpecific)
+	if (AIOwner != nullptr && (ReceiveObserverDeactivatedImplementations & FBTNodeBPImplementationHelper::AISpecific))
 	{
 		ReceiveObserverDeactivatedAI(AIOwner, AIOwner->GetPawn());
 	}
@@ -131,7 +131,7 @@ void UBTDecorator_BlueprintBase::OnCeaseRelevant(UBehaviorTreeComponent& OwnerCo
 
 void UBTDecorator_BlueprintBase::OnNodeActivation(FBehaviorTreeSearchData& SearchData)
 {
-	if (AIOwner != nullptr && ReceiveExecutionStartImplementations & FBTNodeBPImplementationHelper::AISpecific)
+	if (AIOwner != nullptr && (ReceiveExecutionStartImplementations & FBTNodeBPImplementationHelper::AISpecific))
 	{
 		ReceiveExecutionStartAI(AIOwner, AIOwner->GetPawn());
 	}
@@ -143,7 +143,7 @@ void UBTDecorator_BlueprintBase::OnNodeActivation(FBehaviorTreeSearchData& Searc
 
 void UBTDecorator_BlueprintBase::OnNodeDeactivation(FBehaviorTreeSearchData& SearchData, EBTNodeResult::Type NodeResult)
 {
-	if (AIOwner != nullptr && ReceiveExecutionFinishImplementations & FBTNodeBPImplementationHelper::AISpecific)
+	if (AIOwner != nullptr && (ReceiveExecutionFinishImplementations & FBTNodeBPImplementationHelper::AISpecific))
 	{
 		ReceiveExecutionFinishAI(AIOwner, AIOwner->GetPawn(), NodeResult);
 	}
@@ -155,7 +155,7 @@ void UBTDecorator_BlueprintBase::OnNodeDeactivation(FBehaviorTreeSearchData& Sea
 
 void UBTDecorator_BlueprintBase::TickNode(UBehaviorTreeComponent& OwnerComp, uint8* NodeMemory, float DeltaSeconds)
 {
-	if (AIOwner != nullptr && ReceiveTickImplementations & FBTNodeBPImplementationHelper::AISpecific)
+	if (AIOwner != nullptr && (ReceiveTickImplementations & FBTNodeBPImplementationHelper::AISpecific))
 	{
 		ReceiveTickAI(AIOwner, AIOwner->GetPawn(), DeltaSeconds);
 	}
@@ -202,7 +202,7 @@ bool UBTDecorator_BlueprintBase::CalculateRawConditionValueImpl(UBehaviorTreeCom
 		// can't use const functions with blueprints
 		UBTDecorator_BlueprintBase* MyNode = (UBTDecorator_BlueprintBase*)this;
 
-		if (AIOwner != nullptr && PerformConditionCheckImplementations & FBTNodeBPImplementationHelper::AISpecific)
+		if (AIOwner != nullptr && (PerformConditionCheckImplementations & FBTNodeBPImplementationHelper::AISpecific))
 		{
 			CurrentCallResult = MyNode->PerformConditionCheckAI(MyNode->AIOwner, MyNode->AIOwner->GetPawn());
 		}
diff --git a/Engine/Source/Runtime/AIModule/Private/BehaviorTree/Decorators/BTDecorator_ConeCheck.cpp b/Engine/Source/Runtime/AIModule/Private/BehaviorTree/Decorators/BTDecorator_ConeCheck.cpp
--- a/Engine/Source/Runtime/AIModule/Private/BehaviorTree/Decorators/BTDecorator_ConeCheck.cpp
+++ b/Engine/Source/Runtime/AIModule/Private/BehaviorTree/Decorators/BTDecorator_ConeCheck.cpp
@@ -97,10 +97,12 @@ void UBTDecorator_ConeCheck::OnBlackboardChange(const UBlackboardComponent& Blac
 
 void UBTDecorator_ConeCheck::TickNode(UBehaviorTreeComponent& OwnerComp, uint8* NodeMemory, float DeltaSeconds)
 {
-	const TNodeInstanceMemory* DecoratorMemory = reinterpret_cast<TNodeInstanceMemory*>(NodeMemory);
-
-	if (CalcConditionImpl(OwnerComp, NodeMemory) != DecoratorMemory->bLastRawResult)
+	TNodeInstanceMemory* DecoratorMemory = reinterpret_cast<TNodeInstanceMemory*>(NodeMemory);
+	
+	const bool bResult = CalcConditionImpl(OwnerComp, NodeMemory);
+	if (bResult != DecoratorMemory->bLastRawResult)
 	{
+		DecoratorMemory->bLastRawResult = bResult;
 		OwnerComp.RequestExecution(this);
 	}
 }
diff --git a/Engine/Source/Runtime/AIModule/Private/BehaviorTree/Services/BTService_BlueprintBase.cpp b/Engine/Source/Runtime/AIModule/Private/BehaviorTree/Services/BTService_BlueprintBase.cpp
--- a/Engine/Source/Runtime/AIModule/Private/BehaviorTree/Services/BTService_BlueprintBase.cpp
+++ b/Engine/Source/Runtime/AIModule/Private/BehaviorTree/Services/BTService_BlueprintBase.cpp
@@ -7,10 +7,10 @@
 UBTService_BlueprintBase::UBTService_BlueprintBase(const FObjectInitializer& ObjectInitializer) : Super(ObjectInitializer)
 {
 	UClass* StopAtClass = UBTService_BlueprintBase::StaticClass();
-	ReceiveTickImplementations = FBTNodeBPImplementationHelper::CheckEventImplementationVersion(TEXT("ReceiveTick"), TEXT("ReceiveTickAI"), this, StopAtClass);
-	ReceiveActivationImplementations = FBTNodeBPImplementationHelper::CheckEventImplementationVersion(TEXT("ReceiveActivation"), TEXT("ReceiveActivationAI"), this, StopAtClass);
-	ReceiveDeactivationImplementations = FBTNodeBPImplementationHelper::CheckEventImplementationVersion(TEXT("ReceiveDeactivation"), TEXT("ReceiveDeactivationAI"), this, StopAtClass);
-	ReceiveSearchStartImplementations = FBTNodeBPImplementationHelper::CheckEventImplementationVersion(TEXT("ReceiveSearchStart"), TEXT("ReceiveSearchStartAI"), this, StopAtClass);
+	ReceiveTickImplementations = FBTNodeBPImplementationHelper::CheckEventImplementationVersion(TEXT("ReceiveTick"), TEXT("ReceiveTickAI"), *this, *StopAtClass);
+	ReceiveActivationImplementations = FBTNodeBPImplementationHelper::CheckEventImplementationVersion(TEXT("ReceiveActivation"), TEXT("ReceiveActivationAI"), *this, *StopAtClass);
+	ReceiveDeactivationImplementations = FBTNodeBPImplementationHelper::CheckEventImplementationVersion(TEXT("ReceiveDeactivation"), TEXT("ReceiveDeactivationAI"), *this, *StopAtClass);
+	ReceiveSearchStartImplementations = FBTNodeBPImplementationHelper::CheckEventImplementationVersion(TEXT("ReceiveSearchStart"), TEXT("ReceiveSearchStartAI"), *this, *StopAtClass);
 
 	bNotifyBecomeRelevant = ReceiveActivationImplementations != 0;
 	bNotifyCeaseRelevant = bNotifyBecomeRelevant;
@@ -38,7 +38,7 @@ void UBTService_BlueprintBase::OnBecomeRelevant(UBehaviorTreeComponent& OwnerCom
 	Super::OnBecomeRelevant(OwnerComp, NodeMemory);
 
 	// check flag, it could be used because user wants tick
-	if (AIOwner != nullptr && ReceiveActivationImplementations & FBTNodeBPImplementationHelper::AISpecific)
+	if (AIOwner != nullptr && (ReceiveActivationImplementations & FBTNodeBPImplementationHelper::AISpecific))
 	{
 		ReceiveActivationAI(AIOwner, AIOwner->GetPawn());
 	}
@@ -58,7 +58,7 @@ void UBTService_BlueprintBase::OnCeaseRelevant(UBehaviorTreeComponent& OwnerComp
 		// we can't have those resuming activity when node is/was aborted
 		BlueprintNodeHelpers::AbortLatentActions(OwnerComp, *this);
 
-		if (AIOwner != nullptr && ReceiveActivationImplementations & FBTNodeBPImplementationHelper::AISpecific)
+		if (AIOwner != nullptr && (ReceiveActivationImplementations & FBTNodeBPImplementationHelper::AISpecific))
 		{
 			ReceiveDeactivationAI(AIOwner, AIOwner->GetPawn());
 		}
@@ -83,7 +83,7 @@ void UBTService_BlueprintBase::OnSearchStart(FBehaviorTreeSearchData& SearchData
 
 	if (ReceiveSearchStartImplementations != 0)
 	{
-		if (AIOwner != nullptr && ReceiveSearchStartImplementations & FBTNodeBPImplementationHelper::AISpecific)
+		if (AIOwner != nullptr && (ReceiveSearchStartImplementations & FBTNodeBPImplementationHelper::AISpecific))
 		{
 			ReceiveSearchStartAI(AIOwner, AIOwner->GetPawn());
 		}
@@ -98,7 +98,7 @@ void UBTService_BlueprintBase::TickNode(UBehaviorTreeComponent& OwnerComp, uint8
 {
 	Super::TickNode(OwnerComp, NodeMemory, DeltaSeconds);
 
-	if (AIOwner != nullptr && ReceiveTickImplementations & FBTNodeBPImplementationHelper::AISpecific)
+	if (AIOwner != nullptr && (ReceiveTickImplementations & FBTNodeBPImplementationHelper::AISpecific))
 	{
 		ReceiveTickAI(AIOwner, AIOwner->GetPawn(), DeltaSeconds);
 	}
diff --git a/Engine/Source/Runtime/AIModule/Private/BehaviorTree/Tasks/BTTask_BlueprintBase.cpp b/Engine/Source/Runtime/AIModule/Private/BehaviorTree/Tasks/BTTask_BlueprintBase.cpp
--- a/Engine/Source/Runtime/AIModule/Private/BehaviorTree/Tasks/BTTask_BlueprintBase.cpp
+++ b/Engine/Source/Runtime/AIModule/Private/BehaviorTree/Tasks/BTTask_BlueprintBase.cpp
@@ -7,9 +7,9 @@
 UBTTask_BlueprintBase::UBTTask_BlueprintBase(const FObjectInitializer& ObjectInitializer) : Super(ObjectInitializer)
 {
 	UClass* StopAtClass = UBTTask_BlueprintBase::StaticClass();
-	ReceiveTickImplementations = FBTNodeBPImplementationHelper::CheckEventImplementationVersion(TEXT("ReceiveTick"), TEXT("ReceiveTickAI"), this, StopAtClass);
-	ReceiveExecuteImplementations = FBTNodeBPImplementationHelper::CheckEventImplementationVersion(TEXT("ReceiveExecute"), TEXT("ReceiveExecuteAI"), this, StopAtClass);
-	ReceiveAbortImplementations = FBTNodeBPImplementationHelper::CheckEventImplementationVersion(TEXT("ReceiveAbort"), TEXT("ReceiveAbortAI"), this, StopAtClass);
+	ReceiveTickImplementations = FBTNodeBPImplementationHelper::CheckEventImplementationVersion(TEXT("ReceiveTick"), TEXT("ReceiveTickAI"), *this, *StopAtClass);
+	ReceiveExecuteImplementations = FBTNodeBPImplementationHelper::CheckEventImplementationVersion(TEXT("ReceiveExecute"), TEXT("ReceiveExecuteAI"), *this, *StopAtClass);
+	ReceiveAbortImplementations = FBTNodeBPImplementationHelper::CheckEventImplementationVersion(TEXT("ReceiveAbort"), TEXT("ReceiveAbortAI"), *this, *StopAtClass);
 
 	bNotifyTick = ReceiveTickImplementations != FBTNodeBPImplementationHelper::NoImplementation;
 	bNotifyTaskFinished = true;
@@ -41,7 +41,7 @@ EBTNodeResult::Type UBTTask_BlueprintBase::ExecuteTask(UBehaviorTreeComponent& O
 	{
 		bStoreFinishResult = true;
 
-		if (AIOwner != nullptr && ReceiveExecuteImplementations & FBTNodeBPImplementationHelper::AISpecific)
+		if (AIOwner != nullptr && (ReceiveExecuteImplementations & FBTNodeBPImplementationHelper::AISpecific))
 		{
 			ReceiveExecuteAI(AIOwner, AIOwner->GetPawn());
 		}
@@ -69,7 +69,7 @@ EBTNodeResult::Type UBTTask_BlueprintBase::AbortTask(UBehaviorTreeComponent& Own
 	{
 		bStoreFinishResult = true;
 
-		if (AIOwner != nullptr && ReceiveAbortImplementations & FBTNodeBPImplementationHelper::AISpecific)
+		if (AIOwner != nullptr && (ReceiveAbortImplementations & FBTNodeBPImplementationHelper::AISpecific))
 		{
 			ReceiveAbortAI(AIOwner, AIOwner->GetPawn());
 		}
@@ -86,7 +86,7 @@ EBTNodeResult::Type UBTTask_BlueprintBase::AbortTask(UBehaviorTreeComponent& Own
 
 void UBTTask_BlueprintBase::TickTask(UBehaviorTreeComponent& OwnerComp, uint8* NodeMemory, float DeltaSeconds)
 {
-	if (AIOwner != nullptr && ReceiveTickImplementations & FBTNodeBPImplementationHelper::AISpecific)
+	if (AIOwner != nullptr && (ReceiveTickImplementations & FBTNodeBPImplementationHelper::AISpecific))
 	{
 		ReceiveTickAI(AIOwner, AIOwner->GetPawn(), DeltaSeconds);
 	}
diff --git a/Engine/Source/Runtime/AIModule/Private/BehaviorTree/Tasks/BTTask_RunBehavior.cpp b/Engine/Source/Runtime/AIModule/Private/BehaviorTree/Tasks/BTTask_RunBehavior.cpp
--- a/Engine/Source/Runtime/AIModule/Private/BehaviorTree/Tasks/BTTask_RunBehavior.cpp
+++ b/Engine/Source/Runtime/AIModule/Private/BehaviorTree/Tasks/BTTask_RunBehavior.cpp
@@ -4,6 +4,24 @@
 #include "BehaviorTree/Tasks/BTTask_RunBehavior.h"
 #include "BehaviorTree/BehaviorTreeManager.h"
 
+// Each injected decorator must be instanced before using in parent tree, because it's not part of it
+// This requires special memory handling for decorator which were not meant to be instanced
+// and need their memory block for storage.
+//
+// Memory allocated for BTTask_Behavior consist of header and N pairs for each injected decorator:
+//
+// ---- task memory start, address stored in Task.MemoryOffset
+// - FBTInstancedNodeMemory: task's header with index of first instanced decorator
+// ---- injected decorators memory start
+// - FBTInstancedNodeMemory: index of injected decorator 1 (special node memory)
+// - memory_required_by_decorator: memory block requested by decorator1, address stored in Decorator1.MemoryOffset, 
+// - (remaining pairs)
+//
+//
+// Each decorator must provide required memory size with instance index header.
+// Decorator memory will be accessible in continuous block, MemoryOffset points at data inside task's memory block
+// Decorators which should be instanced will access their own object through index in special node memory headers
+
 UBTTask_RunBehavior::UBTTask_RunBehavior(const FObjectInitializer& ObjectInitializer) : Super(ObjectInitializer)
 {
 	NodeName = "Run Behavior";
@@ -29,8 +47,11 @@ uint16 UBTTask_RunBehavior::GetInstanceMemorySize() const
 	int32 MemorySize = 0;
 	if (BehaviorAsset)
 	{
+		// all decorators need to return memory requirements as instanced nodes (include header sizes)
+
 		TArray<uint16> MemoryOffsets;
-		UBehaviorTreeManager::InitializeMemoryHelper(BehaviorAsset->RootDecorators, MemoryOffsets, MemorySize);
+		const bool bForceInstancing = true;
+		UBehaviorTreeManager::InitializeMemoryHelper(BehaviorAsset->RootDecorators, MemoryOffsets, MemorySize, bForceInstancing);
 	}
 
 	// memory block header: index of first node instance in owner's array
@@ -50,7 +71,9 @@ void UBTTask_RunBehavior::InjectNodes(UBehaviorTreeComponent& OwnerComp, uint8*
 	FBTInstancedNodeMemory* NodeMemoryHeader = (FBTInstancedNodeMemory*)NodeMemory;
 	int32 FirstNodeIdx = NodeMemoryHeader->NodeIdx;
 
-	uint8* InjectedMemoryBase = NodeMemory + sizeof(FBTInstancedNodeMemory);
+	const uint16 HeaderSize = sizeof(FBTInstancedNodeMemory);
+	const uint16 InjectedMemoryOffset = GetMemoryOffset() + HeaderSize;
+	uint8* InjectedMemoryBase = NodeMemory + HeaderSize;
 
 	// initialize on first access
 	if (!OwnerComp.NodeInstances.IsValidIndex(InstancedIndex))
@@ -58,13 +81,8 @@ void UBTTask_RunBehavior::InjectNodes(UBehaviorTreeComponent& OwnerComp, uint8*
 		TArray<uint16> MemoryOffsets;
 		int32 MemorySize = 0;
 
-		UBehaviorTreeManager::InitializeMemoryHelper(BehaviorAsset->RootDecorators, MemoryOffsets, MemorySize);
-		const int32 AlignedInstanceMemorySize = UBehaviorTreeManager::GetAlignedDataSize(sizeof(FBTInstancedNodeMemory));
-
-		// prepare dummy memory block for nodes that won't require instancing and offset it by special data size
-		// InitializeInSubtree will read it through GetSpecialNodeMemory function, which moves pointer back 
-		FBTInstancedNodeMemory DummyMemory;
-		uint8* RawDummyMemory = ((uint8*)&DummyMemory) + AlignedInstanceMemorySize;
+		const bool bForceInstancing = true;
+		UBehaviorTreeManager::InitializeMemoryHelper(BehaviorAsset->RootDecorators, MemoryOffsets, MemorySize, bForceInstancing);
 
 		// newly created nodes = full init
 		EBTMemoryInit::Type InitType = EBTMemoryInit::Initialize;
@@ -85,7 +103,7 @@ void UBTTask_RunBehavior::InjectNodes(UBehaviorTreeComponent& OwnerComp, uint8*
 			else
 			{
 				DecoratorOb->ForceInstancing(true);
-				DecoratorOb->InitializeInSubtree(OwnerComp, RawDummyMemory, InstancedIndex, InitType);
+				DecoratorOb->InitializeInSubtree(OwnerComp, InjectedMemoryBase + MemoryOffsets[Idx], InstancedIndex, InitType);
 				DecoratorOb->ForceInstancing(false);
 
 				UBTDecorator* InstancedOb = Cast<UBTDecorator>(OwnerComp.NodeInstances.Last());
@@ -98,7 +116,7 @@ void UBTTask_RunBehavior::InjectNodes(UBehaviorTreeComponent& OwnerComp, uint8*
 		for (int32 Idx = 0; Idx < NumInjectedDecorators; Idx++)
 		{
 			UBTDecorator* InstancedOb = Cast<UBTDecorator>(OwnerComp.NodeInstances[FirstNodeIdx + Idx]);
-			InstancedOb->InitializeNode(GetParentNode(), NewExecutionIdx, GetMemoryOffset() + MemoryOffsets[Idx], GetTreeDepth() - 1);
+			InstancedOb->InitializeNode(GetParentNode(), NewExecutionIdx, InjectedMemoryOffset + MemoryOffsets[Idx], GetTreeDepth() - 1);
 			InstancedOb->MarkInjectedNode();
 
 			NewExecutionIdx++;
@@ -122,7 +140,7 @@ void UBTTask_RunBehavior::InjectNodes(UBehaviorTreeComponent& OwnerComp, uint8*
 			}
 			else
 			{
-				uint8* InjectedNodeMemory = InjectedMemoryBase + (InstancedOb->GetMemoryOffset() - GetMemoryOffset());
+				uint8* InjectedNodeMemory = NodeMemory + (InstancedOb->GetMemoryOffset() - GetMemoryOffset());
 				InstancedOb->InitializeMemory(OwnerComp, InjectedNodeMemory, InitType);
 			}
 		}
@@ -247,7 +265,6 @@ void UBTTask_RunBehavior::CleanupMemory(UBehaviorTreeComponent& OwnerComp, uint8
 	if (GetParentNode() && BehaviorAsset)
 	{
 		FBTInstancedNodeMemory* NodeMemoryHeader = (FBTInstancedNodeMemory*)NodeMemory;
-		uint8* InjectedMemoryBase = NodeMemory + sizeof(FBTInstancedNodeMemory);
 
 		const int32 NumInjectedDecorators = BehaviorAsset->RootDecorators.Num();
 		for (int32 Idx = 0; Idx < NumInjectedDecorators; Idx++)
@@ -258,7 +275,7 @@ void UBTTask_RunBehavior::CleanupMemory(UBehaviorTreeComponent& OwnerComp, uint8
 
 			if (InstancedOb)
 			{
-				uint8* InjectedNodeMemory = InjectedMemoryBase + (InstancedOb->GetMemoryOffset() - GetMemoryOffset());
+				uint8* InjectedNodeMemory = NodeMemory + (InstancedOb->GetMemoryOffset() - GetMemoryOffset());
 				InstancedOb->CleanupMemory(OwnerComp, InjectedNodeMemory, CleanupType);
 			}
 		}
diff --git a/Engine/Source/Runtime/AIModule/Private/GameplayDebugger/AIModuleBasicGameplayDebuggerObject.cpp b/Engine/Source/Runtime/AIModule/Private/GameplayDebugger/AIModuleBasicGameplayDebuggerObject.cpp
new file mode 100644
--- /dev/null
+++ b/Engine/Source/Runtime/AIModule/Private/GameplayDebugger/AIModuleBasicGameplayDebuggerObject.cpp
@@ -0,0 +1,544 @@
+// Copyright 1998-2015 Epic Games, Inc. All Rights Reserved.
+#include "AIModulePrivate.h"
+#include "Misc/CoreMisc.h"
+#include "Net/UnrealNetwork.h"
+#include "DebugRenderSceneProxy.h"
+#include "GameplayTasksComponent.h"
+#include "AIModuleBasicGameplayDebuggerObject.h"
+#if WITH_EDITOR
+#include "Editor/EditorEngine.h"
+#endif
+
+DEFINE_LOG_CATEGORY(LogAIModuleGameplayDebugger);
+
+UAIModuleBasicGameplayDebuggerObject::UAIModuleBasicGameplayDebuggerObject(const FObjectInitializer& ObjectInitializer) : Super(ObjectInitializer)
+{
+}
+
+void UAIModuleBasicGameplayDebuggerObject::GetLifetimeReplicatedProps(TArray< FLifetimeProperty > & OutLifetimeProps) const
+{
+	Super::GetLifetimeReplicatedProps(OutLifetimeProps);
+#if ENABLED_GAMEPLAY_DEBUGGER
+	DOREPLIFETIME(UAIModuleBasicGameplayDebuggerObject, ControllerName);
+	DOREPLIFETIME(UAIModuleBasicGameplayDebuggerObject, PawnName);
+	DOREPLIFETIME(UAIModuleBasicGameplayDebuggerObject, DebugIcon);
+	DOREPLIFETIME(UAIModuleBasicGameplayDebuggerObject, PawnClass);
+
+	DOREPLIFETIME(UAIModuleBasicGameplayDebuggerObject, bIsUsingCharacter);
+	DOREPLIFETIME(UAIModuleBasicGameplayDebuggerObject, MovementBaseInfo);
+	DOREPLIFETIME(UAIModuleBasicGameplayDebuggerObject, MovementModeInfo);
+	DOREPLIFETIME(UAIModuleBasicGameplayDebuggerObject, MontageInfo);
+
+	DOREPLIFETIME(UAIModuleBasicGameplayDebuggerObject, bIsUsingPathFollowing);
+	DOREPLIFETIME(UAIModuleBasicGameplayDebuggerObject, PathFollowingInfo);
+	DOREPLIFETIME(UAIModuleBasicGameplayDebuggerObject, NavDataInfo);
+
+	DOREPLIFETIME(UAIModuleBasicGameplayDebuggerObject, PathPoints);
+	DOREPLIFETIME(UAIModuleBasicGameplayDebuggerObject, NextPathPointIndex);
+
+	DOREPLIFETIME(UAIModuleBasicGameplayDebuggerObject, bIsUsingBehaviorTree);
+#endif
+}
+
+void UAIModuleBasicGameplayDebuggerObject::CollectDataToReplicate(APlayerController* MyPC, AActor *SelectedActor)
+{
+	Super::CollectDataToReplicate(MyPC, SelectedActor);
+
+#if ENABLED_GAMEPLAY_DEBUGGER
+	APawn* MyPawn = Cast<APawn>(SelectedActor);
+	PawnName = MyPawn ? MyPawn->GetHumanReadableName() : TEXT("{red}No selected pawn.");
+	PawnClass = MyPawn ? MyPawn->GetClass()->GetName() : TEXT("");
+	AAIController* MyController = MyPawn ? Cast<AAIController>(MyPawn->Controller) : nullptr;
+
+	bIsUsingCharacter = MyPawn && MyPawn->IsA(ACharacter::StaticClass());
+
+	if (MyPawn && MyPawn->IsPendingKill() == false)
+	{
+		CollectBasicMovementData(MyPawn);
+		CollectBasicPathData(MyPawn);
+		CollectBasicBehaviorData(MyPawn);
+		CollectPathData(MyPawn);
+		CollectBasicAnimationData(MyPawn);
+	}
+
+	if (MyController)
+	{
+		if (MyController->IsPendingKill() == false)
+		{
+			ControllerName = MyController->GetName();
+			DebugIcon = MyController->GetDebugIcon();
+		}
+		else
+		{
+			ControllerName = TEXT("Controller PENDING KILL");
+		}
+	}
+	else
+	{
+		ControllerName = TEXT("No Controller");
+	}
+#endif //ENABLED_GAMEPLAY_DEBUGGER
+}
+
+void UAIModuleBasicGameplayDebuggerObject::CollectBasicMovementData(APawn* MyPawn)
+{
+#if ENABLED_GAMEPLAY_DEBUGGER
+	UCharacterMovementComponent* CharMovement = Cast<UCharacterMovementComponent>(MyPawn->GetMovementComponent());
+	if (CharMovement)
+	{
+		UPrimitiveComponent* FloorComponent = MyPawn->GetMovementBase();
+		AActor* FloorActor = FloorComponent ? FloorComponent->GetOwner() : nullptr;
+		MovementBaseInfo = FloorComponent ? FString::Printf(TEXT("%s.%s"), *GetNameSafe(FloorActor), *FloorComponent->GetName()) : TEXT("None");
+
+		MovementModeInfo = CharMovement->GetMovementName();
+	}
+	else
+	{
+		MovementBaseInfo = TEXT("");
+		MovementModeInfo = TEXT("");
+	}
+#endif //ENABLED_GAMEPLAY_DEBUGGER
+}
+
+void UAIModuleBasicGameplayDebuggerObject::CollectBasicPathData(APawn* MyPawn)
+{
+#if ENABLED_GAMEPLAY_DEBUGGER
+	UNavigationSystem* NavSys = UNavigationSystem::GetCurrent(MyPawn->GetWorld());
+	AAIController* MyAIController = Cast<AAIController>(MyPawn->GetController());
+
+	const ANavigationData* NavData = MyAIController && NavSys ? NavSys->GetNavDataForProps(MyAIController->GetNavAgentPropertiesRef()) : nullptr;
+	if (NavData)
+	{
+		NavDataInfo = NavData->GetConfig().Name.ToString();
+	}
+
+	UPathFollowingComponent* PFC = MyAIController ? MyAIController->GetPathFollowingComponent() : nullptr;
+	bIsUsingPathFollowing = (PFC != nullptr);
+
+	if (PFC)
+	{
+		TArray<FString> Tokens;
+		TArray<EPathFollowingDebugTokens::Type> Flags;
+		if (PFC)
+		{
+			PFC->GetDebugStringTokens(Tokens, Flags);
+		}
+
+		PathFollowingInfo = FString();
+		for (int32 Idx = 0; Idx < Tokens.Num(); Idx++)
+		{
+			switch (Flags[Idx])
+			{
+			case EPathFollowingDebugTokens::Description:
+				PathFollowingInfo += Tokens[Idx];
+				break;
+
+			case EPathFollowingDebugTokens::ParamName:
+				PathFollowingInfo += TEXT(", {yellow}");
+				PathFollowingInfo += Tokens[Idx];
+				PathFollowingInfo += TEXT(":");
+				break;
+
+			case EPathFollowingDebugTokens::PassedValue:
+				PathFollowingInfo += TEXT("{yellow}");
+				PathFollowingInfo += Tokens[Idx];
+				break;
+
+			case EPathFollowingDebugTokens::FailedValue:
+				PathFollowingInfo += TEXT("{orange}");
+				PathFollowingInfo += Tokens[Idx];
+				break;
+
+			default:
+				break;
+			}
+		}
+	}
+	else
+	{
+		PathFollowingInfo = TEXT("");
+	}
+#endif //ENABLED_GAMEPLAY_DEBUGGER
+}
+
+void UAIModuleBasicGameplayDebuggerObject::CollectBasicBehaviorData(APawn* MyPawn)
+{
+#if ENABLED_GAMEPLAY_DEBUGGER && ENABLE_VISUAL_LOG
+	AAIController* MyAIController = Cast<AAIController>(MyPawn->GetController());
+	UBehaviorTreeComponent* BTC = MyAIController ? Cast<UBehaviorTreeComponent>(MyAIController->BrainComponent) : nullptr;
+	bIsUsingBehaviorTree = (BTC != nullptr);
+
+	if (BTC)
+	{
+		CurrentAITask = BTC->DescribeActiveTasks();
+		CurrentAIState = BTC->IsRunning() ? TEXT("Running") : BTC->IsPaused() ? TEXT("Paused") : TEXT("Inactive");
+		CurrentAIAssets = BTC->DescribeActiveTrees();
+	}
+	else
+	{
+		CurrentAITask = TEXT("");
+		CurrentAIState = TEXT("");
+		CurrentAIAssets = TEXT("");
+	}
+
+	UGameplayTasksComponent* GTComponent = MyPawn->FindComponentByClass<UGameplayTasksComponent>();
+	if (GTComponent)
+	{
+		GameplayTasksState = FString::Printf(TEXT("Ticking Tasks: %s\nTask Queue: %s"), *GTComponent->GetTickingTasksDescription(), *GTComponent->GetTasksPriorityQueueDescription());
+	}
+	else
+	{
+		GameplayTasksState = TEXT("");
+	}
+#endif //ENABLED_GAMEPLAY_DEBUGGER
+}
+
+void UAIModuleBasicGameplayDebuggerObject::CollectBasicAnimationData(APawn* MyPawn)
+{
+#if ENABLED_GAMEPLAY_DEBUGGER
+	ACharacter* MyChar = Cast<ACharacter>(MyPawn);
+	MontageInfo = MyChar ? GetNameSafe(MyChar->GetCurrentMontage()) : TEXT("");
+#endif //ENABLED_GAMEPLAY_DEBUGGER
+}
+
+void UAIModuleBasicGameplayDebuggerObject::DrawCollectedData(APlayerController* MyPC, AActor *SelectedActor)
+{
+	Super::DrawCollectedData(MyPC, SelectedActor);
+
+#if ENABLED_GAMEPLAY_DEBUGGER
+	UWorld *ActiveWorld = GetWorld();
+	if (ActiveWorld)
+	{
+		const FVector SelectedActorLoc = SelectedActor ? SelectedActor->GetActorLocation() + FVector(0, 0, SelectedActor->GetSimpleCollisionHalfHeight()) : DebugTools::InvalidLocation;
+		for (FConstPawnIterator Iterator = ActiveWorld->GetPawnIterator(); Iterator; ++Iterator)
+		{
+			APawn* CurrentPawn = Cast<APawn>(*Iterator);
+			AAIController* CurrentController = CurrentPawn != nullptr ? Cast<AAIController>(CurrentPawn->GetController()) : nullptr;
+			if (CurrentController)
+			{
+				if (CurrentController->IsPendingKill() == false)
+				{
+					DebugIcon = CurrentController->GetDebugIcon();
+					if (DebugIcon.Len())
+					{
+						UTexture2D* RegularIcon = (UTexture2D*)StaticLoadObject(UTexture2D::StaticClass(), NULL, *DebugIcon, NULL, LOAD_NoWarn | LOAD_Quiet, NULL);
+						if (RegularIcon)
+						{
+							FCanvasIcon Icon = UCanvas::MakeIcon(RegularIcon);
+							if (Icon.Texture)
+							{
+								const FVector ScreenLoc = UGameplayDebuggerHelper::ProjectLocation(DefaultContext, CurrentPawn->GetActorLocation() + FVector(0, 0, CurrentPawn->GetSimpleCollisionHalfHeight()));
+								const FVector SelectedActorScreenLoc = SelectedActor != nullptr ? UGameplayDebuggerHelper::ProjectLocation(DefaultContext, SelectedActorLoc) : FVector::ZeroVector;
+								const float DesiredIconSize = SelectedActorScreenLoc == ScreenLoc ? 32.0f : 16.f;
+								DrawIcon(OverHeadContext, FColor::White, Icon, ScreenLoc.X, ScreenLoc.Y - DesiredIconSize, DesiredIconSize / Icon.Texture->GetSurfaceWidth());
+							}
+						}
+					}
+				}
+			}
+		}
+
+		if (SelectedActorLoc != DebugTools::InvalidLocation)
+		{
+			DrawOverHeadInformation(UGameplayDebuggerHelper::GetDefaultContext(), UGameplayDebuggerHelper::GetOverHeadContext(), SelectedActorLoc);
+			DrawBasicData(UGameplayDebuggerHelper::GetDefaultContext(), UGameplayDebuggerHelper::GetOverHeadContext());
+		}
+	}
+#endif
+}
+
+void UAIModuleBasicGameplayDebuggerObject::DrawBasicData(UGameplayDebuggerBaseObject::FPrintContext& DefaultContext, UGameplayDebuggerBaseObject::FPrintContext& OverHeadContext)
+{
+#if ENABLED_GAMEPLAY_DEBUGGER
+	UFont* OldFont = DefaultContext.Font.Get();
+	DefaultContext.Font = GEngine->GetMediumFont();
+	PrintString(DefaultContext, FString::Printf(TEXT("Controller Name: {yellow}%s\n"), *ControllerName));
+	DefaultContext.Font = OldFont;
+
+	PrintString(DefaultContext, FString::Printf(TEXT("Pawn Name: {yellow}%s{white}, Pawn Class: {yellow}%s\n"), *PawnName, *PawnClass));
+
+	// movement
+	if (bIsUsingCharacter)
+	{
+		PrintString(DefaultContext, FString::Printf(TEXT("Movement Mode: {yellow}%s{white}, Base: {yellow}%s\n"), *MovementModeInfo, *MovementBaseInfo));
+		PrintString(DefaultContext, FString::Printf(TEXT("NavData: {yellow}%s{white}, Path following: {yellow}%s\n"), *NavDataInfo, *PathFollowingInfo));
+	}
+
+	// logic
+	if (bIsUsingBehaviorTree)
+	{
+		PrintString(DefaultContext, FString::Printf(TEXT("Behavior: {yellow}%s{white}, Tree: {yellow}%s\n"), *CurrentAIState, *CurrentAIAssets));
+		PrintString(DefaultContext, FString::Printf(TEXT("Active task: {yellow}%s\n"), *CurrentAITask));
+	}
+
+	// animation
+	if (bIsUsingCharacter)
+	{
+		PrintString(DefaultContext, FString::Printf(TEXT("Montage: {yellow}%s\n"), *MontageInfo));
+	}
+
+	// putting gameplay tasks' stuff last since it can expand heavily
+	PrintString(DefaultContext, FString::Printf(TEXT("GameplayTasks:\n{yellow}%s\n"), *GameplayTasksState));
+#endif //ENABLED_GAMEPLAY_DEBUGGER
+}
+
+void UAIModuleBasicGameplayDebuggerObject::DrawOverHeadInformation(UGameplayDebuggerBaseObject::FPrintContext& DefaultContext, UGameplayDebuggerBaseObject::FPrintContext& OverHeadContext, FVector SelectedActorLoc)
+{
+#if ENABLED_GAMEPLAY_DEBUGGER
+	if (OverHeadContext.Canvas->SceneView == NULL || OverHeadContext.Canvas->SceneView->ViewFrustum.IntersectBox(SelectedActorLoc, FVector::ZeroVector) == false)
+	{
+		return;
+	}
+
+	const FVector ScreenLoc = OverHeadContext.Canvas->Project(SelectedActorLoc);
+	static const FVector2D FontScale(1.f, 1.f);
+	UFont* Font = GEngine->GetSmallFont();
+
+	float TextXL = 0.f;
+	float YL = 0.f;
+	FString ObjectName = FString::Printf(TEXT("{yellow}%s {white}(%s)"), *ControllerName, *PawnName);
+	CalulateStringSize(OverHeadContext, OverHeadContext.Font.Get(), ObjectName, TextXL, YL);
+
+	bool bDrawFullOverHead = SelectedActorLoc != DebugTools::InvalidLocation;
+	if (bDrawFullOverHead)
+	{
+		OverHeadContext.DefaultX -= (0.5f*TextXL*FontScale.X);
+		OverHeadContext.DefaultY -= (1.2f*YL*FontScale.Y);
+		OverHeadContext.CursorX = OverHeadContext.DefaultX;
+		OverHeadContext.CursorY = OverHeadContext.DefaultY;
+	}
+
+	if (bDrawFullOverHead)
+	{
+		OverHeadContext.FontRenderInfo.bEnableShadow = bDrawFullOverHead;
+		PrintString(OverHeadContext, bDrawFullOverHead ? FColor::White : FColor(255, 255, 255, 128), FString::Printf(TEXT("%s\n"), *ObjectName));
+		OverHeadContext.FontRenderInfo.bEnableShadow = false;
+	}
+
+	FEngineShowFlags EngineShowFlags = OverHeadContext.Canvas->SceneView && OverHeadContext.Canvas->SceneView->Family ? OverHeadContext.Canvas->SceneView->Family->EngineShowFlags : FEngineShowFlags(GIsEditor ? EShowFlagInitMode::ESFIM_Editor : EShowFlagInitMode::ESFIM_Game);
+	if (EngineShowFlags.DebugAI || EngineShowFlags.Game)
+	{
+		DrawPath();
+	}
+#endif //ENABLED_GAMEPLAY_DEBUGGER
+}
+
+void UAIModuleBasicGameplayDebuggerObject::DrawPath()
+{
+#if ENABLED_GAMEPLAY_DEBUGGER
+	static const FColor Grey(100, 100, 100);
+	static const FColor PathColor(192, 192, 192);
+
+	const int32 NumPathVerts = PathPoints.Num();
+
+	for (int32 VertIdx = 0; VertIdx < NumPathVerts - 1; ++VertIdx)
+	{
+		FVector const VertLoc =PathPoints[VertIdx] + NavigationDebugDrawing::PathOffset;
+		DrawDebugSolidBox(GetWorld(), VertLoc, NavigationDebugDrawing::PathNodeBoxExtent, VertIdx < int32(NextPathPointIndex) ? Grey : PathColor, false);
+
+		// draw line to next loc
+		FVector const NextVertLoc = PathPoints[VertIdx + 1] + NavigationDebugDrawing::PathOffset;
+		DrawDebugLine(GetWorld(), VertLoc, NextVertLoc, VertIdx < int32(NextPathPointIndex) ? Grey : PathColor, false
+			, -1.f, 0
+			, NavigationDebugDrawing::PathLineThickness);
+	}
+
+	// draw last vert
+	if (NumPathVerts > 0)
+	{
+		DrawDebugBox(GetWorld(), PathPoints[NumPathVerts - 1] + NavigationDebugDrawing::PathOffset, FVector(15.f), Grey, false);
+	}
+
+#endif //ENABLED_GAMEPLAY_DEBUGGER
+}
+
+void UAIModuleBasicGameplayDebuggerObject::CollectPathData(APawn* MyPawn)
+{
+#if ENABLED_GAMEPLAY_DEBUGGER && ENABLE_VISUAL_LOG
+	bool bRefreshRendering = false;
+	UWorld* World = MyPawn->GetWorld();
+
+	if (AAIController *MyController = Cast<AAIController>(MyPawn->Controller))
+	{
+		if (MyController->PathFollowingComponent && MyController->PathFollowingComponent->GetPath().IsValid())
+		{
+			NextPathPointIndex = MyController->PathFollowingComponent->GetNextPathIndex();
+
+			const FNavPathSharedPtr& NewPath = MyController->PathFollowingComponent->GetPath();
+			if (CurrentPath.HasSameObject(NewPath.Get()) == false || NewPath->GetTimeStamp() > LastStoredPathTimeStamp)
+			{
+				LastStoredPathTimeStamp = NewPath->GetTimeStamp();
+
+				FVisualLogEntry Snapshot;
+				NewPath->DescribeSelfToVisLog(&Snapshot);
+				PathCorridorPolygons.Reset(Snapshot.ElementsToDraw.Num());
+				FPathCorridorPolygons Polygon;
+				for (FVisualLogShapeElement& CurrentShape : Snapshot.ElementsToDraw)
+				{
+					if (CurrentShape.GetType() == EVisualLoggerShapeElement::Polygon)
+					{
+						Polygon.Color = CurrentShape.GetFColor();
+						Polygon.Points.Reset();
+						Polygon.Points.Append(CurrentShape.Points);
+						PathCorridorPolygons.Add(Polygon);
+					}
+				}
+
+				PathPoints.Reset();
+				for (int32 Index = 0; Index < NewPath->GetPathPoints().Num(); ++Index)
+				{
+					PathPoints.Add(NewPath->GetPathPoints()[Index].Location);
+				}
+				CurrentPath = NewPath;
+
+				if (PathCorridorPolygons.Num() && World && World->GetNetMode() != NM_Client)
+				{
+					PathCorridorData.Reset();
+					TArray<uint8> HelpBuffer;
+					FMemoryWriter ArWriter(HelpBuffer);
+					int32 NumPolygons = PathCorridorPolygons.Num();
+					ArWriter << NumPolygons;
+					for (const FPathCorridorPolygons& CurrentPoly : PathCorridorPolygons)
+					{
+						FColor Color = CurrentPoly.Color;
+						ArWriter << Color;
+						int32 NumVerts = CurrentPoly.Points.Num();
+						ArWriter << NumVerts;
+						for (auto Vert : CurrentPoly.Points)
+						{
+							ArWriter << Vert;
+						}
+					}
+
+					RequestDataReplication(TEXT("PathCorridor"), HelpBuffer, EGameplayDebuggerReplicate::WithCompressed);
+				}
+				bRefreshRendering = true;
+			}
+		}
+		else
+		{
+			NextPathPointIndex = INDEX_NONE;
+			CurrentPath = NULL;
+			PathPoints.Reset();
+			PathCorridorPolygons.Reset();
+		}
+	}
+
+	if (bRefreshRendering && World && World->GetNetMode() != NM_DedicatedServer)
+	{
+		MarkRenderStateDirty();
+	}
+#endif //ENABLED_GAMEPLAY_DEBUGGER
+}
+
+void UAIModuleBasicGameplayDebuggerObject::OnDataReplicationComplited(const TArray<uint8>& ReplicatedData)
+{
+#if ENABLED_GAMEPLAY_DEBUGGER
+	FMemoryReader ArReader(ReplicatedData);
+	int32 NumPolygons = 0;
+	ArReader << NumPolygons;
+	PathCorridorPolygons.Reset(NumPolygons);
+	for (int32 PolyIndex = 0; PolyIndex < NumPolygons; ++PolyIndex)
+	{
+		FPathCorridorPolygons Polygon;
+		FColor Color;
+		ArReader << Polygon.Color;
+
+		int32 NumVerts = 0;
+		ArReader << NumVerts;
+		for (int32 VertIndex = 0; VertIndex < NumVerts; ++VertIndex)
+		{
+			FVector CurrentVertex;
+			ArReader << CurrentVertex;
+			Polygon.Points.Add(CurrentVertex);
+		}
+		if (NumVerts > 2)
+		{
+			PathCorridorPolygons.Add(Polygon);
+		}
+	}
+	MarkRenderStateDirty();
+#endif
+}
+
+FDebugRenderSceneProxy* UAIModuleBasicGameplayDebuggerObject::CreateSceneProxy(const UPrimitiveComponent* InComponent, UWorld* World, AActor* SelectedActor)
+{
+#if ENABLED_GAMEPLAY_DEBUGGER
+	class FPathDebugRenderSceneProxy : public FDebugRenderSceneProxy
+	{
+	public:
+		FPathDebugRenderSceneProxy(const UPrimitiveComponent* InPrimitiveCComponent, const FString &InViewFlagName)
+			: FDebugRenderSceneProxy(InPrimitiveCComponent)
+		{
+			DrawType = FDebugRenderSceneProxy::SolidAndWireMeshes;
+			DrawAlpha = 90;
+			ViewFlagName = InViewFlagName;
+			ViewFlagIndex = uint32(FEngineShowFlags::FindIndexByName(*ViewFlagName));
+		}
+
+
+		FPrimitiveViewRelevance GetViewRelevance(const FSceneView* View) const override
+		{
+			FPrimitiveViewRelevance Result;
+			Result.bDrawRelevance = View->Family->EngineShowFlags.GetSingleFlag(ViewFlagIndex);// IsShown(View);
+			Result.bDynamicRelevance = true;
+			// ideally the TranslucencyRelevance should be filled out by the material, here we do it conservative
+			Result.bSeparateTranslucencyRelevance = Result.bNormalTranslucencyRelevance = IsShown(View);
+			return Result;
+		}
+
+	};
+
+	const bool bDrawFullData = SelectedActor != NULL;
+	if (bDrawFullData /*&& ShouldReplicateData(EAIDebugDrawDataView::Basic)*/)
+	{
+		TArray<FDebugRenderSceneProxy::FMesh> Meshes;
+		TArray<FDebugRenderSceneProxy::FDebugLine> Lines;
+		for (FPathCorridorPolygons& CurrentPoly : PathCorridorPolygons)
+		{
+
+			if (CurrentPoly.Points.Num() > 2)
+			{
+				int32 LastIndex = 0;
+				FVector FirstVertex = CurrentPoly.Points[0];
+				FDebugRenderSceneProxy::FMesh TestMesh;
+				TestMesh.Color = CurrentPoly.Color;// FColor::Green;
+				for (int32 Index = 1; Index < CurrentPoly.Points.Num() - 1; Index += 1)
+				{
+					TestMesh.Vertices.Add(FDynamicMeshVertex(FirstVertex));
+					TestMesh.Vertices.Add(FDynamicMeshVertex(CurrentPoly.Points[Index + 0]));
+					TestMesh.Vertices.Add(FDynamicMeshVertex(CurrentPoly.Points[Index + 1]));
+					TestMesh.Indices.Add(LastIndex++);
+					TestMesh.Indices.Add(LastIndex++);
+					TestMesh.Indices.Add(LastIndex++);
+				}
+				Meshes.Add(TestMesh);
+			}
+
+			for (int32 VIdx = 0; VIdx < CurrentPoly.Points.Num(); VIdx++)
+			{
+				Lines.Add(FDebugRenderSceneProxy::FDebugLine(
+					CurrentPoly.Points[VIdx],
+					CurrentPoly.Points[(VIdx + 1) % CurrentPoly.Points.Num()],
+					CurrentPoly.Color,
+					2)
+					);
+			}
+		}
+
+		FString ViewFlagName = TEXT("Game");
+#if WITH_EDITOR
+		UEditorEngine* EEngine = Cast<UEditorEngine>(GEngine);
+		if (EEngine && EEngine->bIsSimulatingInEditor)
+		{
+			ViewFlagName = TEXT("DebugAI");
+		}
+#endif
+		FPathDebugRenderSceneProxy *DebugSceneProxy = new FPathDebugRenderSceneProxy(InComponent, ViewFlagName);
+		DebugSceneProxy->Lines = Lines;
+		DebugSceneProxy->Meshes = Meshes;
+		return DebugSceneProxy;
+	}
+#endif
+	return nullptr;
+}
\ No newline at end of file
diff --git a/Engine/Source/Runtime/AIModule/Private/GameplayDebugger/AIModuleBasicGameplayDebuggerObject.h b/Engine/Source/Runtime/AIModule/Private/GameplayDebugger/AIModuleBasicGameplayDebuggerObject.h
new file mode 100644
--- /dev/null
+++ b/Engine/Source/Runtime/AIModule/Private/GameplayDebugger/AIModuleBasicGameplayDebuggerObject.h
@@ -0,0 +1,102 @@
+// Copyright 1998-2015 Epic Games, Inc. All Rights Reserved.
+
+#pragma once
+#include "Debug/GameplayDebuggerBaseObject.h"
+#include "AIModuleBasicGameplayDebuggerObject.generated.h"
+
+UCLASS()
+class UAIModuleBasicGameplayDebuggerObject : public UGameplayDebuggerBaseObject
+{
+	GENERATED_UCLASS_BODY()
+
+protected:
+	struct FPathCorridorPolygons
+	{
+		TArray<FVector> Points;
+		FColor Color;
+	};
+	
+	UPROPERTY(Replicated)
+	FString ControllerName;
+
+	UPROPERTY(Replicated)
+	FString PawnName;
+
+	UPROPERTY(Replicated)
+	FString PawnClass;
+
+	UPROPERTY(Replicated)
+	FString DebugIcon;
+
+	UPROPERTY(Replicated)
+	FString MovementBaseInfo;
+
+	UPROPERTY(Replicated)
+	FString MovementModeInfo;
+
+	UPROPERTY(Replicated)
+	FString PathFollowingInfo;
+
+	UPROPERTY(Replicated)
+	FString CurrentAITask;
+
+	UPROPERTY(Replicated)
+	FString CurrentAIState;
+
+	UPROPERTY(Replicated)
+	FString CurrentAIAssets;
+
+	UPROPERTY(Replicated)
+	FString GameplayTasksState;
+
+	UPROPERTY(Replicated)
+	FString NavDataInfo;
+
+	UPROPERTY(Replicated)
+	FString MontageInfo;
+
+	UPROPERTY(Replicated)
+	uint32 bIsUsingPathFollowing : 1;
+
+	UPROPERTY(Replicated)
+	uint32 bIsUsingCharacter : 1;
+
+	UPROPERTY(Replicated)
+	uint32 bIsUsingBehaviorTree : 1;
+
+	/** Begin path replication data */
+	UPROPERTY(Replicated)
+	int32 NextPathPointIndex;
+
+	UPROPERTY(Replicated)
+	TArray<FVector> PathPoints;
+
+	TArray<uint8> PathCorridorData;
+
+	TArray<FPathCorridorPolygons> PathCorridorPolygons;
+	/** End path replication data*/
+
+
+public:
+	virtual void CollectDataToReplicate(APlayerController* MyPC, AActor *SelectedActor) override;
+
+	virtual void DrawCollectedData(APlayerController* MyPC, AActor *SelectedActor) override;
+
+	FDebugRenderSceneProxy* CreateSceneProxy(const UPrimitiveComponent* InComponent, UWorld* World, AActor* SelectedActor) override;
+
+protected:
+	FNavPathWeakPtr CurrentPath;
+	float LastStoredPathTimeStamp;
+
+	void CollectBasicMovementData(APawn* MyPawn);
+	void CollectBasicPathData(APawn* MyPawn);
+	void CollectBasicBehaviorData(APawn* MyPawn);
+	void CollectPathData(APawn* MyPawn);
+	void CollectBasicAnimationData(APawn* MyPawn);
+
+	void DrawBasicData(UGameplayDebuggerBaseObject::FPrintContext& DefaultContext, UGameplayDebuggerBaseObject::FPrintContext& OverHeadContext);
+	void DrawOverHeadInformation(UGameplayDebuggerBaseObject::FPrintContext& DefaultContext, UGameplayDebuggerBaseObject::FPrintContext& OverHeadContext, FVector SelectedActorLoc);
+	void DrawPath();
+
+	void OnDataReplicationComplited(const TArray<uint8>& ReplicatedData) override;
+};
diff --git a/Engine/Source/Runtime/AIModule/Private/GameplayDebugger/BTGameplayDebuggerObject.cpp b/Engine/Source/Runtime/AIModule/Private/GameplayDebugger/BTGameplayDebuggerObject.cpp
new file mode 100644
--- /dev/null
+++ b/Engine/Source/Runtime/AIModule/Private/GameplayDebugger/BTGameplayDebuggerObject.cpp
@@ -0,0 +1,74 @@
+// Copyright 1998-2015 Epic Games, Inc. All Rights Reserved.
+#include "AIModulePrivate.h"
+#include "Misc/CoreMisc.h"
+#include "Net/UnrealNetwork.h"
+#include "DebugRenderSceneProxy.h"
+#include "GameplayTasksComponent.h"
+#include "BTGameplayDebuggerObject.h"
+
+UBTGameplayDebuggerObject::UBTGameplayDebuggerObject(const FObjectInitializer& ObjectInitializer) : Super(ObjectInitializer)
+{
+}
+
+void UBTGameplayDebuggerObject::GetLifetimeReplicatedProps(TArray< FLifetimeProperty > & OutLifetimeProps) const
+{
+	Super::GetLifetimeReplicatedProps(OutLifetimeProps);
+
+#if ENABLED_GAMEPLAY_DEBUGGER
+	DOREPLIFETIME(UBTGameplayDebuggerObject, BrainComponentName);
+	DOREPLIFETIME(UBTGameplayDebuggerObject, BrainComponentString);
+#endif
+}
+
+void UBTGameplayDebuggerObject::CollectDataToReplicate(APlayerController* MyPC, AActor *SelectedActor)
+{
+	Super::CollectDataToReplicate(MyPC, SelectedActor);
+
+#if ENABLED_GAMEPLAY_DEBUGGER
+	BrainComponentName = TEXT("");
+	BrainComponentString = TEXT("");
+
+	APawn* MyPawn = Cast<APawn>(SelectedActor);
+	AAIController* MyController = MyPawn ? Cast<AAIController>(MyPawn->Controller) : nullptr;
+	if (MyController != NULL && MyController->BrainComponent != NULL && MyController->IsPendingKillPending() == false)
+	{
+		BrainComponentName = MyController->BrainComponent != NULL ? MyController->BrainComponent->GetName() : TEXT("");
+		BrainComponentString = MyController->BrainComponent != NULL ? MyController->BrainComponent->GetDebugInfoString() : TEXT("");
+
+		BlackboardString = MyController->BrainComponent->GetBlackboardComponent() ? MyController->BrainComponent->GetBlackboardComponent()->GetDebugInfoString(EBlackboardDescription::KeyWithValue) : TEXT("");
+
+		UWorld* ActiveWorld = GetWorld();
+		if (ActiveWorld && ActiveWorld->GetNetMode() != NM_Standalone)
+		{
+			TArray<uint8> UncompressedBuffer;
+			FMemoryWriter ArWriter(UncompressedBuffer);
+
+			ArWriter << BlackboardString;
+
+			RequestDataReplication(TEXT("BlackboardString"), UncompressedBuffer, EGameplayDebuggerReplicate::WithoutCompression);
+		}
+	}
+#endif //ENABLED_GAMEPLAY_DEBUGGER
+}
+
+void UBTGameplayDebuggerObject::OnDataReplicationComplited(const TArray<uint8>& ReplicatedData)
+{
+#if ENABLED_GAMEPLAY_DEBUGGER
+	FMemoryReader ArReader(ReplicatedData);
+	ArReader << BlackboardString;
+#endif
+}
+
+void UBTGameplayDebuggerObject::DrawCollectedData(APlayerController* MyPC, AActor* SelectedActor)
+{
+	Super::DrawCollectedData(MyPC, SelectedActor);
+
+#if ENABLED_GAMEPLAY_DEBUGGER
+	if (SelectedActor)
+	{
+		PrintString(DefaultContext, FString::Printf(TEXT("Brain Component: {yellow}%s\n"), *BrainComponentName));
+		PrintString(DefaultContext, BrainComponentString);
+		PrintString(DefaultContext, FColor::White, BlackboardString, 600.0f, GetDefaultContext().DefaultY);
+	}
+#endif
+}
diff --git a/Engine/Source/Runtime/AIModule/Private/GameplayDebugger/BTGameplayDebuggerObject.h b/Engine/Source/Runtime/AIModule/Private/GameplayDebugger/BTGameplayDebuggerObject.h
new file mode 100644
--- /dev/null
+++ b/Engine/Source/Runtime/AIModule/Private/GameplayDebugger/BTGameplayDebuggerObject.h
@@ -0,0 +1,24 @@
+// Copyright 1998-2015 Epic Games, Inc. All Rights Reserved.
+
+#pragma once
+#include "Debug/GameplayDebuggerBaseObject.h"
+#include "BTGameplayDebuggerObject.generated.h"
+
+UCLASS()
+class UBTGameplayDebuggerObject : public UGameplayDebuggerBaseObject
+{
+	GENERATED_UCLASS_BODY()
+
+	UPROPERTY(Replicated)
+	FString BrainComponentName;
+
+	UPROPERTY(Replicated)
+	FString BrainComponentString;
+
+	FString BlackboardString;
+
+	virtual FString GetCategoryName() const override { return TEXT("BehaviorTree"); };
+	virtual void CollectDataToReplicate(APlayerController* MyPC, AActor *SelectedActor) override;
+	virtual void DrawCollectedData(APlayerController* MyPC, AActor *SelectedActor) override;
+	virtual void OnDataReplicationComplited(const TArray<uint8>& ReplicatedData) override;
+};
diff --git a/Engine/Source/Runtime/AIModule/Private/GameplayDebugger/EQSGameplayDebuggerObject.cpp b/Engine/Source/Runtime/AIModule/Private/GameplayDebugger/EQSGameplayDebuggerObject.cpp
new file mode 100644
--- /dev/null
+++ b/Engine/Source/Runtime/AIModule/Private/GameplayDebugger/EQSGameplayDebuggerObject.cpp
@@ -0,0 +1,400 @@
+// Copyright 1998-2015 Epic Games, Inc. All Rights Reserved.
+#include "AIModulePrivate.h"
+#include "Misc/CoreMisc.h"
+#include "Net/UnrealNetwork.h"
+#include "EnvironmentQuery/EnvQueryTypes.h"
+#include "EnvironmentQuery/EnvQueryManager.h"
+#include "EnvironmentQuery/EQSRenderingComponent.h"
+#include "EnvironmentQuery/EnvQueryTest.h"
+#include "DebugRenderSceneProxy.h"
+#include "EQSGameplayDebuggerObject.h"
+#if WITH_EDITOR
+#	include "Editor/EditorEngine.h"
+#endif
+
+#if ENABLED_GAMEPLAY_DEBUGGER && USE_EQS_DEBUGGER
+	static TAutoConsoleVariable<int32> CVarEQSDetailsOnHUD(
+		TEXT("ai.gd.EQSOnHUD"),
+		1,
+		TEXT("Enable or disable EQS details table on screen.\n")
+		TEXT(" 0: Disable details\n")
+		TEXT(" 1: Enable details (default)"),
+		ECVF_Cheat);
+#endif
+
+UEQSGameplayDebuggerObject::UEQSGameplayDebuggerObject(const FObjectInitializer& ObjectInitializer) : Super(ObjectInitializer), MaxEQSQueries(5)
+{
+#if ENABLED_GAMEPLAY_DEBUGGER && USE_EQS_DEBUGGER
+	ItemDescriptionWidth = 312.f; // 200.0f;
+	ItemScoreWidth = 50.0f;
+	TestScoreWidth = 100.0f;
+	LastDataCaptureTime = 0;
+#endif
+}
+
+void UEQSGameplayDebuggerObject::OnCycleDetailsView()
+{
+#if ENABLED_GAMEPLAY_DEBUGGER && USE_EQS_DEBUGGER
+	if (++CurrentEQSIndex >= EQSLocalData.Num())
+	{
+		CurrentEQSIndex = 0;
+	}
+
+	MarkRenderStateDirty();
+#endif //ENABLED_GAMEPLAY_DEBUGGER
+}
+
+const FEnvQueryResult* UEQSGameplayDebuggerObject::GetQueryResult() const
+{
+#if ENABLED_GAMEPLAY_DEBUGGER && USE_EQS_DEBUGGER
+	return CachedQueryInstance.Get();
+#endif
+	return nullptr;
+}
+
+const FEnvQueryInstance* UEQSGameplayDebuggerObject::GetQueryInstance() const
+{
+#if ENABLED_GAMEPLAY_DEBUGGER && USE_EQS_DEBUGGER
+	return CachedQueryInstance.Get();
+#endif
+	return nullptr;
+}
+
+void UEQSGameplayDebuggerObject::OnDataReplicationComplited(const TArray<uint8>& ReplicatedData)
+{
+#if ENABLED_GAMEPLAY_DEBUGGER && USE_EQS_DEBUGGER
+	// decode scoring data
+	FMemoryReader ArReader(ReplicatedData);
+	ArReader << EQSLocalData;
+
+	MarkRenderStateDirty();
+#endif
+}
+
+void UEQSGameplayDebuggerObject::CollectDataToReplicate(APlayerController* MyPC, AActor *SelectedActor)
+{
+	Super::CollectDataToReplicate(MyPC, SelectedActor);
+#if	ENABLED_GAMEPLAY_DEBUGGER && USE_EQS_DEBUGGER
+	if (GetWorld()->TimeSeconds - LastDataCaptureTime < 2)
+	{
+		return;
+	}
+
+	UWorld* World = GetWorld();
+	UEnvQueryManager* QueryManager = World ? UEnvQueryManager::GetCurrent(World) : NULL;
+	const AActor* Owner = SelectedActor;
+
+	if (QueryManager == NULL || Owner == NULL)
+	{
+		return;
+	}
+
+	LastDataCaptureTime = GetWorld()->TimeSeconds;
+
+	auto AllQueries = QueryManager->GetDebugger().GetAllQueriesForOwner(Owner);
+	const class APawn* OwnerAsPawn = Cast<class APawn>(Owner);
+	if (OwnerAsPawn != NULL && OwnerAsPawn->GetController())
+	{
+		const auto& AllControllerQueries = QueryManager->GetDebugger().GetAllQueriesForOwner(OwnerAsPawn->GetController());
+		AllQueries.Append(AllControllerQueries);
+	}
+	struct FEnvQueryInfoSort
+	{
+		FORCEINLINE bool operator()(const FEQSDebugger::FEnvQueryInfo& A, const FEQSDebugger::FEnvQueryInfo& B) const
+		{
+			return (A.Timestamp < B.Timestamp);
+		}
+	};
+	TArray<FEQSDebugger::FEnvQueryInfo> QueriesToSort = AllQueries;
+	QueriesToSort.Sort(FEnvQueryInfoSort()); //sort queries by timestamp
+	QueriesToSort.SetNum(FMath::Min<int32>(MaxEQSQueries, AllQueries.Num()));
+
+	for (int32 Index = AllQueries.Num() - 1; Index >= 0; --Index)
+	{
+		auto &CurrentQuery = AllQueries[Index];
+		if (QueriesToSort.Find(CurrentQuery) == INDEX_NONE)
+		{
+			AllQueries.RemoveAt(Index);
+		}
+	}
+
+
+	EQSLocalData.Reset();
+	for (int32 Index = 0; Index < FMath::Min<int32>(MaxEQSQueries, AllQueries.Num()); ++Index)
+	{
+		EQSDebug::FQueryData* CurrentLocalData = NULL;
+		CachedQueryInstance = AllQueries[Index].Instance;
+		const float CachedTimestamp = AllQueries[Index].Timestamp;
+
+		if (!CurrentLocalData)
+		{
+			EQSLocalData.AddZeroed();
+			CurrentLocalData = &EQSLocalData[EQSLocalData.Num() - 1];
+		}
+
+		if (CachedQueryInstance.IsValid())
+		{
+			UEnvQueryDebugHelpers::QueryToDebugData(*CachedQueryInstance, *CurrentLocalData);
+		}
+		CurrentLocalData->Timestamp = AllQueries[Index].Timestamp;
+	}
+
+	TArray<uint8> UncompressedBuffer;
+	FMemoryWriter ArWriter(UncompressedBuffer);
+	ArWriter << EQSLocalData;
+
+	RequestDataReplication(TEXT("EQSData"), UncompressedBuffer);
+#endif
+}
+
+void UEQSGameplayDebuggerObject::DrawEQSItemDetails(int32 ItemIdx, AActor* SelectedActor)
+{
+#if ENABLED_GAMEPLAY_DEBUGGER && USE_EQS_DEBUGGER
+	const float PosY = DefaultContext.CursorY + 1.0f;
+	float PosX = DefaultContext.CursorX;
+
+	const int32 EQSIndex = EQSLocalData.Num() > 0 ? FMath::Clamp(CurrentEQSIndex, 0, EQSLocalData.Num() - 1) : INDEX_NONE;
+	auto& CurrentLocalData = EQSLocalData[EQSIndex];
+	const EQSDebug::FItemData& ItemData = CurrentLocalData.Items[ItemIdx];
+
+	PrintString(DefaultContext, FColor::White, ItemData.Desc, PosX, PosY);
+	PosX += ItemDescriptionWidth;
+
+	FString ScoreDesc = FString::Printf(TEXT("%.2f"), ItemData.TotalScore);
+	PrintString(DefaultContext, FColor::Yellow, ScoreDesc, PosX, PosY);
+	PosX += ItemScoreWidth;
+
+	FCanvasTileItem ActiveTileItem(FVector2D(0, PosY + 15.0f), GWhiteTexture, FVector2D(0, 2.0f), FLinearColor::Yellow);
+	FCanvasTileItem BackTileItem(FVector2D(0, PosY + 15.0f), GWhiteTexture, FVector2D(0, 2.0f), FLinearColor(0.1f, 0.1f, 0.1f));
+	const float BarWidth = 80.0f;
+
+	const int32 NumTests = ItemData.TestScores.Num();
+
+	float TotalWeightedScore = 0.0f;
+	for (int32 Idx = 0; Idx < NumTests; Idx++)
+	{
+		TotalWeightedScore += ItemData.TestScores[Idx];
+	}
+
+	for (int32 Idx = 0; Idx < NumTests; Idx++)
+	{
+		const float ScoreW = ItemData.TestScores[Idx];
+		const float ScoreN = ItemData.TestValues[Idx];
+		FString DescScoreW = FString::Printf(TEXT("%.2f"), ScoreW);
+		FString DescScoreN = (ScoreN == UEnvQueryTypes::SkippedItemValue) ? TEXT("SKIP") : FString::Printf(TEXT("%.2f"), ScoreN);
+		FString TestDesc = DescScoreW + FString(" {LightBlue}") + DescScoreN;
+
+		float Pct = (TotalWeightedScore > KINDA_SMALL_NUMBER) ? (ScoreW / TotalWeightedScore) : 0.0f;
+		ActiveTileItem.Position.X = PosX;
+		ActiveTileItem.Size.X = BarWidth * Pct;
+		BackTileItem.Position.X = PosX + ActiveTileItem.Size.X;
+		BackTileItem.Size.X = FMath::Max(BarWidth * (1.0f - Pct), 0.0f);
+
+		DrawItem(DefaultContext, ActiveTileItem, ActiveTileItem.Position.X, ActiveTileItem.Position.Y);
+		DrawItem(DefaultContext, BackTileItem, BackTileItem.Position.X, BackTileItem.Position.Y);
+
+		PrintString(DefaultContext, FColor::Green, TestDesc, PosX, PosY);
+		PosX += TestScoreWidth;
+	}
+#endif //USE_EQS_DEBUGGER
+}
+
+void UEQSGameplayDebuggerObject::DrawCollectedData(APlayerController* MyPC, AActor *SelectedActor)
+{
+	Super::DrawCollectedData(MyPC, SelectedActor);
+
+#if ENABLED_GAMEPLAY_DEBUGGER && USE_EQS_DEBUGGER
+	PrintString(DefaultContext, TEXT("\n{green}EQS {white}[Use + key to switch query]\n"));
+
+	if (EQSLocalData.Num() == 0)
+	{
+		return;
+	}
+
+	const int32 EQSIndex = EQSLocalData.Num() > 0 ? FMath::Clamp(CurrentEQSIndex, 0, EQSLocalData.Num() - 1) : INDEX_NONE;
+	if (!EQSLocalData.IsValidIndex(EQSIndex))
+	{
+		return;
+	}
+
+	{
+		int32 Index = 0;
+		PrintString(DefaultContext, TEXT("{white}Queries: "));
+		for (auto CurrentQuery : EQSLocalData)
+		{
+			if (EQSIndex == Index)
+			{
+				PrintString(DefaultContext, FString::Printf(TEXT("{green}%s, "), *CurrentQuery.Name));
+			}
+			else
+			{
+				PrintString(DefaultContext, FString::Printf(TEXT("{yellow}%s, "), *CurrentQuery.Name));
+			}
+			Index++;
+		}
+		PrintString(DefaultContext, TEXT("\n"));
+	}
+
+	auto& CurrentLocalData = EQSLocalData[EQSIndex];
+
+	/** find and draw item selection */
+	int32 BestItemIndex = INDEX_NONE;
+	{
+		FVector CamLocation;
+		FVector FireDir;
+		if (!MyPC->GetSpectatorPawn())
+		{
+			FRotator CamRotation;
+			MyPC->GetPlayerViewPoint(CamLocation, CamRotation);
+			FireDir = CamRotation.Vector();
+		}
+		else
+		{
+			FireDir = DefaultContext.Canvas->SceneView->GetViewDirection();
+			CamLocation = DefaultContext.Canvas->SceneView->ViewMatrices.ViewOrigin;
+		}
+
+		float bestAim = 0;
+		for (int32 Index = 0; Index < CurrentLocalData.RenderDebugHelpers.Num(); ++Index)
+		{
+			auto& CurrentItem = CurrentLocalData.RenderDebugHelpers[Index];
+
+			const FVector AimDir = CurrentItem.Location - CamLocation;
+			float FireDist = AimDir.SizeSquared();
+
+			FireDist = FMath::Sqrt(FireDist);
+			float newAim = FireDir | AimDir;
+			newAim = newAim / FireDist;
+			if (newAim > bestAim)
+			{
+				BestItemIndex = Index;
+				bestAim = newAim;
+			}
+		}
+
+		if (BestItemIndex != INDEX_NONE)
+		{
+			DrawDebugSphere(GetWorld(), CurrentLocalData.RenderDebugHelpers[BestItemIndex].Location, CurrentLocalData.RenderDebugHelpers[BestItemIndex].Radius, 8, FColor::Red, false);
+			int32 FailedTestIndex = CurrentLocalData.RenderDebugHelpers[BestItemIndex].FailedTestIndex;
+			if (FailedTestIndex != INDEX_NONE)
+			{
+				PrintString(DefaultContext, FString::Printf(TEXT("{red}Selected item failed with test %d: {yellow}%s {LightBlue}(%s)\n")
+					, FailedTestIndex
+					, *CurrentLocalData.Tests[FailedTestIndex].ShortName
+					, *CurrentLocalData.Tests[FailedTestIndex].Detailed
+					));
+				PrintString(DefaultContext, FString::Printf(TEXT("{white}'%s' with score %3.3f\n\n"), *CurrentLocalData.RenderDebugHelpers[BestItemIndex].AdditionalInformation, CurrentLocalData.RenderDebugHelpers[BestItemIndex].FailedScore));
+			}
+		}
+	}
+
+	PrintString(DefaultContext, FString::Printf(TEXT("{white}Timestamp: {yellow}%.3f (%.2fs ago)\n")
+		, CurrentLocalData.Timestamp, GetWorld()->GetTimeSeconds() - CurrentLocalData.Timestamp
+		));
+	PrintString(DefaultContext, FString::Printf(TEXT("{white}Query ID: {yellow}%d\n")
+		, CurrentLocalData.Id
+		));
+
+	PrintString(DefaultContext, FString::Printf(TEXT("{white}Query contains %d options: "), CurrentLocalData.Options.Num()));
+	for (int32 OptionIndex = 0; OptionIndex < CurrentLocalData.Options.Num(); ++OptionIndex)
+	{
+		if (OptionIndex == CurrentLocalData.UsedOption)
+		{
+			PrintString(DefaultContext, FString::Printf(TEXT("{green}%s, "), *CurrentLocalData.Options[OptionIndex]));
+		}
+		else
+		{
+			PrintString(DefaultContext, FString::Printf(TEXT("{yellow}%s, "), *CurrentLocalData.Options[OptionIndex]));
+		}
+	}
+	PrintString(DefaultContext, TEXT("\n"));
+
+	const float RowHeight = 20.0f;
+	const int32 NumTests = CurrentLocalData.Tests.Num();
+	const IConsoleVariable* CVarEQSOnHUD = IConsoleManager::Get().FindConsoleVariable(TEXT("ai.gd.EQSOnHUD"));
+	const bool bEQSOnHUD = !CVarEQSOnHUD || CVarEQSOnHUD->GetInt();
+	if (CurrentLocalData.NumValidItems > 0 && bEQSOnHUD)
+	{
+		// draw test weights for best X items
+		const int32 NumItems = CurrentLocalData.Items.Num();
+		
+		FCanvasTileItem TileItem(FVector2D(0, 0), GWhiteTexture, FVector2D(DefaultContext.Canvas->SizeX, RowHeight), FLinearColor::Black);
+		FLinearColor ColorOdd(0, 0, 0, 0.6f);
+		FLinearColor ColorEven(0, 0, 0.4f, 0.4f);
+		TileItem.BlendMode = SE_BLEND_Translucent;
+
+		// table header		
+		{
+			DefaultContext.CursorY += RowHeight;
+			const float HeaderY = DefaultContext.CursorY + 3.0f;
+			TileItem.SetColor(ColorOdd);
+			DrawItem(DefaultContext, TileItem, 0, DefaultContext.CursorY);
+
+			float HeaderX = DefaultContext.CursorX;
+			PrintString(DefaultContext, FColor::Yellow, FString::Printf(TEXT("Num items: %d"), CurrentLocalData.NumValidItems), HeaderX, HeaderY);
+			HeaderX += ItemDescriptionWidth;
+
+			PrintString(DefaultContext, FColor::White, TEXT("Score"), HeaderX, HeaderY);
+			HeaderX += ItemScoreWidth;
+
+			for (int32 TestIdx = 0; TestIdx < NumTests; TestIdx++)
+			{
+				PrintString(DefaultContext, FColor::White, FString::Printf(TEXT("Test %d"), TestIdx), HeaderX, HeaderY);
+				HeaderX += TestScoreWidth;
+			}
+
+			DefaultContext.CursorY += RowHeight;
+		}
+
+		// valid items
+		for (int32 Idx = 0; Idx < NumItems; Idx++)
+		{
+			TileItem.SetColor((Idx % 2) ? ColorOdd : ColorEven);
+			DrawItem(DefaultContext, TileItem, 0, DefaultContext.CursorY);
+
+			DrawEQSItemDetails(Idx, SelectedActor);
+			DefaultContext.CursorY += RowHeight;
+		}
+		DefaultContext.CursorY += RowHeight;
+	}
+
+	// test description
+	PrintString(DefaultContext, TEXT("All tests from used option:\n"));
+	for (int32 TestIdx = 0; TestIdx < NumTests; TestIdx++)
+	{
+		FString TestDesc = FString::Printf(TEXT("{white}Test %d = {yellow}%s {LightBlue}(%s)\n"), TestIdx,
+			*CurrentLocalData.Tests[TestIdx].ShortName,
+			*CurrentLocalData.Tests[TestIdx].Detailed);
+
+		PrintString(DefaultContext, TestDesc);
+	}
+
+#endif //USE_EQS_DEBUGGER
+}
+
+class FDebugRenderSceneProxy* UEQSGameplayDebuggerObject::CreateSceneProxy(const UPrimitiveComponent* InComponent, UWorld* World, AActor* SelectedActor)
+{
+#if ENABLED_GAMEPLAY_DEBUGGER && USE_EQS_DEBUGGER
+	if (/*IsClientEQSSceneProxyEnabled() &&*/ SelectedActor != NULL)
+	{
+		const int32 EQSIndex = EQSLocalData.Num() > 0 ? FMath::Clamp(CurrentEQSIndex, 0, EQSLocalData.Num() - 1) : INDEX_NONE;
+		if (EQSLocalData.IsValidIndex(EQSIndex))
+		{
+			auto& CurrentLocalData = EQSLocalData[EQSIndex];
+
+			FString ViewFlagName = TEXT("Game");
+#if WITH_EDITOR
+			UEditorEngine* EEngine = Cast<UEditorEngine>(GEngine);
+			if (EEngine && EEngine->bIsSimulatingInEditor)
+			{
+				ViewFlagName = TEXT("DebugAI");
+			}
+#endif
+			return new FEQSSceneProxy(InComponent, ViewFlagName, CurrentLocalData.SolidSpheres, CurrentLocalData.Texts);
+		}
+	}
+#endif // USE_EQS_DEBUGGER
+
+	return nullptr;
+}
+
diff --git a/Engine/Source/Runtime/AIModule/Private/GameplayDebugger/EQSGameplayDebuggerObject.h b/Engine/Source/Runtime/AIModule/Private/GameplayDebugger/EQSGameplayDebuggerObject.h
new file mode 100644
--- /dev/null
+++ b/Engine/Source/Runtime/AIModule/Private/GameplayDebugger/EQSGameplayDebuggerObject.h
@@ -0,0 +1,48 @@
+// Copyright 1998-2015 Epic Games, Inc. All Rights Reserved.
+
+#pragma once
+#include "Debug/GameplayDebuggerBaseObject.h"
+#include "EQSGameplayDebuggerObject.generated.h"
+
+UCLASS()
+class UEQSGameplayDebuggerObject : public UGameplayDebuggerBaseObject, public IEQSQueryResultSourceInterface
+{
+	GENERATED_UCLASS_BODY()
+
+	virtual FString GetCategoryName() const override { return TEXT("EQS"); };
+	virtual void CollectDataToReplicate(APlayerController* MyPC, AActor *SelectedActor) override;
+	virtual void DrawCollectedData(APlayerController* MyPC, AActor *SelectedActor) override;
+	virtual class FDebugRenderSceneProxy* CreateSceneProxy(const UPrimitiveComponent* InComponent, UWorld* World, AActor* SelectedActor) override;
+
+protected:
+	/** local EQS debug data, decoded from EQSRepData blob */
+#if  ENABLED_GAMEPLAY_DEBUGGER && USE_EQS_DEBUGGER
+	TArray<EQSDebug::FQueryData> EQSLocalData;
+#endif
+
+	UFUNCTION()
+	virtual void OnCycleDetailsView();
+
+	void DrawEQSItemDetails(int32 ItemIdx, AActor* SelectedActor);
+
+	// IEQSQueryResultSourceInterface start
+	virtual const struct FEnvQueryResult* GetQueryResult() const override;
+	virtual const struct FEnvQueryInstance* GetQueryInstance() const  override;
+
+	virtual bool GetShouldDebugDrawLabels() const override { return bDrawEQSLabels; }
+	virtual bool GetShouldDrawFailedItems() const override{ return bDrawEQSFailedItems; }
+	// IEQSQueryResultSourceInterface end
+
+	void OnDataReplicationComplited(const TArray<uint8>& ReplicatedData) override;
+
+	uint32 bDrawEQSLabels : 1;
+	uint32 bDrawEQSFailedItems : 1;
+	int32 MaxEQSQueries;
+	int32 CurrentEQSIndex;
+	TSharedPtr<FEnvQueryInstance> CachedQueryInstance;
+private:
+	float ItemDescriptionWidth;
+	float ItemScoreWidth;
+	float TestScoreWidth;
+	float LastDataCaptureTime;
+};
diff --git a/Engine/Source/Runtime/AIModule/Private/GameplayDebugger/NavMeshGameplayDebuggerObject.cpp b/Engine/Source/Runtime/AIModule/Private/GameplayDebugger/NavMeshGameplayDebuggerObject.cpp
new file mode 100644
--- /dev/null
+++ b/Engine/Source/Runtime/AIModule/Private/GameplayDebugger/NavMeshGameplayDebuggerObject.cpp
@@ -0,0 +1,448 @@
+// Copyright 1998-2015 Epic Games, Inc. All Rights Reserved.
+#include "AIModulePrivate.h"
+#include "Misc/CoreMisc.h"
+#include "Net/UnrealNetwork.h"
+#include "NavMeshRenderingHelpers.h"
+#include "DebugRenderSceneProxy.h"
+#include "NavMeshGameplayDebuggerObject.h"
+
+
+UNavMeshGameplayDebuggerObject::UNavMeshGameplayDebuggerObject(const FObjectInitializer& ObjectInitializer) : Super(ObjectInitializer)
+{
+	LastDataCaptureTime = 0;
+}
+
+void UNavMeshGameplayDebuggerObject::CollectDataToReplicate(APlayerController* MyPC, AActor *SelectedActor)
+{
+	Super::CollectDataToReplicate(MyPC, SelectedActor);
+
+#if ENABLED_GAMEPLAY_DEBUGGER
+	UWorld* ActiveWorld = GetWorld();
+	if (ActiveWorld && ActiveWorld->TimeSeconds - LastDataCaptureTime > 5)
+	{
+		ENetMode NetMode = ActiveWorld->GetNetMode();
+		LastDataCaptureTime = ActiveWorld->TimeSeconds;
+		APawn const *PlayerPawn = MyPC ? MyPC->GetPawnOrSpectator() : nullptr;
+		APawn const *SelectedActorAsPawn = Cast<APawn>(SelectedActor);
+
+		APawn const *DestPawn = SelectedActorAsPawn ? SelectedActorAsPawn : (SelectedActor ? Cast<APawn>(SelectedActor) : PlayerPawn);
+		if (DestPawn)
+		{
+			ServerCollectNavmeshData(DestPawn);
+		}
+	}
+#endif //ENABLED_GAMEPLAY_DEBUGGER
+}
+
+void UNavMeshGameplayDebuggerObject::DrawCollectedData(APlayerController* MyPC, AActor *SelectedActor)
+{
+	Super::DrawCollectedData(MyPC, SelectedActor);
+
+#if ENABLED_GAMEPLAY_DEBUGGER
+	if (LocalReplicationData.RepDataSize != 0 && LocalReplicationData.RepData.Num() != LocalReplicationData.RepDataSize)
+	{
+		const int32 RepPct = 100 * LocalReplicationData.RepData.Num() / LocalReplicationData.RepDataSize;
+		PrintString(DefaultContext, FString::Printf(TEXT("NavMesh Replicating: {red}%d%%\n"), RepPct));
+		return;
+	}
+#endif
+}
+
+bool UNavMeshGameplayDebuggerObject::ServerDiscardNavmeshData_Validate()
+{
+#if ENABLED_GAMEPLAY_DEBUGGER
+	return true;
+#endif
+	return false;
+}
+
+namespace NavMeshGameplayDebugger
+{
+	struct FShortVector
+	{
+		int16 X;
+		int16 Y;
+		int16 Z;
+
+		FShortVector()
+			:X(0), Y(0), Z(0)
+		{
+
+		}
+
+		FShortVector(const FVector& Vec)
+			: X(Vec.X)
+			, Y(Vec.Y)
+			, Z(Vec.Z)
+		{
+
+		}
+		FShortVector& operator=(const FVector& R)
+		{
+			X = R.X;
+			Y = R.Y;
+			Z = R.Z;
+			return *this;
+		}
+
+		FVector ToVector() const
+		{
+			return FVector(X, Y, Z);
+		}
+	};
+
+	struct FOffMeshLinkFlags
+	{
+		uint8 Radius : 6;
+		uint8 Direction : 1;
+		uint8 ValidEnds : 1;
+	};
+	struct FOffMeshLink
+	{
+		FShortVector Left;
+		FShortVector Right;
+		FColor Color;
+		union
+		{
+			FOffMeshLinkFlags PackedFlags;
+			uint8 ByteFlags;
+		};
+	};
+
+	struct FAreaPolys
+	{
+		TArray<int32> Indices;
+		FColor Color;
+	};
+
+	struct FTileData
+	{
+		FVector Location;
+		TArray<FAreaPolys> Areas;
+		TArray<FShortVector> Verts;
+		TArray<FOffMeshLink> Links;
+	};
+}
+
+FArchive& operator<<(FArchive& Ar, NavMeshGameplayDebugger::FShortVector& ShortVector)
+{
+#if ENABLED_GAMEPLAY_DEBUGGER
+	Ar << ShortVector.X;
+	Ar << ShortVector.Y;
+	Ar << ShortVector.Z;
+#endif
+	return Ar;
+}
+
+FArchive& operator<<(FArchive& Ar, NavMeshGameplayDebugger::FOffMeshLink& Data)
+{
+#if ENABLED_GAMEPLAY_DEBUGGER
+	Ar << Data.Left;
+	Ar << Data.Right;
+	Ar << Data.Color.R;
+	Ar << Data.Color.G;
+	Ar << Data.Color.B;
+	Ar << Data.ByteFlags;
+#endif
+	return Ar;
+}
+
+FArchive& operator<<(FArchive& Ar, NavMeshGameplayDebugger::FAreaPolys& Data)
+{
+#if ENABLED_GAMEPLAY_DEBUGGER
+	Ar << Data.Indices;
+	Ar << Data.Color.R;
+	Ar << Data.Color.G;
+	Ar << Data.Color.B;
+	Data.Color.A = 255;
+#endif
+	return Ar;
+}
+
+
+FArchive& operator<<(FArchive& Ar, NavMeshGameplayDebugger::FTileData& Data)
+{
+#if ENABLED_GAMEPLAY_DEBUGGER
+	Ar << Data.Areas;
+	Ar << Data.Location;
+	Ar << Data.Verts;
+	Ar << Data.Links;
+#endif
+	return Ar;
+}
+
+void UNavMeshGameplayDebuggerObject::ServerDiscardNavmeshData_Implementation()
+{
+#if ENABLED_GAMEPLAY_DEBUGGER
+	LocalReplicationData.RepData.Empty();
+#endif
+}
+
+FORCEINLINE bool LineInCorrectDistance(const FVector& PlayerLoc, const FVector& Start, const FVector& End, float CorrectDistance = -1)
+{
+#if ENABLED_GAMEPLAY_DEBUGGER
+	const float MaxDistance = CorrectDistance > 0 ? (CorrectDistance*CorrectDistance) : ARecastNavMesh::GetDrawDistanceSq();
+
+	if ((FVector::DistSquared(Start, PlayerLoc) > MaxDistance || FVector::DistSquared(End, PlayerLoc) > MaxDistance))
+	{
+		return false;
+	}
+#endif
+	return true;
+}
+
+#if WITH_RECAST
+ARecastNavMesh* UNavMeshGameplayDebuggerObject::GetNavData(const AActor *TargetActor)
+{
+#if ENABLED_GAMEPLAY_DEBUGGER
+	UNavigationSystem* NavSys = UNavigationSystem::GetCurrent(GetWorld());
+	if (NavSys == NULL)
+	{
+		return NULL;
+	}
+
+	// Try to get the correct nav-mesh relative to the selected actor.
+	const APawn* TargetPawn = Cast<APawn>(TargetActor);
+	if (TargetPawn != NULL)
+	{
+		const FNavAgentProperties& NavAgentProperties = TargetPawn->GetNavAgentPropertiesRef();
+		return Cast<ARecastNavMesh>(NavSys->GetNavDataForProps(NavAgentProperties));
+	}
+
+	// If it wasn't found, just get the main nav-mesh data.
+	return Cast<ARecastNavMesh>(NavSys->GetMainNavData(FNavigationSystem::DontCreate));
+#endif //ENABLED_GAMEPLAY_DEBUGGER
+
+	return NULL;
+}
+#endif
+
+void UNavMeshGameplayDebuggerObject::ServerCollectNavmeshData(const AActor *SelectedActor)
+{
+#if ENABLED_GAMEPLAY_DEBUGGER && WITH_RECAST
+	UNavigationSystem* NavSys = UNavigationSystem::GetCurrent(GetWorld());
+	ARecastNavMesh* NavData = NavSys ? GetNavData(SelectedActor) : nullptr;
+	if (NavData == NULL)
+	{
+		LocalReplicationData.RepData.Empty();
+		return;
+	}
+
+	const FVector TargetLocation = SelectedActor ? SelectedActor->GetActorLocation() : FVector::ZeroVector;
+
+	TArray<int32> Indices;
+	int32 TileX = 0;
+	int32 TileY = 0;
+	const int32 DeltaX[] = { 0, 1, 1, 0, -1, -1, -1, 0, 1 };
+	const int32 DeltaY[] = { 0, 0, 1, 1, 1, 0, -1, -1, -1 };
+
+	NavData->BeginBatchQuery();
+
+	// add 3x3 neighborhood of target
+	int32 TargetTileX = 0;
+	int32 TargetTileY = 0;
+	NavData->GetNavMeshTileXY(TargetLocation, TargetTileX, TargetTileY);
+	for (int32 i = 0; i < ARRAY_COUNT(DeltaX); i++)
+	{
+		const int32 NeiX = TargetTileX + DeltaX[i];
+		const int32 NeiY = TargetTileY + DeltaY[i];
+		NavData->GetNavMeshTilesAt(NeiX, NeiY, Indices);
+	}
+
+	const FNavDataConfig& NavConfig = NavData->GetConfig();
+	FColor NavMeshColors[RECAST_MAX_AREAS];
+	NavMeshColors[RECAST_DEFAULT_AREA] = NavConfig.Color.DWColor() > 0 ? NavConfig.Color : NavMeshRenderColor_RecastMesh;
+	for (uint8 i = 0; i < RECAST_DEFAULT_AREA; ++i)
+	{
+		NavMeshColors[i] = NavData->GetAreaIDColor(i);
+	}
+
+	TArray<uint8> UncompressedBuffer;
+	FMemoryWriter ArWriter(UncompressedBuffer);
+
+	int32 NumIndices = Indices.Num();
+	ArWriter << NumIndices;
+
+	for (int32 i = 0; i < NumIndices; i++)
+	{
+		FRecastDebugGeometry NavMeshGeometry;
+		NavMeshGeometry.bGatherPolyEdges = false;
+		NavMeshGeometry.bGatherNavMeshEdges = false;
+		NavData->GetDebugGeometry(NavMeshGeometry, Indices[i]);
+
+		NavMeshGameplayDebugger::FTileData TileData;
+		const FBox TileBoundingBox = NavData->GetNavMeshTileBounds(Indices[i]);
+		TileData.Location = TileBoundingBox.GetCenter();
+		for (int32 VertIndex = 0; VertIndex < NavMeshGeometry.MeshVerts.Num(); ++VertIndex)
+		{
+			const NavMeshGameplayDebugger::FShortVector SV = NavMeshGeometry.MeshVerts[VertIndex] - TileData.Location;
+			TileData.Verts.Add(SV);
+		}
+
+		for (int32 iArea = 0; iArea < RECAST_MAX_AREAS; iArea++)
+		{
+			const int32 NumTris = NavMeshGeometry.AreaIndices[iArea].Num();
+			if (NumTris)
+			{
+				NavMeshGameplayDebugger::FAreaPolys AreaPolys;
+				for (int32 AreaIndicesIndex = 0; AreaIndicesIndex < NavMeshGeometry.AreaIndices[iArea].Num(); ++AreaIndicesIndex)
+				{
+					AreaPolys.Indices.Add(NavMeshGeometry.AreaIndices[iArea][AreaIndicesIndex]);
+				}
+				AreaPolys.Color = NavMeshColors[iArea];
+				TileData.Areas.Add(AreaPolys);
+			}
+		}
+
+		TileData.Links.Reserve(NavMeshGeometry.OffMeshLinks.Num());
+		for (int32 iLink = 0; iLink < NavMeshGeometry.OffMeshLinks.Num(); iLink++)
+		{
+			const FRecastDebugGeometry::FOffMeshLink& SrcLink = NavMeshGeometry.OffMeshLinks[iLink];
+			NavMeshGameplayDebugger::FOffMeshLink Link;
+			Link.Left = SrcLink.Left - TileData.Location;
+			Link.Right = SrcLink.Right - TileData.Location;
+			Link.Color = ((SrcLink.Direction && SrcLink.ValidEnds) || (SrcLink.ValidEnds & FRecastDebugGeometry::OMLE_Left)) ? DarkenColor(NavMeshColors[SrcLink.AreaID]) : NavMeshRenderColor_OffMeshConnectionInvalid;
+			Link.PackedFlags.Radius = (int8)SrcLink.Radius;
+			Link.PackedFlags.Direction = SrcLink.Direction;
+			Link.PackedFlags.ValidEnds = SrcLink.ValidEnds;
+			TileData.Links.Add(Link);
+		}
+
+		ArWriter << TileData;
+	}
+	NavData->FinishBatchQuery();
+	const ENetMode NetMode = GetWorld()->GetNetMode();
+	RequestDataReplication(TEXT("NavMesh"), UncompressedBuffer, NetMode != NM_Standalone ? EGameplayDebuggerReplicate::WithCompressed : EGameplayDebuggerReplicate::WithoutCompression);
+#endif
+}
+
+void UNavMeshGameplayDebuggerObject::OnDataReplicationComplited(const TArray<uint8>& ReplicatedData)
+{
+#if ENABLED_GAMEPLAY_DEBUGGER
+	MarkRenderStateDirty();
+#endif
+}
+
+void UNavMeshGameplayDebuggerObject::PrepareNavMeshData(struct FNavMeshSceneProxyData* CurrentData, const UPrimitiveComponent* InComponent, UWorld* World) const
+{
+#if ENABLED_GAMEPLAY_DEBUGGER && WITH_RECAST
+	if (LocalReplicationData.RepDataSize != LocalReplicationData.RepData.Num())
+	{
+		return;
+	}
+
+	if (CurrentData)
+	{
+		CurrentData->Reset();
+		CurrentData->bNeedsNewData = false;
+
+		// read serialized values
+		CurrentData->bEnableDrawing = LocalReplicationData.RepData.Num() > 0;
+		if (!CurrentData->bEnableDrawing)
+		{
+			return;
+		}
+
+		FMemoryReader ArReader(LocalReplicationData.RepData);
+
+		int32 NumTiles = 0;
+		ArReader << NumTiles;
+
+		int32 IndexesOffset = 0;
+		for (int32 iTile = 0; iTile < NumTiles; iTile++)
+		{
+			NavMeshGameplayDebugger::FTileData TileData;
+			ArReader << TileData;
+
+			FVector OffsetLocation = TileData.Location;
+			TArray<FVector> Verts;
+			Verts.Reserve(TileData.Verts.Num());
+			for (int32 VertIndex = 0; VertIndex < TileData.Verts.Num(); ++VertIndex)
+			{
+				const FVector Loc = TileData.Verts[VertIndex].ToVector() + OffsetLocation;
+				Verts.Add(Loc);
+			}
+			CurrentData->Bounds += FBox(Verts);
+
+			for (int32 iArea = 0; iArea < TileData.Areas.Num(); iArea++)
+			{
+				const NavMeshGameplayDebugger::FAreaPolys& SrcArea = TileData.Areas[iArea];
+				FNavMeshSceneProxyData::FDebugMeshData DebugMeshData;
+				DebugMeshData.ClusterColor = SrcArea.Color;
+				DebugMeshData.ClusterColor.A = 128;
+
+				for (int32 iTri = 0; iTri < SrcArea.Indices.Num(); iTri += 3)
+				{
+					const int32 Index0 = SrcArea.Indices[iTri + 0];
+					const int32 Index1 = SrcArea.Indices[iTri + 1];
+					const int32 Index2 = SrcArea.Indices[iTri + 2];
+
+					FVector V0 = Verts[Index0];
+					FVector V1 = Verts[Index1];
+					FVector V2 = Verts[Index2];
+					CurrentData->TileEdgeLines.Add(FDebugRenderSceneProxy::FDebugLine(V0 + CurrentData->NavMeshDrawOffset, V1 + CurrentData->NavMeshDrawOffset, NavMeshRenderColor_Recast_TileEdges));
+					CurrentData->TileEdgeLines.Add(FDebugRenderSceneProxy::FDebugLine(V1 + CurrentData->NavMeshDrawOffset, V2 + CurrentData->NavMeshDrawOffset, NavMeshRenderColor_Recast_TileEdges));
+					CurrentData->TileEdgeLines.Add(FDebugRenderSceneProxy::FDebugLine(V2 + CurrentData->NavMeshDrawOffset, V0 + CurrentData->NavMeshDrawOffset, NavMeshRenderColor_Recast_TileEdges));
+
+					AddTriangleHelper(DebugMeshData, Index0 + IndexesOffset, Index1 + IndexesOffset, Index2 + IndexesOffset);
+				}
+
+				for (int32 iVert = 0; iVert < Verts.Num(); iVert++)
+				{
+					AddVertexHelper(DebugMeshData, Verts[iVert] + CurrentData->NavMeshDrawOffset);
+				}
+				CurrentData->MeshBuilders.Add(DebugMeshData);
+				IndexesOffset += Verts.Num();
+			}
+
+			for (int32 i = 0; i < TileData.Links.Num(); i++)
+			{
+				const NavMeshGameplayDebugger::FOffMeshLink& SrcLink = TileData.Links[i];
+
+				const FVector V0 = SrcLink.Left.ToVector() + OffsetLocation + CurrentData->NavMeshDrawOffset;
+				const FVector V1 = SrcLink.Right.ToVector() + OffsetLocation + CurrentData->NavMeshDrawOffset;
+				const FColor LinkColor = SrcLink.Color;
+
+				CacheArc(CurrentData->NavLinkLines, V0, V1, 0.4f, 4, LinkColor, LinkLines_LineThickness);
+
+				const FVector VOffset(0, 0, FVector::Dist(V0, V1) * 1.333f);
+				CacheArrowHead(CurrentData->NavLinkLines, V1, V0 + VOffset, 30.f, LinkColor, LinkLines_LineThickness);
+				if (SrcLink.PackedFlags.Direction)
+				{
+					CacheArrowHead(CurrentData->NavLinkLines, V0, V1 + VOffset, 30.f, LinkColor, LinkLines_LineThickness);
+				}
+
+				// if the connection as a whole is valid check if there are any of ends is invalid
+				if (LinkColor != NavMeshRenderColor_OffMeshConnectionInvalid)
+				{
+					if (SrcLink.PackedFlags.Direction && (SrcLink.PackedFlags.ValidEnds & FRecastDebugGeometry::OMLE_Left) == 0)
+					{
+						// left end invalid - mark it
+						DrawWireCylinder(CurrentData->NavLinkLines, V0, FVector(1, 0, 0), FVector(0, 1, 0), FVector(0, 0, 1), NavMeshRenderColor_OffMeshConnectionInvalid, SrcLink.PackedFlags.Radius, 30 /*NavMesh->AgentMaxStepHeight*/, 16, 0, DefaultEdges_LineThickness);
+					}
+
+					if ((SrcLink.PackedFlags.ValidEnds & FRecastDebugGeometry::OMLE_Right) == 0)
+					{
+						DrawWireCylinder(CurrentData->NavLinkLines, V1, FVector(1, 0, 0), FVector(0, 1, 0), FVector(0, 0, 1), NavMeshRenderColor_OffMeshConnectionInvalid, SrcLink.PackedFlags.Radius, 30 /*NavMesh->AgentMaxStepHeight*/, 16, 0, DefaultEdges_LineThickness);
+					}
+				}
+			}
+		}
+	}
+#endif
+}
+
+class FDebugRenderSceneProxy* UNavMeshGameplayDebuggerObject::CreateSceneProxy(const UPrimitiveComponent* InComponent, UWorld* World, AActor* SelectedActor)
+{
+#if ENABLED_GAMEPLAY_DEBUGGER && WITH_RECAST	
+	FNavMeshSceneProxyData NewNavmeshRenderData;
+	NewNavmeshRenderData.Reset();
+	NewNavmeshRenderData.bNeedsNewData = false;
+	NewNavmeshRenderData.bEnableDrawing = false;
+	PrepareNavMeshData(&NewNavmeshRenderData, InComponent, World);
+
+	return NewNavmeshRenderData.bEnableDrawing ? new FRecastRenderingSceneProxy(InComponent, &NewNavmeshRenderData, true) : nullptr;
+#endif
+	return nullptr;
+}
+
diff --git a/Engine/Source/Runtime/AIModule/Private/GameplayDebugger/NavMeshGameplayDebuggerObject.h b/Engine/Source/Runtime/AIModule/Private/GameplayDebugger/NavMeshGameplayDebuggerObject.h
new file mode 100644
--- /dev/null
+++ b/Engine/Source/Runtime/AIModule/Private/GameplayDebugger/NavMeshGameplayDebuggerObject.h
@@ -0,0 +1,33 @@
+// Copyright 1998-2015 Epic Games, Inc. All Rights Reserved.
+
+#pragma once
+#include "Debug/GameplayDebuggerBaseObject.h"
+#include "NavMeshGameplayDebuggerObject.generated.h"
+
+UCLASS()
+class UNavMeshGameplayDebuggerObject : public UGameplayDebuggerBaseObject
+{
+	GENERATED_UCLASS_BODY()
+
+	virtual FString GetCategoryName() const override { return TEXT("NavMesh"); };
+	virtual void CollectDataToReplicate(APlayerController* MyPC, AActor *SelectedActor) override;
+	virtual void DrawCollectedData(APlayerController* MyPC, AActor *SelectedActor) override;
+	virtual class FDebugRenderSceneProxy* CreateSceneProxy(const UPrimitiveComponent* InComponent, UWorld* World, AActor* SelectedActor) override;
+
+	void ServerCollectNavmeshData(const AActor *SelectedActor);
+
+	UFUNCTION(Reliable, Server, WithValidation)
+	void ServerDiscardNavmeshData();
+
+protected:
+	float LastDataCaptureTime;
+
+	void PrepareNavMeshData(struct FNavMeshSceneProxyData*, const UPrimitiveComponent* InComponent, UWorld* World) const;
+
+	virtual void OnDataReplicationComplited(const TArray<uint8>& ReplicatedData) override;
+
+#if WITH_RECAST
+	ARecastNavMesh* GetNavData(const AActor *SelectedActor);
+#endif
+
+};
diff --git a/Engine/Source/Runtime/AIModule/Private/GameplayDebugger/PerceptionGameplayDebuggerObject.cpp b/Engine/Source/Runtime/AIModule/Private/GameplayDebugger/PerceptionGameplayDebuggerObject.cpp
new file mode 100644
--- /dev/null
+++ b/Engine/Source/Runtime/AIModule/Private/GameplayDebugger/PerceptionGameplayDebuggerObject.cpp
@@ -0,0 +1,85 @@
+// Copyright 1998-2015 Epic Games, Inc. All Rights Reserved.
+#include "AIModulePrivate.h"
+#include "Misc/CoreMisc.h"
+#include "Net/UnrealNetwork.h"
+#include "NavMeshRenderingHelpers.h"
+#include "DebugRenderSceneProxy.h"
+#include "PerceptionGameplayDebuggerObject.h"
+
+
+UPerceptionGameplayDebuggerObject::UPerceptionGameplayDebuggerObject(const FObjectInitializer& ObjectInitializer) : Super(ObjectInitializer)
+{
+	LastDataCaptureTime = 0;
+}
+
+
+void UPerceptionGameplayDebuggerObject::GetLifetimeReplicatedProps(TArray< FLifetimeProperty > & OutLifetimeProps) const
+{
+	Super::GetLifetimeReplicatedProps(OutLifetimeProps);
+
+#if ENABLED_GAMEPLAY_DEBUGGER
+	DOREPLIFETIME(UPerceptionGameplayDebuggerObject, SensingComponentLocation);
+	DOREPLIFETIME(UPerceptionGameplayDebuggerObject, PerceptionLegend);
+	DOREPLIFETIME(UPerceptionGameplayDebuggerObject, DistanceFromPlayer);
+	DOREPLIFETIME(UPerceptionGameplayDebuggerObject, DistanceFromSensor);
+#endif //ENABLED_GAMEPLAY_DEBUGGER
+}
+
+void UPerceptionGameplayDebuggerObject::CollectDataToReplicate(APlayerController* MyPC, AActor *SelectedActor)
+{
+	Super::CollectDataToReplicate(MyPC, SelectedActor);
+
+#if ENABLED_GAMEPLAY_DEBUGGER
+	if (GetWorld()->TimeSeconds - LastDataCaptureTime < 2)
+	{
+		return;
+	}
+
+	APawn* MyPawn = Cast<APawn>(SelectedActor);
+	if (MyPawn)
+	{
+		AAIController* BTAI = Cast<AAIController>(MyPawn->GetController());
+		if (BTAI)
+		{
+			UAIPerceptionComponent* PerceptionComponent = BTAI->GetAIPerceptionComponent();
+			if (PerceptionComponent == nullptr)
+			{
+				PerceptionComponent = MyPawn->FindComponentByClass<UAIPerceptionComponent>();
+			}
+			if (PerceptionComponent)
+			{
+				TArray<FString> PerceptionTexts;
+				GenericShapeElements.Reset();
+
+				PerceptionComponent->GrabGameplayDebuggerData(PerceptionTexts, GenericShapeElements);
+				DistanceFromPlayer = DistanceFromSensor = -1;
+
+				if (MyPC && MyPC->GetPawn())
+				{
+					DistanceFromPlayer = (MyPawn->GetActorLocation() - MyPC->GetPawn()->GetActorLocation()).Size();
+					DistanceFromSensor = SensingComponentLocation != FVector::ZeroVector ? (SensingComponentLocation - MyPC->GetPawn()->GetActorLocation()).Size() : -1;
+				}
+			}
+
+			UAIPerceptionSystem* PerceptionSys = UAIPerceptionSystem::GetCurrent(MyPawn->GetWorld());
+			if (PerceptionSys)
+			{
+				PerceptionLegend = PerceptionSys->GetPerceptionDebugLegend();
+			}
+		}
+	}
+#endif
+}
+
+void UPerceptionGameplayDebuggerObject::DrawCollectedData(APlayerController* MyPC, AActor *SelectedActor)
+{
+	Super::DrawCollectedData(MyPC, SelectedActor);
+
+#if ENABLED_GAMEPLAY_DEBUGGER
+	PrintString(DefaultContext, FString::Printf(TEXT("Draw Colors:")));
+	PrintString(DefaultContext, PerceptionLegend);
+
+	PrintString(DefaultContext, FString::Printf(TEXT("\nDistance Sensor-PlayerPawn: %.1f\n"), DistanceFromSensor));
+	PrintString(DefaultContext, FString::Printf(TEXT("Distance Pawn-PlayerPawn: %.1f\n"), DistanceFromPlayer));
+#endif
+}
diff --git a/Engine/Source/Runtime/AIModule/Private/GameplayDebugger/PerceptionGameplayDebuggerObject.h b/Engine/Source/Runtime/AIModule/Private/GameplayDebugger/PerceptionGameplayDebuggerObject.h
new file mode 100644
--- /dev/null
+++ b/Engine/Source/Runtime/AIModule/Private/GameplayDebugger/PerceptionGameplayDebuggerObject.h
@@ -0,0 +1,32 @@
+// Copyright 1998-2015 Epic Games, Inc. All Rights Reserved.
+
+#pragma once
+#include "Debug/GameplayDebuggerBaseObject.h"
+#include "Debug/DebugDrawService.h"
+#include "PerceptionGameplayDebuggerObject.generated.h"
+
+UCLASS()
+class UPerceptionGameplayDebuggerObject : public UGameplayDebuggerBaseObject
+{
+	GENERATED_UCLASS_BODY()
+
+protected:
+	UPROPERTY(Replicated)
+	FVector SensingComponentLocation;
+
+	UPROPERTY(Replicated)
+	FString PerceptionLegend;
+
+	UPROPERTY(Replicated)
+	float DistanceFromPlayer;
+
+	UPROPERTY(Replicated)
+	float DistanceFromSensor;
+
+	float LastDataCaptureTime;
+
+public:
+	virtual FString GetCategoryName() const override { return TEXT("Perception"); };
+	virtual void CollectDataToReplicate(APlayerController* MyPC, AActor *SelectedActor) override;
+	virtual void DrawCollectedData(APlayerController* MyPC, AActor *SelectedActor) override;
+};
diff --git a/Engine/Source/Runtime/AIModule/Private/Navigation/CrowdFollowingComponent.cpp b/Engine/Source/Runtime/AIModule/Private/Navigation/CrowdFollowingComponent.cpp
--- a/Engine/Source/Runtime/AIModule/Private/Navigation/CrowdFollowingComponent.cpp
+++ b/Engine/Source/Runtime/AIModule/Private/Navigation/CrowdFollowingComponent.cpp
@@ -215,6 +215,45 @@ void UCrowdFollowingComponent::SetCrowdRotateToVelocity(bool bEnable)
 	bRotateToVelocity = bEnable;
 }
 
+void UCrowdFollowingComponent::SetAvoidanceGroup(int32 GroupFlags, bool bUpdateAgent)
+{
+	if (AvoidanceGroup.Packed != GroupFlags)
+	{
+		AvoidanceGroup.SetFlagsDirectly(GroupFlags);
+
+		if (bUpdateAgent)
+		{
+			UpdateCrowdAgentParams();
+		}
+	}
+}
+
+void UCrowdFollowingComponent::SetGroupsToAvoid(int32 GroupFlags, bool bUpdateAgent)
+{
+	if (GroupsToAvoid.Packed != GroupFlags)
+	{
+		GroupsToAvoid.SetFlagsDirectly(GroupFlags);
+
+		if (bUpdateAgent)
+		{
+			UpdateCrowdAgentParams();
+		}
+	}
+}
+
+void UCrowdFollowingComponent::SetGroupsToIgnore(int32 GroupFlags, bool bUpdateAgent)
+{
+	if (GroupsToIgnore.Packed != GroupFlags)
+	{
+		GroupsToIgnore.SetFlagsDirectly(GroupFlags);
+
+		if (bUpdateAgent)
+		{
+			UpdateCrowdAgentParams();
+		}
+	}
+}
+
 void UCrowdFollowingComponent::SuspendCrowdSteering(bool bSuspend)
 {
 	if (bSuspendCrowdSimulation != bSuspend)
@@ -309,17 +348,20 @@ void UCrowdFollowingComponent::SetCrowdSimulation(bool bEnable)
 
 void UCrowdFollowingComponent::Initialize()
 {
+	const bool bPrevEnabled = bEnableCrowdSimulation;
+
 	Super::Initialize();
 
-	UCrowdManager* CrowdManager = UCrowdManager::GetCurrent(GetWorld());
-	if (CrowdManager)
-	{
-		ICrowdAgentInterface* IAgent = Cast<ICrowdAgentInterface>(this);
-		CrowdManager->RegisterAgent(IAgent);
-	}
-	else
+	const bool bWasRegistered = RegisterCrowdAgent();
+	if (!bWasRegistered)
 	{
-		bEnableCrowdSimulation = false;
+		// crowd manager might not be created yet if this component was spawned during level's begin play (possessing a pawn placed in level)
+		UNavigationSystem* NavSys = UNavigationSystem::GetCurrent(GetWorld());
+		if (NavSys && !NavSys->IsInitialized())
+		{
+			NavSys->OnNavigationInitDone.AddUObject(this, &UCrowdFollowingComponent::OnPendingNavigationInit);
+			bEnableCrowdSimulation = bPrevEnabled;
+		}
 	}
 }
 
@@ -655,7 +697,7 @@ void UCrowdFollowingComponent::SetMoveSegment(int32 SegmentStartIndex)
 		UCrowdManager* CrowdManager = UCrowdManager::GetCurrent(GetWorld());
 		if (CrowdManager)
 		{
-			CrowdManager->SetAgentMovePath(this, NavMeshPath, PathStartIndex, PathPartEndIdx);
+			CrowdManager->SetAgentMovePath(this, NavMeshPath, PathStartIndex, PathPartEndIdx, CurrentTargetPt);
 		}
 #endif
 	}
@@ -854,6 +896,33 @@ void UCrowdFollowingComponent::SwitchToNextPathPart()
 	SetMoveSegment(NewPartStart);
 }
 
+bool UCrowdFollowingComponent::RegisterCrowdAgent()
+{
+	bool bSuccess = false;
+
+	UCrowdManager* CrowdManager = UCrowdManager::GetCurrent(GetWorld());
+	if (CrowdManager)
+	{
+		ICrowdAgentInterface* IAgent = Cast<ICrowdAgentInterface>(this);
+		CrowdManager->RegisterAgent(IAgent);
+		bSuccess = true;
+	}
+	else
+	{
+		bEnableCrowdSimulation = false;
+	}
+
+	return bSuccess;
+}
+
+void UCrowdFollowingComponent::OnPendingNavigationInit()
+{
+	RegisterCrowdAgent();
+
+	// set movement component, it will cache MyNavData from its NavAgentProperties
+	SetMovementComponent(MovementComp);
+}
+
 void UCrowdFollowingComponent::GetDebugStringTokens(TArray<FString>& Tokens, TArray<EPathFollowingDebugTokens::Type>& Flags) const
 {
 	if (!bEnableCrowdSimulation)
