Copying //UE4/Dev-Framework to //UE4/Main

==========================
MAJOR FEATURES + CHANGES
==========================

Change 2786974 on 2015/12/02 by [REDACTED_PERSON]

	UE-23930 Fix for concatenator node's not working correctly when it has child nodes that are mixer nodes.

	- Fix was to track the number of sounds a sound a node can simultaneously play. In concatenator node, when a sound is notified as finishing, it tracks the sound index in the current child node before incrementing the child node index.

Change 2787015 on 2015/12/02 by [REDACTED_PERSON]

	changed color of root level decorator nodes in behavikor tree editor
	#ue UE-23957

Change 2787249 on 2015/12/02 by [REDACTED_PERSON]

	Make scene queries thread safe by ensuring that any data that is not thread safe is not returned. This is for the benefit of the user, but also we cannot access these pointers off the game thread.

Change 2788469 on 2015/12/03 by [REDACTED_PERSON]

	Ability system cleanup:
	Pass parameters around by const ref instead of value (FGameplayAbilityTargetDataHandle , TArray, FHitResult)
	Eliminate unnecessary multiple derferences of weak pointers in a single function
	Remove uses of auto, switch to using nullptr
	Reorganize booleans to properly pack them
	Const functions
	Properly mark functions virtual and override

Change 2788787 on 2015/12/03 by [REDACTED_PERSON]

	Reinitialize top level state machines when they become relevant, to match behavior of nested state machines.

Change 2789417 on 2015/12/03 by [REDACTED_PERSON]

	UE-19482 Fixing error in reporting audio asset memory usage for PS4 (and other platforms)

	- Issue was that USoundWave::GetResourceSize() was incorrectly summing uncompressed PCM data size alongside compressed data size for PS4
	- Added check using same condition in FAudioDevice::Precache which determines if compressed asset is fully decompressed into memory (and other decompression modes/types).

	#tests Tested loading maps in PS4 and confirming the uncompressed PCM data is not counted in audio asset resource size checks. Tested cooking assets in a map.

Change 2790152 on 2015/12/04 by [REDACTED_PERSON]

	Avoid unnecessary TArray and FHitResult copies

Change 2790182 on 2015/12/04 by [REDACTED_PERSON]

	Fixed notifies not being triggered when server calls Montage_JumpToSection.

	#tests Agora60p golden path, hyperbreach ultimate networked

Change 2790325 on 2015/12/04 by [REDACTED_PERSON]

	#ue4 - Optimized USceneComponent::SetWorldTransform() to avoid unnecessary copies and avoid unaligned SIMD reads and writes.

Change 2792284 on 2015/12/06 by [REDACTED_PERSON]

	Avoid unnecessary FGameplayTagContainer copies

Change 2792305 on 2015/12/06 by [REDACTED_PERSON]

	Avoid unnecessary FGameplayAbilitySpec copies

Change 2792592 on 2015/12/07 by [REDACTED_PERSON]

	Remove component reregistering logic from SetSkeletalMesh

Change 2792652 on 2015/12/07 by [REDACTED_PERSON]

	Add a way to opt out of ignoring trigger volumes. This makes it possible to run collision module off the game thread.

Change 2793378 on 2015/12/07 by [REDACTED_PERSON]

	fixed resetting path data between repaths
	#ue4 UE-22624

Change 2794690 on 2015/12/08 by [REDACTED_PERSON]

	#ANIM: Skeleton
	- fix crash when retargeting source that has been edited in editor
	- make sure to copy sockets when retarget skeleton

Change 2794694 on 2015/12/08 by [REDACTED_PERSON]

	Fixed duplicate slot names in anim slot groups. This was caused by not building the slot->group mapping at serialize time. COL would then re-add all the used slots to the group for a second time as the mapping wasn't built until postload.

Change 2795241 on 2015/12/08 by [REDACTED_PERSON]

	fixed potential division by zero in acceleration driven path following
	#ue4

Change 2796109 on 2015/12/09 by [REDACTED_PERSON]

	Pass FGameplayCueParameters by const ref in more places (avoids malloc allocations due to containing 2 FGameplayTagContainers)

Change 2796110 on 2015/12/09 by [REDACTED_PERSON]

diff --git a/Engine/Plugins/2D/Paper2D/Source/Paper2D/Private/GroupedSpriteSceneProxy.cpp b/Engine/Plugins/2D/Paper2D/Source/Paper2D/Private/GroupedSpriteSceneProxy.cpp
--- a/Engine/Plugins/2D/Paper2D/Source/Paper2D/Private/GroupedSpriteSceneProxy.cpp
+++ b/Engine/Plugins/2D/Paper2D/Source/Paper2D/Private/GroupedSpriteSceneProxy.cpp
@@ -121,5 +121,5 @@ void FGroupedSpriteSceneProxy::SetOneBodySetup_RenderThread(int32 InstanceIndex,
 void FGroupedSpriteSceneProxy::SetAllBodySetups_RenderThread(TArray<TWeakObjectPtr<UBodySetup>> Setups)
 {
 	check(Setups.Num() == NumInstances);
-	BodySetups = Setups;
+	BodySetups = MoveTemp(Setups);
 }
diff --git a/Engine/Source/Developer/CollisionAnalyzer/Private/SCAQueryDetails.cpp b/Engine/Source/Developer/CollisionAnalyzer/Private/SCAQueryDetails.cpp
--- a/Engine/Source/Developer/CollisionAnalyzer/Private/SCAQueryDetails.cpp
+++ b/Engine/Source/Developer/CollisionAnalyzer/Private/SCAQueryDetails.cpp
@@ -285,7 +285,7 @@ ECheckBoxState SCAQueryDetails::GetShowMissesState() const
 
 
 /** See if an array of results contains a particular component */
-static bool ResultsContainComponent(const TArray<FHitResult> Results, UPrimitiveComponent* Component)
+static bool ResultsContainComponent(const TArray<FHitResult>& Results, UPrimitiveComponent* Component)
 {
 	for(int32 i=0; i<Results.Num(); i++)
 	{
diff --git a/Engine/Source/Developer/TargetPlatform/Private/TargetPlatformManagerModule.cpp b/Engine/Source/Developer/TargetPlatform/Private/TargetPlatformManagerModule.cpp
--- a/Engine/Source/Developer/TargetPlatform/Private/TargetPlatformManagerModule.cpp
+++ b/Engine/Source/Developer/TargetPlatform/Private/TargetPlatformManagerModule.cpp
@@ -953,7 +953,7 @@ RETRY_SETUPANDVALIDATE:
 
 private:
 
-	void SetupEnvironmentVariables(TArray<FString> &EnvVarNames, TArray<FString> EnvVarValues)
+	void SetupEnvironmentVariables(TArray<FString> &EnvVarNames, const TArray<FString>& EnvVarValues)
 	{
 		for (int i = 0; i < EnvVarNames.Num(); ++i)
 		{
diff --git a/Engine/Source/Developer/Windows/WindowsTargetPlatform/Classes/WindowsTargetSettings.h b/Engine/Source/Developer/Windows/WindowsTargetPlatform/Classes/WindowsTargetSettings.h
--- a/Engine/Source/Developer/Windows/WindowsTargetPlatform/Classes/WindowsTargetSettings.h
+++ b/Engine/Source/Developer/Windows/WindowsTargetPlatform/Classes/WindowsTargetSettings.h
@@ -6,6 +6,7 @@
 
 #pragma once
 
+#include "AudioDevice.h"
 #include "WindowsTargetSettings.generated.h"
 
 UENUM()
@@ -38,4 +39,9 @@ public:
 	 */
 	UPROPERTY(EditAnywhere, config, Category="OS Info", Meta=(DisplayName = "Minimum OS Version"))
 	TEnumAsByte<EMinimumSupportedOS> MinimumOSVersion;
+
+	/** The audio device name to use if not the default windows audio device. Leave blank to use default audio device. */
+	UPROPERTY(config, EditAnywhere, Category = "Audio")
+	FString AudioDevice;
+
 };
diff --git a/Engine/Source/Editor/BehaviorTreeEditor/Private/BehaviorTreeEditor.cpp b/Engine/Source/Editor/BehaviorTreeEditor/Private/BehaviorTreeEditor.cpp
--- a/Engine/Source/Editor/BehaviorTreeEditor/Private/BehaviorTreeEditor.cpp
+++ b/Engine/Source/Editor/BehaviorTreeEditor/Private/BehaviorTreeEditor.cpp
@@ -54,6 +54,7 @@ FBehaviorTreeEditor::FBehaviorTreeEditor()
 	bShowDecoratorRangeLower = false;
 	bShowDecoratorRangeSelf = false;
 	bSelectedNodeIsInjected = false;
+	bSelectedNodeIsRootLevel = false;
 	SelectedNodesCount = 0;
 
 	bHasMultipleTaskBP = false;
@@ -304,6 +305,7 @@ void FBehaviorTreeEditor::RestoreBehaviorTree()
 	bShowDecoratorRangeLower = false;
 	bShowDecoratorRangeSelf = false;
 	bSelectedNodeIsInjected = false;
+	bSelectedNodeIsRootLevel = false;
 	MyGraph->UpdateAbortHighlight(EmptyMode, EmptyMode);
 }
 
@@ -440,6 +442,11 @@ EVisibility FBehaviorTreeEditor::GetInjectedNodeVisibility() const
 	return FBehaviorTreeDebugger::IsPIENotSimulating() && bSelectedNodeIsInjected ? EVisibility::Visible : EVisibility::Collapsed;
 }
 
+EVisibility FBehaviorTreeEditor::GetRootLevelNodeVisibility() const
+{
+	return FBehaviorTreeDebugger::IsPIENotSimulating() && bSelectedNodeIsRootLevel ? EVisibility::Visible : EVisibility::Collapsed;
+}
+
 FGraphAppearanceInfo FBehaviorTreeEditor::GetGraphAppearance() const
 {
 	FGraphAppearanceInfo AppearanceInfo;
@@ -613,6 +620,20 @@ TSharedRef<SWidget> FBehaviorTreeEditor::SpawnProperties()
 					.Text(LOCTEXT("InjectedNode", "Node is injected by subtree and can't be edited"))
 				]
 			]
+			+ SVerticalBox::Slot()
+			.HAlign(HAlign_Fill)
+			.Padding(0.0f, 5.0f)
+			[
+				SNew(SBorder)
+				.BorderBackgroundColor(BehaviorTreeColors::NodeBody::InjectedSubNode)
+				.BorderImage(FEditorStyle::GetBrush("Graph.StateNode.Body"))
+				.Visibility(this, &FBehaviorTreeEditor::GetRootLevelNodeVisibility)
+				.Padding(FMargin(5.0f))
+				[
+					SNew(STextBlock)
+					.Text(LOCTEXT("RootLevelNode", "Root level decorators are not executed\nThey will be injected into a parent tree"))
+				]
+			]
 			+SVerticalBox::Slot()
 			.HAlign(HAlign_Fill)
 			.Padding(0.0f, 5.0f)
@@ -696,6 +717,7 @@ void FBehaviorTreeEditor::OnSelectedNodesChanged(const TSet<class UObject*>& New
 	bShowDecoratorRangeSelf = false;
 	bForceDisablePropertyEdit = SelectionInfo.bInjectedNode;
 	bSelectedNodeIsInjected = SelectionInfo.bInjectedNode;
+	bSelectedNodeIsRootLevel = SelectionInfo.bRootLevelNode;
 
 	if (Selection.Num() == 1)
 	{
diff --git a/Engine/Source/Editor/BehaviorTreeEditor/Private/BehaviorTreeEditorUtils.cpp b/Engine/Source/Editor/BehaviorTreeEditor/Private/BehaviorTreeEditorUtils.cpp
--- a/Engine/Source/Editor/BehaviorTreeEditor/Private/BehaviorTreeEditorUtils.cpp
+++ b/Engine/Source/Editor/BehaviorTreeEditor/Private/BehaviorTreeEditorUtils.cpp
@@ -29,6 +29,7 @@ TArray<UObject*> BehaviorTreeEditorUtils::GetSelectionForPropertyEditor(const TS
 			Selection.Add(GraphNode_Decorator1->NodeInstance);
 			OutSelectionInfo.FoundDecorator = Cast<UBTDecorator>(GraphNode_Decorator1->NodeInstance);
 			OutSelectionInfo.bInjectedNode = OutSelectionInfo.bInjectedNode || GraphNode_Decorator1->bInjectedNode;
+			OutSelectionInfo.bRootLevelNode = GraphNode_Decorator1->bRootLevel;
 			continue;
 		}
 
@@ -52,6 +53,7 @@ TArray<UObject*> BehaviorTreeEditorUtils::GetSelectionForPropertyEditor(const TS
 		{
 			OutSelectionInfo.FoundGraphNode_CompDecorator = GraphNode_CompDecorator;
 			OutSelectionInfo.bInjectedNode = OutSelectionInfo.bInjectedNode || GraphNode_CompDecorator->bInjectedNode;
+			OutSelectionInfo.bRootLevelNode = OutSelectionInfo.bRootLevelNode || GraphNode_CompDecorator->bRootLevel;
 		}
 
 		Selection.Add(SelectionEntry);
diff --git a/Engine/Source/Editor/BehaviorTreeEditor/Private/BehaviorTreeEditorUtils.h b/Engine/Source/Editor/BehaviorTreeEditor/Private/BehaviorTreeEditorUtils.h
--- a/Engine/Source/Editor/BehaviorTreeEditor/Private/BehaviorTreeEditorUtils.h
+++ b/Engine/Source/Editor/BehaviorTreeEditor/Private/BehaviorTreeEditorUtils.h
@@ -10,12 +10,14 @@ namespace BehaviorTreeEditorUtils
 			: FoundGraphNode_CompDecorator(nullptr)
 			, FoundDecorator(nullptr)
 			, bInjectedNode(false)
+			, bRootLevelNode(false)
 		{
 		}
 
 		UBehaviorTreeGraphNode_CompositeDecorator* FoundGraphNode_CompDecorator;
 		UBTDecorator* FoundDecorator;
 		bool bInjectedNode;
+		bool bRootLevelNode;
 	};
 
 	/** Given a selection of nodes, return the instances that should be selected be selected for editing in the property panel */
diff --git a/Engine/Source/Editor/BehaviorTreeEditor/Private/SGraphNode_BehaviorTree.cpp b/Engine/Source/Editor/BehaviorTreeEditor/Private/SGraphNode_BehaviorTree.cpp
--- a/Engine/Source/Editor/BehaviorTreeEditor/Private/SGraphNode_BehaviorTree.cpp
+++ b/Engine/Source/Editor/BehaviorTreeEditor/Private/SGraphNode_BehaviorTree.cpp
@@ -178,6 +178,7 @@ FSlateColor SGraphNode_BehaviorTree::GetBorderBackgroundColor() const
 	const bool bSelectedSubNode = BTParentNode && GetOwnerPanel()->SelectionManager.SelectedNodes.Contains(GraphNode);
 	
 	UBTNode* NodeInstance = BTGraphNode ? Cast<UBTNode>(BTGraphNode->NodeInstance) : NULL;
+	const bool bIsConnectedTreeRoot = BTGraphNode && BTGraphNode->IsA<UBehaviorTreeGraphNode_Root>() && BTGraphNode->Pins.IsValidIndex(0) && BTGraphNode->Pins[0]->LinkedTo.Num() > 0;
 	const bool bIsDisconnected = NodeInstance && NodeInstance->GetExecutionIndex() == MAX_uint16;
 	const bool bIsService = BTGraphNode && BTGraphNode->IsA(UBehaviorTreeGraphNode_Service::StaticClass());
 	const bool bIsRootDecorator = BTGraphNode && BTGraphNode->bRootLevel;
@@ -208,6 +209,7 @@ FSlateColor SGraphNode_BehaviorTree::GetBorderBackgroundColor() const
 		!bIsRootDecorator && !bIsInjected && bIsDisconnected ? BehaviorTreeColors::NodeBorder::Disconnected :
 		bIsInDebuggerActiveState ? BehaviorTreeColors::NodeBorder::ActiveDebugging :
 		bIsInDebuggerPrevState ? BehaviorTreeColors::NodeBorder::InactiveDebugging :
+		bIsConnectedTreeRoot ? BehaviorTreeColors::NodeBorder::Root :
 		BehaviorTreeColors::NodeBorder::Inactive;
 }
 
@@ -230,7 +232,8 @@ FSlateColor SGraphNode_BehaviorTree::GetBackgroundColor() const
 	}
 	else if (BTGraph_Decorator || Cast<UBehaviorTreeGraphNode_CompositeDecorator>(GraphNode))
 	{
-		NodeColor = bIsActiveForDebugger ? BehaviorTreeColors::Debugger::ActiveDecorator : BehaviorTreeColors::NodeBody::Decorator;
+		NodeColor = bIsActiveForDebugger ? BehaviorTreeColors::Debugger::ActiveDecorator : 
+			BTGraphNode->bRootLevel ? BehaviorTreeColors::NodeBody::InjectedSubNode : BehaviorTreeColors::NodeBody::Decorator;
 	}
 	else if (Cast<UBehaviorTreeGraphNode_Task>(GraphNode))
 	{
@@ -245,6 +248,10 @@ FSlateColor SGraphNode_BehaviorTree::GetBackgroundColor() const
 	{
 		NodeColor = bIsActiveForDebugger ? BehaviorTreeColors::Debugger::ActiveService : BehaviorTreeColors::NodeBody::Service;
 	}
+	else if (Cast<UBehaviorTreeGraphNode_Root>(GraphNode) && GraphNode->Pins.IsValidIndex(0) && GraphNode->Pins[0]->LinkedTo.Num() > 0)
+	{
+		NodeColor = BehaviorTreeColors::NodeBody::Root;
+	}
 
 	return (FlashAlpha > 0.0f) ? FMath::Lerp(NodeColor, FlashColor, FlashAlpha) : NodeColor;
 }
diff --git a/Engine/Source/Editor/DetailCustomizations/Private/WindowsTargetSettingsDetails.cpp b/Engine/Source/Editor/DetailCustomizations/Private/WindowsTargetSettingsDetails.cpp
--- a/Engine/Source/Editor/DetailCustomizations/Private/WindowsTargetSettingsDetails.cpp
+++ b/Engine/Source/Editor/DetailCustomizations/Private/WindowsTargetSettingsDetails.cpp
@@ -7,6 +7,7 @@
 #include "IExternalImagePickerModule.h"
 #include "GameProjectGenerationModule.h"
 #include "ISourceControlModule.h"
+#include "AudioDevice.h"
 
 #define LOCTEXT_NAMESPACE "WindowsTargetSettingsDetails"
 
@@ -234,6 +235,50 @@ void FWindowsTargetSettingsDetails::CustomizeDetails( IDetailLayoutBuilder& Deta
 			.OnPostExternalImageCopy(FOnPostExternalImageCopy::CreateSP(this, &FWindowsTargetSettingsDetails::HandlePostExternalIconCopy))
 		]
 	];
+
+	TSharedPtr<IPropertyHandle> AudioDevicePropertyHandle = DetailBuilder.GetProperty("AudioDevice");
+	IDetailCategoryBuilder& AudioDeviceCategory = DetailBuilder.EditCategory("Audio");
+	IDetailPropertyRow& AudioDevicePropertyRow = AudioDeviceCategory.AddProperty(AudioDevicePropertyHandle);
+
+	AudioDevicePropertyRow.CustomWidget()
+	.NameContent()
+	[
+		AudioDevicePropertyHandle->CreatePropertyNameWidget()
+	]
+	.ValueContent()
+	.MaxDesiredWidth(500.0f)
+	.MinDesiredWidth(100.0f)
+	[
+		SNew(SHorizontalBox)
+
+		+ SHorizontalBox::Slot()
+		.FillWidth(1.0f)
+		[
+			SNew(SEditableTextBox)
+			.ForegroundColor(this, &FWindowsTargetSettingsDetails::HandleAudioDeviceBoxForegroundColor, AudioDevicePropertyHandle)
+			.OnTextChanged(this, &FWindowsTargetSettingsDetails::HandleAudioDeviceTextBoxTextChanged, AudioDevicePropertyHandle)
+			.OnTextCommitted(this, &FWindowsTargetSettingsDetails::HandleAudioDeviceTextBoxTextComitted, AudioDevicePropertyHandle)
+			.Text(this, &FWindowsTargetSettingsDetails::HandleAudioDeviceTextBoxText, AudioDevicePropertyHandle)
+			.ToolTipText(AudioDevicePropertyHandle->GetToolTipText())
+		]
+
+		+ SHorizontalBox::Slot()
+			.AutoWidth()
+			.Padding(2.0f, 0.0f, 0.0f, 0.0f)
+			[
+				SNew(SComboButton)
+				.ButtonContent()
+				[
+					SNullWidget::NullWidget
+				]
+				.ContentPadding(FMargin(6.0f, 1.0f))
+					.MenuContent()
+					[
+						MakeAudioDeviceMenu(AudioDevicePropertyHandle)
+					]
+				.ToolTipText(LOCTEXT("AudioDevicesButtonToolTip", "Pick from the list of available audio devices"))
+			]
+	];
 }
 
 
@@ -255,6 +300,112 @@ bool FWindowsTargetSettingsDetails::HandlePostExternalIconCopy(const FString& In
 	return true;
 }
 
+void FWindowsTargetSettingsDetails::HandleAudioDeviceSelected(FString AudioDeviceName, TSharedPtr<IPropertyHandle> PropertyHandle)
+{
+	PropertyHandle->SetValue(AudioDeviceName);
+}
+
+FSlateColor FWindowsTargetSettingsDetails::HandleAudioDeviceBoxForegroundColor(TSharedPtr<IPropertyHandle> PropertyHandle) const
+{
+	FString Value;
+
+	if (PropertyHandle->GetValue(Value) == FPropertyAccess::Success)
+	{
+		if (Value.IsEmpty() || IsValidAudioDeviceName(Value))
+		{
+			static const FName InvertedForegroundName("InvertedForeground");
+
+			// Return a valid slate color for a valid audio device
+			return FEditorStyle::GetSlateColor(InvertedForegroundName);
+		}
+	}
+
+	// Return Red, which means its an invalid audio device
+	return FLinearColor::Red;
+}
+
+FText FWindowsTargetSettingsDetails::HandleAudioDeviceTextBoxText(TSharedPtr<IPropertyHandle> PropertyHandle) const
+{
+	FString Value;
+
+	if (PropertyHandle->GetValue(Value) == FPropertyAccess::Success)
+	{
+		FString WindowsAudioDeviceName;
+		GConfig->GetString(TEXT("/Script/WindowsTargetPlatform.WindowsTargetSettings"), TEXT("AudioDevice"), WindowsAudioDeviceName, GEngineIni);
+		return FText::FromString(WindowsAudioDeviceName);
+	}
+
+	return FText::GetEmpty();
+}
+
+void FWindowsTargetSettingsDetails::HandleAudioDeviceTextBoxTextChanged(const FText& InText, TSharedPtr<IPropertyHandle> PropertyHandle)
+{
+	PropertyHandle->SetValue(InText.ToString());
+}
+
+void FWindowsTargetSettingsDetails::HandleAudioDeviceTextBoxTextComitted(const FText& InText, ETextCommit::Type CommitType, TSharedPtr<IPropertyHandle> PropertyHandle)
+{
+	FString Value;
+
+	// Clear the property if its not valid
+	if ((PropertyHandle->GetValue(Value) != FPropertyAccess::Success) || !IsValidAudioDeviceName(Value))
+	{
+		PropertyHandle->SetValue(FString());
+	}
+}
+
+bool FWindowsTargetSettingsDetails::IsValidAudioDeviceName(const FString& InDeviceName) const
+{
+	bool bIsValid = false;
+
+	FAudioDevice* AudioDevice = GEngine->GetMainAudioDevice();
+	if (AudioDevice)
+	{
+		TArray<FString> DeviceNames;
+		AudioDevice->GetAudioDeviceList(DeviceNames);
+
+		for (FString& DeviceName : DeviceNames)
+		{
+			if (InDeviceName == DeviceName)
+			{
+				bIsValid = true;
+				break;
+			}
+		}
+	}
+
+	return bIsValid;
+}
+
+TSharedRef<SWidget> FWindowsTargetSettingsDetails::MakeAudioDeviceMenu(const TSharedPtr<IPropertyHandle>& PropertyHandle)
+{
+	FMenuBuilder MenuBuilder(true, nullptr);
+	FAudioDevice* AudioDevice = GEngine->GetMainAudioDevice();
+	if (AudioDevice)
+	{
+		TArray<FString> AudioDeviceNames;
+		AudioDevice->GetAudioDeviceList(AudioDeviceNames);
+
+		// Construct the custom menu widget from the list of device names
+		MenuBuilder.BeginSection(NAME_None, LOCTEXT("AudioDevicesSectionHeader", "Audio Devices"));
+		{
+			for (int32 i = 0; i < AudioDeviceNames.Num(); i++)
+			{
+				FUIAction Action(FExecuteAction::CreateRaw(this, &FWindowsTargetSettingsDetails::HandleAudioDeviceSelected, AudioDeviceNames[i], PropertyHandle));
+				MenuBuilder.AddMenuEntry(
+					FText::FromString(AudioDeviceNames[i]),
+					FText::FromString(TEXT("")),
+					FSlateIcon(),
+					Action
+					);
+			}
+		}
+		MenuBuilder.EndSection();
+	}
+
+	return MenuBuilder.MakeWidget();
+}
+
 FTargetShaderFormatsPropertyDetails::FTargetShaderFormatsPropertyDetails(IDetailLayoutBuilder* InDetailBuilder)
 : DetailBuilder(InDetailBuilder)
 {
diff --git a/Engine/Source/Editor/DetailCustomizations/Private/WindowsTargetSettingsDetails.h b/Engine/Source/Editor/DetailCustomizations/Private/WindowsTargetSettingsDetails.h
--- a/Engine/Source/Editor/DetailCustomizations/Private/WindowsTargetSettingsDetails.h
+++ b/Engine/Source/Editor/DetailCustomizations/Private/WindowsTargetSettingsDetails.h
@@ -66,6 +66,29 @@ private:
 	/** Delegate handler to set the path to start picking from */
 	bool HandlePostExternalIconCopy(const FString& InChosenImage);
 
+	/** Handles when a new audio device is selected from list of available audio devices. */
+	void HandleAudioDeviceSelected(FString AudioDeviceName, TSharedPtr<IPropertyHandle> PropertyHandle);
+
+	/** Handles changing the foreground color of the audio device box. */
+	FSlateColor HandleAudioDeviceBoxForegroundColor(TSharedPtr<IPropertyHandle> PropertyHandle) const;
+
+	/** Handles getting text of the audio device list text block. */
+	FText HandleAudioDeviceTextBoxText(TSharedPtr<IPropertyHandle> PropertyHandle) const;
+
+	/** Handles text changes in the audio device list text block. */
+	void HandleAudioDeviceTextBoxTextChanged(const FText& InText, TSharedPtr<IPropertyHandle> PropertyHandle);
+
+	/** Handles committing changes in the audio list text block. */
+	void HandleAudioDeviceTextBoxTextComitted(const FText& InText, ETextCommit::Type CommitType, TSharedPtr<IPropertyHandle> PropertyHandle);
+
+protected:
+
+	/** Checks if the device name is valid. */
+	bool IsValidAudioDeviceName(const FString& InDeviceName) const;
+
+	/** Creates a widget for the audio device picker. */
+	TSharedRef<SWidget> MakeAudioDeviceMenu(const TSharedPtr<IPropertyHandle>& PropertyHandle);
+
 private:
 	/** Reference to the target shader formats property view */
 	TSharedPtr<FTargetShaderFormatsPropertyDetails> TargetShaderFormatsDetails;
diff --git a/Engine/Source/Editor/EditorStyle/Private/SlateEditorStyle.cpp b/Engine/Source/Editor/EditorStyle/Private/SlateEditorStyle.cpp
--- a/Engine/Source/Editor/EditorStyle/Private/SlateEditorStyle.cpp
+++ b/Engine/Source/Editor/EditorStyle/Private/SlateEditorStyle.cpp
@@ -5458,6 +5458,7 @@ void FSlateEditorStyle::FStyle::SetupClassIconsAndThumbnails()
 			TEXT("SlateWidgetStyleAsset"),
 			TEXT("SoundAttenuation"),
 			TEXT("SoundClass"),
+			TEXT("SoundConcurrency"),
 			TEXT("SoundCue"),
 			TEXT("SoundMix"),
 			TEXT("SphereReflectionCapture"),
diff --git a/Engine/Source/Editor/GameplayTagsEditor/Private/GameplayTagsK2Node_SwitchGameplayTagContainer.cpp b/Engine/Source/Editor/GameplayTagsEditor/Private/GameplayTagsK2Node_SwitchGameplayTagContainer.cpp
--- a/Engine/Source/Editor/GameplayTagsEditor/Private/GameplayTagsK2Node_SwitchGameplayTagContainer.cpp
+++ b/Engine/Source/Editor/GameplayTagsEditor/Private/GameplayTagsK2Node_SwitchGameplayTagContainer.cpp
@@ -43,14 +43,19 @@ bool UGameplayTagsK2Node_SwitchGameplayTagContainer::NotEqual_TagContainerTagCon
 {
 	FGameplayTagContainer TagContainer;
 
+	const FString OpenParenthesesStr(TEXT("("));
+	const FString CloseParenthesesStr(TEXT(")"));
+
 	// Convert string to Tag Container before compare
-	FString TagString = B;
-	if (TagString.StartsWith(TEXT("(")) && TagString.EndsWith(TEXT(")")))
+	FString TagString = MoveTemp(B);
+	if (TagString.StartsWith(OpenParenthesesStr, ESearchCase::CaseSensitive) && TagString.EndsWith(CloseParenthesesStr, ESearchCase::CaseSensitive))
 	{
 		TagString = TagString.LeftChop(1);
 		TagString = TagString.RightChop(1);
 
-		TagString.Split("=", NULL, &TagString);
+		const FString EqualStr(TEXT("="));
+
+		TagString.Split(EqualStr, nullptr, &TagString, ESearchCase::CaseSensitive);
 
 		TagString = TagString.LeftChop(1);
 		TagString = TagString.RightChop(1);
@@ -58,39 +63,43 @@ bool UGameplayTagsK2Node_SwitchGameplayTagContainer::NotEqual_TagContainerTagCon
 		FString ReadTag;
 		FString Remainder;
 
-		while (TagString.Split(TEXT(","), &ReadTag, &Remainder))
+		const FString CommaStr(TEXT(","));
+		const FString QuoteStr(TEXT("\""));
+
+		while (TagString.Split(CommaStr, &ReadTag, &Remainder, ESearchCase::CaseSensitive))
 		{
-			ReadTag.Split("=", NULL, &ReadTag);
-			if (ReadTag.EndsWith(TEXT(")")))
+			ReadTag.Split(EqualStr, nullptr, &ReadTag, ESearchCase::CaseSensitive);
+			if (ReadTag.EndsWith(CloseParenthesesStr, ESearchCase::CaseSensitive))
 			{
 				ReadTag = ReadTag.LeftChop(1);
-				if (ReadTag.StartsWith(TEXT("\"")) && ReadTag.EndsWith(TEXT("\"")))
+				if (ReadTag.StartsWith(QuoteStr, ESearchCase::CaseSensitive) && ReadTag.EndsWith(QuoteStr, ESearchCase::CaseSensitive))
 				{
 					ReadTag = ReadTag.LeftChop(1);
 					ReadTag = ReadTag.RightChop(1);
 				}
 			}
 			TagString = Remainder;
-			FGameplayTag Tag = IGameplayTagsModule::Get().GetGameplayTagsManager().RequestGameplayTag(FName(*ReadTag));
+
+			const FGameplayTag Tag = IGameplayTagsModule::Get().GetGameplayTagsManager().RequestGameplayTag(FName(*ReadTag));
 			TagContainer.AddTag(Tag);
 		}
 		if (Remainder.IsEmpty())
 		{
-			Remainder = TagString;
+			Remainder = MoveTemp(TagString);
 		}
 		if (!Remainder.IsEmpty())
 		{
-			Remainder.Split("=", NULL, &Remainder);
-			if (Remainder.EndsWith(TEXT(")")))
+			Remainder.Split(EqualStr, nullptr, &Remainder, ESearchCase::CaseSensitive);
+			if (Remainder.EndsWith(CloseParenthesesStr, ESearchCase::CaseSensitive))
 			{
 				Remainder = Remainder.LeftChop(1);
-				if (Remainder.StartsWith(TEXT("\"")) && Remainder.EndsWith(TEXT("\"")))
+				if (Remainder.StartsWith(QuoteStr, ESearchCase::CaseSensitive) && Remainder.EndsWith(QuoteStr, ESearchCase::CaseSensitive))
 				{
 					Remainder = Remainder.LeftChop(1);
 					Remainder = Remainder.RightChop(1);
 				}
 			}
-			FGameplayTag Tag = IGameplayTagsModule::Get().GetGameplayTagsManager().RequestGameplayTag(FName(*Remainder));
+			const FGameplayTag Tag = IGameplayTagsModule::Get().GetGameplayTagsManager().RequestGameplayTag(FName(*Remainder));
 			TagContainer.AddTag(Tag);
 		}
 	}
@@ -101,8 +110,8 @@ bool UGameplayTagsK2Node_SwitchGameplayTagContainer::NotEqual_TagContainerTagCon
 void UGameplayTagsK2Node_SwitchGameplayTagContainer::PostEditChangeProperty(struct FPropertyChangedEvent& PropertyChangedEvent)
 {
 	bool bIsDirty = false;
-	FName PropertyName = (PropertyChangedEvent.Property != NULL) ? PropertyChangedEvent.Property->GetFName() : NAME_None;
-	if (PropertyName == TEXT("PinContainers"))
+	FName PropertyName = (PropertyChangedEvent.Property != nullptr) ? PropertyChangedEvent.Property->GetFName() : NAME_None;
+	if (PropertyName == GET_MEMBER_NAME_CHECKED(UGameplayTagsK2Node_SwitchGameplayTagContainer, PinContainers))
 	{
 		bIsDirty = true;
 	}
@@ -172,7 +181,7 @@ void UGameplayTagsK2Node_SwitchGameplayTagContainer::CreateCasePins()
 	const UEdGraphSchema_K2* K2Schema = GetDefault<UEdGraphSchema_K2>();
 	for (int32 Index = 0; Index < PinNames.Num(); ++Index)
   	{
-		UEdGraphPin * NewPin = CreatePin(EGPD_Output, K2Schema->PC_Exec, TEXT(""), NULL, false, false, PinContainers[Index].ToString());
+		UEdGraphPin * NewPin = CreatePin(EGPD_Output, K2Schema->PC_Exec, TEXT(""), nullptr, false, false, PinContainers[Index].ToString());
 		NewPin->PinFriendlyName = FText::FromString(PinNames[Index].ToString());
   	}
 }
@@ -208,7 +217,7 @@ void UGameplayTagsK2Node_SwitchGameplayTagContainer::AddPinToSwitchNode()
 	PinNames.Add(FName(*PinName));
 
 	const UEdGraphSchema_K2* K2Schema = GetDefault<UEdGraphSchema_K2>();
-	UEdGraphPin* NewPin = CreatePin(EGPD_Output, K2Schema->PC_Exec, TEXT(""), NULL, false, false, PinName);
+	UEdGraphPin* NewPin = CreatePin(EGPD_Output, K2Schema->PC_Exec, TEXT(""), nullptr, false, false, PinName);
 	NewPin->PinFriendlyName = FText::FromString(PinName);
 	PinContainers.Add(FGameplayTagContainer());
 }
diff --git a/Engine/Source/Editor/KismetCompiler/Private/AnimBlueprintCompiler.cpp b/Engine/Source/Editor/KismetCompiler/Private/AnimBlueprintCompiler.cpp
--- a/Engine/Source/Editor/KismetCompiler/Private/AnimBlueprintCompiler.cpp
+++ b/Engine/Source/Editor/KismetCompiler/Private/AnimBlueprintCompiler.cpp
@@ -1852,7 +1852,11 @@ void FAnimBlueprintCompiler::FEvaluationHandlerRecord::PatchFunctionNameAndCopyR
 				if (AnimNodeSinglePropertyHandler.ArrayPins.Num() > 0)
 				{
 					UArrayProperty* SimpleCopyPropertyDest = CastChecked<UArrayProperty>(NodeVariableProperty->Struct->FindPropertyByName(PropertyIt.Key()));
-					check(SimpleCopyPropertySource->GetSize() == SimpleCopyPropertyDest->Inner->GetSize());
+					if(SimpleCopyPropertySource->GetSize() != SimpleCopyPropertyDest->Inner->GetSize())
+					{
+						bOnlyUsesCopyRecords = false;
+						break;
+					}
 
 					for (TMap<int32, UEdGraphPin*>::TConstIterator ArrayIt(AnimNodeSinglePropertyHandler.ArrayPins); ArrayIt; ++ArrayIt)
 					{
@@ -1870,13 +1874,21 @@ void FAnimBlueprintCompiler::FEvaluationHandlerRecord::PatchFunctionNameAndCopyR
 				else
 				{
 					UProperty* SimpleCopyPropertyDest = NodeVariableProperty->Struct->FindPropertyByName(PropertyIt.Key());
-					check(SimpleCopyPropertyDest);
+					if(SimpleCopyPropertyDest == nullptr)
+					{
+						bOnlyUsesCopyRecords = false;
+						break;
+					}
+
 					if(AnimNodeSinglePropertyHandler.SubStructPropertyName != NAME_None)
 					{
 						UStructProperty* SourceStructProperty = CastChecked<UStructProperty>(SimpleCopyPropertySource);
 						UProperty* SourceSubProperty = SourceStructProperty->Struct->FindPropertyByName(AnimNodeSinglePropertyHandler.SubStructPropertyName);
-						check(SourceSubProperty);
-						check(SourceSubProperty->GetSize() == SimpleCopyPropertyDest->GetSize());
+						if(SourceSubProperty == nullptr || SourceSubProperty->GetSize() != SimpleCopyPropertyDest->GetSize())
+						{
+							bOnlyUsesCopyRecords = false;
+							break;
+						}
 
 						// Local sub-struct variable get
 						FExposedValueCopyRecord CopyRecord;
@@ -1891,7 +1903,11 @@ void FAnimBlueprintCompiler::FEvaluationHandlerRecord::PatchFunctionNameAndCopyR
 					}
 					else
 					{
-						check(SimpleCopyPropertySource->GetSize() == SimpleCopyPropertyDest->GetSize());
+						if(SimpleCopyPropertySource->GetSize() != SimpleCopyPropertyDest->GetSize())
+						{
+							bOnlyUsesCopyRecords = false;
+							break;
+						}
 
 						// Local variable get
 						FExposedValueCopyRecord CopyRecord;
@@ -2075,9 +2091,30 @@ bool FAnimBlueprintCompiler::FEvaluationHandlerRecord::CheckForLogicalNot(FAnimN
 	return false;
 }
 
+/** The functions that we can safely native-break */
+static const FName NativeBreakFunctionNameWhitelist[] =
+{
+	FName(TEXT("BreakVector")),
+	FName(TEXT("BreakVector2D")),
+	FName(TEXT("BreakRotator")),
+};
+
+/** Check whether a native break function can be safely used in the fast-path copy system (ie. source and dest data will be the same) */
+static bool IsWhitelistedNativeBreak(const FName& InFunctionName)
+{
+	for(const FName& FunctionName : NativeBreakFunctionNameWhitelist)
+	{
+		if(InFunctionName == FunctionName)
+		{
+			return true;
+		}
+	}
+
+	return false;
+}
+
 bool FAnimBlueprintCompiler::FEvaluationHandlerRecord::CheckForStructMemberAccess(FAnimNodeSinglePropertyHandler& Handler, UEdGraphPin* SourcePin)
 {
-	static const FName BreakTransformName(TEXT("BreakTransform"));
 	if(SourcePin)
 	{
 		UEdGraphPin* DestPin = nullptr;
@@ -2097,9 +2134,7 @@ bool FAnimBlueprintCompiler::FEvaluationHandlerRecord::CheckForStructMemberAcces
 		else if(UK2Node_CallFunction* NativeBreakNode = Cast<UK2Node_CallFunction>(FollowKnots(SourcePin, DestPin)))
 		{
 			UFunction* Function = NativeBreakNode->FunctionReference.ResolveMember<UFunction>(UKismetMathLibrary::StaticClass());
-			if( Function && Function->HasMetaData(TEXT("NativeBreakFunc")) &&
-				Function->GetFName() != BreakTransformName) // Skip Break Transform as it is not compatible (it is not a pure "break" function as it performs type 
-															// conversion on the "Rotation" component. This means we cannot do a fast path copy as src/dest copy types do not match)
+			if(Function && Function->HasMetaData(TEXT("NativeBreakFunc")) && IsWhitelistedNativeBreak(Function->GetFName()))
 			{
 				if(UEdGraphPin* InputPin = FindFirstInputPin(NativeBreakNode))
 				{
diff --git a/Engine/Source/Editor/Persona/Private/AnimationEditorViewportClient.cpp b/Engine/Source/Editor/Persona/Private/AnimationEditorViewportClient.cpp
--- a/Engine/Source/Editor/Persona/Private/AnimationEditorViewportClient.cpp
+++ b/Engine/Source/Editor/Persona/Private/AnimationEditorViewportClient.cpp
@@ -1856,7 +1856,7 @@ void FAnimationViewportClient::DrawMeshBones(USkeletalMeshComponent * MeshCompon
 	}
 }
 
-void FAnimationViewportClient::DrawBones(const USkeletalMeshComponent* MeshComponent, const TArray<FBoneIndexType> & RequiredBones, const TArray<FTransform> & WorldTransforms, FPrimitiveDrawInterface* PDI, const TArray<FLinearColor> BoneColours, float LineThickness/*=0.f*/) const
+void FAnimationViewportClient::DrawBones(const USkeletalMeshComponent* MeshComponent, const TArray<FBoneIndexType> & RequiredBones, const TArray<FTransform> & WorldTransforms, FPrimitiveDrawInterface* PDI, const TArray<FLinearColor>& BoneColours, float LineThickness/*=0.f*/) const
 {
 	check ( MeshComponent && MeshComponent->SkeletalMesh );
 
diff --git a/Engine/Source/Editor/Persona/Private/AnimationEditorViewportClient.h b/Engine/Source/Editor/Persona/Private/AnimationEditorViewportClient.h
--- a/Engine/Source/Editor/Persona/Private/AnimationEditorViewportClient.h
+++ b/Engine/Source/Editor/Persona/Private/AnimationEditorViewportClient.h
@@ -313,7 +313,7 @@ private:
 	/** Draw Bones for non retargeted animation. */
 	void DrawMeshBonesBakedAnimation(UDebugSkelMeshComponent * MeshComponent, FPrimitiveDrawInterface* PDI) const;
 	/** Draws Bones for RequiredBones with WorldTransform **/
-	void DrawBones(const USkeletalMeshComponent* MeshComponent, const TArray<FBoneIndexType> & RequiredBones, const TArray<FTransform> & WorldTransforms, FPrimitiveDrawInterface* PDI, const TArray<FLinearColor> BoneColours, float LineThickness=0.f) const;
+	void DrawBones(const USkeletalMeshComponent* MeshComponent, const TArray<FBoneIndexType> & RequiredBones, const TArray<FTransform> & WorldTransforms, FPrimitiveDrawInterface* PDI, const TArray<FLinearColor>& BoneColours, float LineThickness=0.f) const;
 	/** Draw Sub set of Bones **/
 	void DrawMeshSubsetBones(const USkeletalMeshComponent* MeshComponent, const TArray<int32>& BonesOfInterest, FPrimitiveDrawInterface* PDI) const;
 	/** Draws Gizmo for the Transform in foreground **/
diff --git a/Engine/Source/Editor/UMGEditor/Private/WidgetBlueprintEditor.cpp b/Engine/Source/Editor/UMGEditor/Private/WidgetBlueprintEditor.cpp
--- a/Engine/Source/Editor/UMGEditor/Private/WidgetBlueprintEditor.cpp
+++ b/Engine/Source/Editor/UMGEditor/Private/WidgetBlueprintEditor.cpp
@@ -831,7 +831,7 @@ TSharedRef<FExtender> FWidgetBlueprintEditor::GetContextSensitiveSequencerExtend
 	return AddTrackMenuExtender;
 }
 
-void FWidgetBlueprintEditor::ExtendSequencerAddTrackMenu( FMenuBuilder& AddTrackMenuBuilder, TArray<UObject*> ContextObjects )
+void FWidgetBlueprintEditor::ExtendSequencerAddTrackMenu( FMenuBuilder& AddTrackMenuBuilder, const TArray<UObject*> ContextObjects )
 {
 	if ( ContextObjects.Num() == 1 )
 	{
diff --git a/Engine/Source/Editor/UMGEditor/Private/WidgetBlueprintEditor.h b/Engine/Source/Editor/UMGEditor/Private/WidgetBlueprintEditor.h
--- a/Engine/Source/Editor/UMGEditor/Private/WidgetBlueprintEditor.h
+++ b/Engine/Source/Editor/UMGEditor/Private/WidgetBlueprintEditor.h
@@ -185,7 +185,7 @@ private:
 	TSharedRef<FExtender> GetContextSensitiveSequencerExtender( const TSharedRef<FUICommandList> CommandList, const TArray<UObject*> ContextSensitiveObjects );
 
 	/** Extends the sequencer add track menu. */
-	void ExtendSequencerAddTrackMenu( FMenuBuilder& AddTrackMenuBuilder, TArray<UObject*> ContextObjects );
+	void ExtendSequencerAddTrackMenu( FMenuBuilder& AddTrackMenuBuilder, const TArray<UObject*> ContextObjects );
 
 	/** Add an animation track for the supplied slot to the current animation. */
 	void AddSlotTrack( UPanelSlot* Slot );
diff --git a/Engine/Source/Editor/UnrealEd/Private/ApexClothingUtils.cpp b/Engine/Source/Editor/UnrealEd/Private/ApexClothingUtils.cpp
--- a/Engine/Source/Editor/UnrealEd/Private/ApexClothingUtils.cpp
+++ b/Engine/Source/Editor/UnrealEd/Private/ApexClothingUtils.cpp
@@ -917,7 +917,7 @@ bool AssociateClothingAssetWithSkeletalMesh(USkeletalMesh* SkelMesh, int32 LODIn
 		FMultiSizeIndexContainerData AdjacencyIndexData;
 		IMeshUtilities& MeshUtilities = FModuleManager::Get().LoadModuleChecked<IMeshUtilities>("MeshUtilities");
 
-		UE_LOG(LogApexClothingUtils, Warning, TEXT("Building adjacency information for a clothing section in skeletal mesh '%s'."), *SkelMesh->GetPathName());
+		UE_LOG(LogApexClothingUtils, Warning, TEXT("Building adjacency information for a clothing section in skeletal mesh '%s'. Please resave the asset."), *SkelMesh->GetPathName());
 
 		LODModel.GetVertices(Vertices);
 		AdjacencyIndexData.DataTypeSize = LODModel.MultiSizeIndexContainer.GetDataTypeSize();
diff --git a/Engine/Source/Editor/UnrealEd/Private/EditorAnimUtils.cpp b/Engine/Source/Editor/UnrealEd/Private/EditorAnimUtils.cpp
--- a/Engine/Source/Editor/UnrealEd/Private/EditorAnimUtils.cpp
+++ b/Engine/Source/Editor/UnrealEd/Private/EditorAnimUtils.cpp
@@ -464,7 +464,8 @@ namespace EditorAnimUtils
 	{
 		// Copy curve data from source asset, preserving data in the target if present.
 		const FSmartNameMapping* OldNameMapping = OldSkeleton->GetSmartNameContainer(ContainerName);
-		const FSmartNameMapping* NewNameMapping = NewSkeleton->GetSmartNameContainer(ContainerName);
+		const FSmartNameMapping* NewNameMapping = NewSkeleton->GetOrAddSmartNameContainer(ContainerName);
+
 		SequenceBase->RawCurveData.UpdateLastObservedNames(OldNameMapping, CurveType);
 
 		switch (CurveType)
diff --git a/Engine/Source/Runtime/AIModule/Classes/Navigation/PathFollowingComponent.h b/Engine/Source/Runtime/AIModule/Classes/Navigation/PathFollowingComponent.h
--- a/Engine/Source/Runtime/AIModule/Classes/Navigation/PathFollowingComponent.h
+++ b/Engine/Source/Runtime/AIModule/Classes/Navigation/PathFollowingComponent.h
@@ -12,6 +12,7 @@ AIMODULE_API DECLARE_LOG_CATEGORY_EXTERN(LogPathFollowing, Warning, All);
 class UNavMovementComponent;
 class UCanvas;
 class AActor;
+class APawn;
 class INavLinkCustomInterface;
 class INavAgentInterface;
 class UNavigationComponent;
@@ -415,6 +416,15 @@ protected:
 	/** direction of current move segment */
 	FVector MoveSegmentDirection;
 
+	/** braking distance for acceleration driven path following */
+	float CachedBrakingDistance;
+
+	/** max speed used for CachedBrakingDistance */
+	float CachedBrakingMaxSpeed;
+
+	/** index of path point for starting deceleration */
+	int32 DecelerationSegmentIndex;
+
 	/** reset path following data */
 	virtual void Reset();
 
@@ -436,6 +446,9 @@ protected:
 	/** update blocked movement detection, @returns true if new sample was added */
 	virtual bool UpdateBlockDetection();
 
+	/** updates braking distance and deceleration segment */
+	virtual void UpdateDecelerationData();
+
 	/** check if move is completed */
 	bool HasReachedDestination(const FVector& CurrentLocation) const;
 
diff --git a/Engine/Source/Runtime/AIModule/Private/Navigation/CrowdFollowingComponent.cpp b/Engine/Source/Runtime/AIModule/Private/Navigation/CrowdFollowingComponent.cpp
--- a/Engine/Source/Runtime/AIModule/Private/Navigation/CrowdFollowingComponent.cpp
+++ b/Engine/Source/Runtime/AIModule/Private/Navigation/CrowdFollowingComponent.cpp
@@ -321,13 +321,26 @@ bool UCrowdFollowingComponent::UpdateCachedGoal(FVector& NewGoalPos)
 
 void UCrowdFollowingComponent::ApplyCrowdAgentVelocity(const FVector& NewVelocity, const FVector& DestPathCorner, bool bTraversingLink)
 {
-	if (IsCrowdSimulationEnabled() && Status == EPathFollowingStatus::Moving)
+	if (IsCrowdSimulationEnabled() && Status == EPathFollowingStatus::Moving && MovementComp)
 	{
 		if (bAffectFallingVelocity || CharacterMovement == NULL || CharacterMovement->MovementMode != MOVE_Falling)
 		{
-			MovementComp->RequestDirectMove(NewVelocity, false);
-
 			UpdateCachedDirections(NewVelocity, DestPathCorner, bTraversingLink);
+
+			const bool bAccelerationBased = MovementComp->UseAccelerationForPathFollowing();
+			if (bAccelerationBased)
+			{
+				const float MaxSpeed = GetCrowdAgentMaxSpeed();
+				const float NewSpeed = NewVelocity.Size();
+				const float SpeedPct = FMath::Clamp(NewSpeed / MaxSpeed, 0.0f, 1.0f);
+				const FVector MoveInput = FMath::IsNearlyZero(NewSpeed) ? FVector::ZeroVector : ((NewVelocity / NewSpeed) * SpeedPct);
+
+				MovementComp->RequestPathMove(MoveInput);
+			}
+			else
+			{
+				MovementComp->RequestDirectMove(NewVelocity, false);
+			}
 		}
 	}
 }
diff --git a/Engine/Source/Runtime/AIModule/Private/Navigation/PathFollowingComponent.cpp b/Engine/Source/Runtime/AIModule/Private/Navigation/PathFollowingComponent.cpp
--- a/Engine/Source/Runtime/AIModule/Private/Navigation/PathFollowingComponent.cpp
+++ b/Engine/Source/Runtime/AIModule/Private/Navigation/PathFollowingComponent.cpp
@@ -55,6 +55,10 @@ UPathFollowingComponent::UPathFollowingComponent(const FObjectInitializer& Objec
 	MoveSegmentStartRef = INVALID_NAVNODEREF;
 	MoveSegmentEndRef = INVALID_NAVNODEREF;
 
+	CachedBrakingDistance = 100.0f;
+	CachedBrakingMaxSpeed = 0.0f;
+	DecelerationSegmentIndex = INDEX_NONE;
+
 	bStopOnOverlap = true;
 	Status = EPathFollowingStatus::Idle;
 }
@@ -219,6 +223,7 @@ FAIRequestID UPathFollowingComponent::RequestMove(FNavPathSharedPtr InPath, FReq
 		OnRequestFinished = OnComplete;	
 		bStopOnOverlap = bInStopOnOverlap;
 		SetDestinationActor(InDestinationActor);
+		UpdateDecelerationData();
 
 	#if ENABLE_VISUAL_LOG
 		const FVector CurrentLocation = MovementComp ? MovementComp->GetActorFeetLocation() : FVector::ZeroVector;
@@ -273,6 +278,7 @@ bool UPathFollowingComponent::UpdateMove(FNavPathSharedPtr InPath, FAIRequestID
 	Path = InPath;
 	OnPathUpdated();
 	GetWorld()->GetTimerManager().ClearTimer(WaitingForPathTimer);
+	UpdateDecelerationData();
 
 	if (Status == EPathFollowingStatus::Waiting || Status == EPathFollowingStatus::Moving)
 	{
@@ -424,7 +430,8 @@ void UPathFollowingComponent::OnPathFinished(EPathFollowingResult::Type Result)
 	Reset();
 	UpdateMoveFocus();
 
-	if (MovementComp && MovementComp->CanStopPathFollowing() && bStopMovementOnFinish)
+	if (MovementComp && MovementComp->CanStopPathFollowing() && 
+		bStopMovementOnFinish && !MovementComp->UseAccelerationForPathFollowing())
 	{
 		MovementComp->StopMovementKeepPathing();
 	}
@@ -481,9 +488,9 @@ void UPathFollowingComponent::TickComponent(float DeltaTime, enum ELevelTick Tic
 void UPathFollowingComponent::SetMovementComponent(UNavMovementComponent* MoveComp)
 {
 	MovementComp = MoveComp;
-	MyNavData = NULL;
+	MyNavData = nullptr;
 
-	if (MoveComp != NULL)
+	if (MoveComp)
 	{
 		const FNavAgentProperties& NavAgentProps = MoveComp->GetNavAgentPropertiesRef();
 		MyDefaultAcceptanceRadius = NavAgentProps.AgentRadius;
@@ -517,6 +524,7 @@ void UPathFollowingComponent::Reset()
 	MoveSegmentStartIndex = 0;
 	MoveSegmentStartRef = INVALID_NAVNODEREF;
 	MoveSegmentEndRef = INVALID_NAVNODEREF;
+	DecelerationSegmentIndex = INDEX_NONE;
 
 	LocationSamples.Reset();
 	LastSampleTime = 0.0f;
@@ -867,21 +875,44 @@ void UPathFollowingComponent::UpdatePathSegment()
 
 void UPathFollowingComponent::FollowPathSegment(float DeltaTime)
 {
-	if (MovementComp == NULL || !Path.IsValid())
+	if (!Path.IsValid() || MovementComp == nullptr)
 	{
 		return;
 	}
 
-	//const FVector CurrentLocation = MovementComp->IsMovingOnGround() ? MovementComp->GetActorFeetLocation() : MovementComp->GetActorLocation();
 	const FVector CurrentLocation = MovementComp->GetActorFeetLocation();
 	const FVector CurrentTarget = GetCurrentTargetLocation();
-	FVector MoveVelocity = (CurrentTarget - CurrentLocation) / DeltaTime;
 
-	const int32 LastSegmentStartIndex = Path->GetPathPoints().Num() - 2;
-	const bool bNotFollowingLastSegment = (MoveSegmentStartIndex < LastSegmentStartIndex);
+	const bool bAccelerationBased = MovementComp->UseAccelerationForPathFollowing();
+	if (bAccelerationBased)
+	{
+		FVector MoveInput = (CurrentTarget - CurrentLocation).GetSafeNormal();
+
+		if (MoveSegmentStartIndex >= DecelerationSegmentIndex)
+		{
+			const FVector PathEnd = Path->GetEndLocation();
+			const float DistToEndSq = FVector::DistSquared(CurrentLocation, PathEnd);
+			const bool bShouldDecelerate = DistToEndSq < FMath::Square(CachedBrakingDistance);
+			if (bShouldDecelerate)
+			{
+				const float SpeedPct = FMath::Clamp(FMath::Sqrt(DistToEndSq) / CachedBrakingDistance, 0.0f, 1.0f);
+				MoveInput *= SpeedPct;
+			}
+		}
+
+		PostProcessMove.ExecuteIfBound(this, MoveInput);
+		MovementComp->RequestPathMove(MoveInput);
+	}
+	else
+	{
+		FVector MoveVelocity = (CurrentTarget - CurrentLocation) / DeltaTime;
+
+		const int32 LastSegmentStartIndex = Path->GetPathPoints().Num() - 2;
+		const bool bNotFollowingLastSegment = (MoveSegmentStartIndex < LastSegmentStartIndex);
 
-	PostProcessMove.ExecuteIfBound(this, MoveVelocity);
-	MovementComp->RequestDirectMove(MoveVelocity, bNotFollowingLastSegment);
+		PostProcessMove.ExecuteIfBound(this, MoveVelocity);
+		MovementComp->RequestDirectMove(MoveVelocity, bNotFollowingLastSegment);
+	}
 }
 
 bool UPathFollowingComponent::HasReached(const FVector& TestPoint, float InAcceptanceRadius, bool bExactSpot) const
@@ -1249,6 +1280,42 @@ void UPathFollowingComponent::ForceBlockDetectionUpdate()
 	LastSampleTime = 0.0f;
 }
 
+void UPathFollowingComponent::UpdateDecelerationData()
+{
+	const float CurrentMaxSpeed = MovementComp ? MovementComp->GetMaxSpeed() : 0.0f;
+	bool bUpdatePathSegment = (DecelerationSegmentIndex == INDEX_NONE);
+	if (CurrentMaxSpeed != CachedBrakingMaxSpeed)
+	{
+		const float PrevBrakingDistance = CachedBrakingDistance;
+		CachedBrakingDistance = MovementComp->GetPathFollowingBrakingDistance(CurrentMaxSpeed);
+		CachedBrakingMaxSpeed = CurrentMaxSpeed;
+		
+		bUpdatePathSegment = bUpdatePathSegment || (PrevBrakingDistance != CachedBrakingDistance);
+	}
+
+	if (bUpdatePathSegment && Path.IsValid())
+	{
+		DecelerationSegmentIndex = 0;
+
+		const TArray<FNavPathPoint>& PathPoints = Path->GetPathPoints();
+		float PathLengthFromEnd = 0.0f;
+
+		for (int32 Idx = PathPoints.Num() - 1; Idx > 0; Idx--)
+		{
+			const float PathSegmentLength = FVector::Dist(PathPoints[Idx].Location, PathPoints[Idx - 1].Location);
+			PathLengthFromEnd += PathSegmentLength;
+
+			if (PathLengthFromEnd > CachedBrakingDistance)
+			{
+				DecelerationSegmentIndex = Idx - 1;
+				break;
+			}
+		}
+
+		UE_VLOG(GetOwner(), LogPathFollowing, Log, TEXT("Updated deceleration segment: %d (MaxSpeed:%.2f, BrakingDistance:%.2f"), DecelerationSegmentIndex, CachedBrakingMaxSpeed, CachedBrakingDistance);
+	}
+}
+
 void UPathFollowingComponent::UpdateMoveFocus()
 {
 	AAIController* AIOwner = Cast<AAIController>(GetOwner());
diff --git a/Engine/Source/Runtime/AnimGraphRuntime/Private/AnimNodes/AnimNode_LayeredBoneBlend.cpp b/Engine/Source/Runtime/AnimGraphRuntime/Private/AnimNodes/AnimNode_LayeredBoneBlend.cpp
--- a/Engine/Source/Runtime/AnimGraphRuntime/Private/AnimNodes/AnimNode_LayeredBoneBlend.cpp
+++ b/Engine/Source/Runtime/AnimGraphRuntime/Private/AnimNodes/AnimNode_LayeredBoneBlend.cpp
@@ -48,10 +48,7 @@ void FAnimNode_LayeredBoneBlend::CacheBones(const FAnimationCacheBonesContext& C
 		BlendPoses[ChildIndex].CacheBones(Context);
 	}
 
-	if (Context.AnimInstanceProxy->GetRequiredBones().GetBoneIndicesArray().Num() != DesiredBoneBlendWeights.Num())
-	{
-		ReinitializeBoneBlendWeights(Context.AnimInstanceProxy->GetRequiredBones(), Context.AnimInstanceProxy->GetSkeleton());
-	}
+	ReinitializeBoneBlendWeights(Context.AnimInstanceProxy->GetRequiredBones(), Context.AnimInstanceProxy->GetSkeleton());
 }
 
 void FAnimNode_LayeredBoneBlend::Update(const FAnimationUpdateContext& Context)
diff --git a/Engine/Source/Runtime/Core/Private/Android/AndroidInputInterface.cpp b/Engine/Source/Runtime/Core/Private/Android/AndroidInputInterface.cpp
--- a/Engine/Source/Runtime/Core/Private/Android/AndroidInputInterface.cpp
+++ b/Engine/Source/Runtime/Core/Private/Android/AndroidInputInterface.cpp
@@ -918,7 +918,7 @@ void FAndroidInputInterface::SendControllerEvents()
 	}
 }
 
-void FAndroidInputInterface::QueueTouchInput(TArray<TouchInput> InTouchEvents)
+void FAndroidInputInterface::QueueTouchInput(const TArray<TouchInput>& InTouchEvents)
 {
 	FScopeLock Lock(&TouchInputCriticalSection);
 
diff --git a/Engine/Source/Runtime/Core/Private/Async/TaskGraph.cpp b/Engine/Source/Runtime/Core/Private/Async/TaskGraph.cpp
--- a/Engine/Source/Runtime/Core/Private/Async/TaskGraph.cpp
+++ b/Engine/Source/Runtime/Core/Private/Async/TaskGraph.cpp
@@ -1266,6 +1266,11 @@ public:
 	**/
 	virtual ~FTaskGraphImplementation()
 	{
+		for (auto& Callback : ShutdownCallbacks)
+		{
+			Callback();
+		}
+		ShutdownCallbacks.Empty();
 		for (int32 ThreadIndex = 0; ThreadIndex < NumThreads; ThreadIndex++)
 		{
 			Thread(ThreadIndex).RequestQuit(-1);
@@ -1560,6 +1565,12 @@ public:
 		TGraphTask<FTriggerEventGraphTask>::CreateTask(&Tasks, CurrentThreadIfKnown).ConstructAndDispatchWhenReady(InEvent);
 	}
 
+	virtual void AddShutdownCallback(TFunction<void()>& Callback)
+	{
+		ShutdownCallbacks.Emplace(Callback);
+	}
+
+
 	// Scheduling utilities
 
 	void StartTaskThread(int32 IndexToStart)
@@ -2067,6 +2078,9 @@ private:
 	/** Thread safe list of stalled thread "Hints". **/
 	TLockFreePointerListUnordered<FTaskThreadBase>		StalledUnnamedThreads; 
 
+	/** Array of callbacks to call before shutdown. **/
+	TArray<TFunction<void()> > ShutdownCallbacks;
+
 #if USE_NEW_LOCK_FREE_LISTS
 	#if USE_INTRUSIVE_TASKQUEUES
 		FLockFreePointerListFIFOIntrusive<FBaseGraphTask>		IncomingAnyThreadTasks;
diff --git a/Engine/Source/Runtime/Core/Private/HAL/ConsoleManager.cpp b/Engine/Source/Runtime/Core/Private/HAL/ConsoleManager.cpp
--- a/Engine/Source/Runtime/Core/Private/HAL/ConsoleManager.cpp
+++ b/Engine/Source/Runtime/Core/Private/HAL/ConsoleManager.cpp
@@ -576,7 +576,7 @@ public:
 		delete this; 
 	} 
 
-	virtual bool Execute( const TArray< FString > Args, UWorld* InWorld, FOutputDevice& OutputDevice ) override
+	virtual bool Execute( const TArray< FString >& Args, UWorld* InWorld, FOutputDevice& OutputDevice ) override
 	{
 		// NOTE: Args are ignored for FConsoleCommand.  Use FConsoleCommandWithArgs if you need parameters.
 		return Delegate.ExecuteIfBound();
@@ -607,7 +607,7 @@ public:
 		delete this; 
 	} 
 
-	virtual bool Execute( const TArray< FString > Args, UWorld* InWorld, FOutputDevice& OutputDevice ) override
+	virtual bool Execute( const TArray< FString >& Args, UWorld* InWorld, FOutputDevice& OutputDevice ) override
 	{
 		return Delegate.ExecuteIfBound( Args );
 	}
@@ -636,7 +636,7 @@ public:
 		delete this; 
 	} 
 
-	virtual bool Execute( const TArray< FString > Args, UWorld* InWorld, FOutputDevice& OutputDevice ) override
+	virtual bool Execute( const TArray< FString >& Args, UWorld* InWorld, FOutputDevice& OutputDevice ) override
 	{
 		return Delegate.ExecuteIfBound( InWorld );
 	}
@@ -665,7 +665,7 @@ public:
 		delete this; 
 	} 
 
-	virtual bool Execute( const TArray< FString > Args, UWorld* InWorld, FOutputDevice& OutputDevice ) override
+	virtual bool Execute( const TArray< FString >& Args, UWorld* InWorld, FOutputDevice& OutputDevice ) override
 	{
 		return Delegate.ExecuteIfBound( Args, InWorld );
 	}
@@ -694,7 +694,7 @@ public:
 		delete this; 
 	} 
 
-	virtual bool Execute( const TArray< FString > Args, UWorld* InWorld, FOutputDevice& OutputDevice ) override
+	virtual bool Execute( const TArray< FString >& Args, UWorld* InWorld, FOutputDevice& OutputDevice ) override
 	{
 		return Delegate.ExecuteIfBound( OutputDevice );
 	}
@@ -722,7 +722,7 @@ public:
 		delete this; 
 	} 
 
-	virtual bool Execute( const TArray< FString > Args, UWorld* InCmdWorld, FOutputDevice& OutputDevice ) override
+	virtual bool Execute( const TArray< FString >& Args, UWorld* InCmdWorld, FOutputDevice& OutputDevice ) override
 	{
 		return false;
 	}
diff --git a/Engine/Source/Runtime/Core/Private/IOS/IOSInputInterface.cpp b/Engine/Source/Runtime/Core/Private/IOS/IOSInputInterface.cpp
--- a/Engine/Source/Runtime/Core/Private/IOS/IOSInputInterface.cpp
+++ b/Engine/Source/Runtime/Core/Private/IOS/IOSInputInterface.cpp
@@ -434,7 +434,7 @@ if (Controller.Previous##Gamepad != nil && Gamepad.GCAxis.value != Controller.Pr
 	}
 }
 
-void FIOSInputInterface::QueueTouchInput(TArray<TouchInput> InTouchEvents)
+void FIOSInputInterface::QueueTouchInput(const TArray<TouchInput>& InTouchEvents)
 {
 	FScopeLock Lock(&CriticalSection);
 
@@ -570,7 +570,6 @@ bool FIOSInputInterface::Exec(UWorld* InWorld, const TCHAR* Cmd, FOutputDevice&
 
 	return bHandledCommand;
 }
-
 bool FIOSInputInterface::IsControllerAssignedToGamepad(int32 ControllerId)
 {
 	return ControllerId < ARRAY_COUNT(Controllers) &&
diff --git a/Engine/Source/Runtime/Core/Private/Stats/Stats2.cpp b/Engine/Source/Runtime/Core/Private/Stats/Stats2.cpp
--- a/Engine/Source/Runtime/Core/Private/Stats/Stats2.cpp
+++ b/Engine/Source/Runtime/Core/Private/Stats/Stats2.cpp
@@ -703,13 +703,14 @@ IStatGroupEnableManager& IStatGroupEnableManager::Get()
 FName FStatNameAndInfo::ToLongName(FName InStatName, char const* InGroup, char const* InCategory, TCHAR const* InDescription)
 {
 	FString LongName;
+	LongName.Reserve(255);
 	if (InGroup)
 	{
 		LongName += TEXT("//");
 		LongName += InGroup;
 		LongName += TEXT("//");
 	}
-	LongName += InStatName.ToString();
+	InStatName.AppendString(LongName);
 	if (InDescription)
 	{
 		LongName += TEXT("///");
diff --git a/Engine/Source/Runtime/Core/Public/Android/AndroidInputInterface.h b/Engine/Source/Runtime/Core/Public/Android/AndroidInputInterface.h
--- a/Engine/Source/Runtime/Core/Public/Android/AndroidInputInterface.h
+++ b/Engine/Source/Runtime/Core/Public/Android/AndroidInputInterface.h
@@ -194,7 +194,7 @@ public:
 	 */
 	void SendControllerEvents();
 
-	static void QueueTouchInput(TArray<TouchInput> InTouchEvents);
+	static void QueueTouchInput(const TArray<TouchInput>& InTouchEvents);
 
 	static void ResetGamepadAssignments();
 	static void ResetGamepadAssignmentToController(int32 ControllerId);
diff --git a/Engine/Source/Runtime/Core/Public/Async/TaskGraphInterfaces.h b/Engine/Source/Runtime/Core/Public/Async/TaskGraphInterfaces.h
--- a/Engine/Source/Runtime/Core/Public/Async/TaskGraphInterfaces.h
+++ b/Engine/Source/Runtime/Core/Public/Async/TaskGraphInterfaces.h
@@ -226,6 +226,11 @@ public:
 		TriggerEventWhenTasksComplete(InEvent, Prerequistes, CurrentThreadIfKnown);
 	}
 
+	/**
+	*	When a set of tasks complete, fire a scoped event
+	*	@param	Callback - function to call prior to shutting down the taskgraph
+	**/
+	virtual void AddShutdownCallback(TFunction<void()>& Callback) = 0;
 };
 
 /** 
diff --git a/Engine/Source/Runtime/Core/Public/HAL/IConsoleManager.h b/Engine/Source/Runtime/Core/Public/HAL/IConsoleManager.h
--- a/Engine/Source/Runtime/Core/Public/HAL/IConsoleManager.h
+++ b/Engine/Source/Runtime/Core/Public/HAL/IConsoleManager.h
@@ -272,7 +272,7 @@ struct IConsoleCommand : public IConsoleObject
 	 * @param	InWorld		World context for this command
 	 * @return	True if the delegate for this command was executed successfully
 	 */
-	virtual bool Execute( const TArray< FString > Args, UWorld* InWorld, class FOutputDevice& OutputDevice ) = 0;
+	virtual bool Execute( const TArray< FString >& Args, UWorld* InWorld, class FOutputDevice& OutputDevice ) = 0;
 };
 
 /**
diff --git a/Engine/Source/Runtime/Core/Public/IOS/IOSInputInterface.h b/Engine/Source/Runtime/Core/Public/IOS/IOSInputInterface.h
--- a/Engine/Source/Runtime/Core/Public/IOS/IOSInputInterface.h
+++ b/Engine/Source/Runtime/Core/Public/IOS/IOSInputInterface.h
@@ -48,7 +48,7 @@ public:
 	 */
 	void SendControllerEvents();
 
-	static void QueueTouchInput(TArray<TouchInput> InTouchEvents);
+	static void QueueTouchInput(const TArray<TouchInput>& InTouchEvents);
 
 	//~ Begin Exec Interface
 	virtual bool Exec(UWorld* InWorld, const TCHAR* Cmd, FOutputDevice& Ar) override;
diff --git a/Engine/Source/Runtime/Core/Public/Math/Sphere.h b/Engine/Source/Runtime/Core/Public/Math/Sphere.h
--- a/Engine/Source/Runtime/Core/Public/Math/Sphere.h
+++ b/Engine/Source/Runtime/Core/Public/Math/Sphere.h
@@ -181,6 +181,9 @@ struct FCapsuleShape
 	FVector Orientation;
 	float Length;
 
+	FCapsuleShape()
+	{}
+
 	FCapsuleShape(FVector InCenter, float InRadius, FVector InOrientation, float InLength) :
 		Center(InCenter),
 		Radius(InRadius),
diff --git a/Engine/Source/Runtime/Core/Public/Misc/Variant.h b/Engine/Source/Runtime/Core/Public/Misc/Variant.h
--- a/Engine/Source/Runtime/Core/Public/Misc/Variant.h
+++ b/Engine/Source/Runtime/Core/Public/Misc/Variant.h
@@ -186,7 +186,7 @@ public:
 	FVariant& operator=( const TArray<uint8> InArray )
 	{
 		Type = EVariantTypes::ByteArray;
-		Value = InArray;
+		Value = MoveTemp(InArray);
 
 		return *this;
 	}
diff --git a/Engine/Source/Runtime/Core/Public/UObject/ObjectVersion.h b/Engine/Source/Runtime/Core/Public/UObject/ObjectVersion.h
--- a/Engine/Source/Runtime/Core/Public/UObject/ObjectVersion.h
+++ b/Engine/Source/Runtime/Core/Public/UObject/ObjectVersion.h
@@ -596,6 +596,9 @@ enum EUnrealEngineObjectUE4Version
 	VER_UE4_INJECT_BLUEPRINT_STRUCT_PIN_CONVERSION_NODES,
 	// Saving tag data for Array Property's inner property
 	VER_UE4_INNER_ARRAY_TAG_INFO,
+	// Fixed duplicating slot node names in skeleton due to skeleton preload on compile
+	VER_UE4_FIX_SLOT_NAME_DUPLICATION,
+
 	// -----<new versions can be added before this line>-------------------------------------------------
 	// - this needs to be the last line (see note below)
 	VER_UE4_AUTOMATIC_VERSION_PLUS_ONE,
diff --git a/Engine/Source/Runtime/Engine/Classes/AI/Navigation/NavigationData.h b/Engine/Source/Runtime/Engine/Classes/AI/Navigation/NavigationData.h
--- a/Engine/Source/Runtime/Engine/Classes/AI/Navigation/NavigationData.h
+++ b/Engine/Source/Runtime/Engine/Classes/AI/Navigation/NavigationData.h
@@ -173,6 +173,10 @@ struct ENGINE_API FNavigationPath : public TSharedFromThis<FNavigationPath, ESPM
 	void Invalidate();
 	void RePathFailed();
 
+	/** Resets all variables describing generated path before attempting new pathfinding call. 
+	  * This function will NOT reset setup variables like goal actor, filter, observer, etc */
+	virtual void ResetForRepath();
+
 	virtual void DebugDraw(const ANavigationData* NavData, FColor PathColor, class UCanvas* Canvas, bool bPersistent, const uint32 NextPathPointIndex = 0) const;
 	
 #if ENABLE_VISUAL_LOG
@@ -250,6 +254,9 @@ struct ENGINE_API FNavigationPath : public TSharedFromThis<FNavigationPath, ESPM
 private:
 	bool DoesPathIntersectBoxImplementation(const FBox& Box, const FVector& StartLocation, uint32 StartingIndex, int32* IntersectingSegmentIndex, FVector* AgentExtent) const;
 
+	/** reset variables describing built path, leaves setup variables required for rebuilding */
+	void InternalResetNavigationPath();
+
 public:
 
 	/** type safe casts */
diff --git a/Engine/Source/Runtime/Engine/Classes/AI/Navigation/RecastNavMesh.h b/Engine/Source/Runtime/Engine/Classes/AI/Navigation/RecastNavMesh.h
--- a/Engine/Source/Runtime/Engine/Classes/AI/Navigation/RecastNavMesh.h
+++ b/Engine/Source/Runtime/Engine/Classes/AI/Navigation/RecastNavMesh.h
@@ -111,7 +111,7 @@ struct ENGINE_API FNavMeshPath : public FNavigationPath
 
 	void ApplyFlags(int32 NavDataFlags);
 
-	void Reset();
+	virtual void ResetForRepath() override;
 
 	/** get flags of path point or corridor poly (depends on bStringPulled flag) */
 	bool GetNodeFlags(int32 NodeIdx, FNavMeshNodeFlags& Flags) const;
@@ -154,6 +154,8 @@ struct ENGINE_API FNavMeshPath : public FNavigationPath
 
 private:
 	bool DoesPathIntersectBoxImplementation(const FBox& Box, const FVector& StartLocation, uint32 StartingIndex, int32* IntersectingSegmentIndex, FVector* AgentExtent) const;
+	void InternalResetNavMeshPath();
+
 public:
 
 #if ENABLE_VISUAL_LOG
diff --git a/Engine/Source/Runtime/Engine/Classes/Animation/AnimCompress.h b/Engine/Source/Runtime/Engine/Classes/Animation/AnimCompress.h
--- a/Engine/Source/Runtime/Engine/Classes/Animation/AnimCompress.h
+++ b/Engine/Source/Runtime/Engine/Classes/Animation/AnimCompress.h
@@ -53,7 +53,7 @@ public:
 	 * @param	bOutput			If false don't generate output or compute memory savings.
 	 * @return					false if a skeleton was needed by the algorithm but not provided.
 	 */
-	ENGINE_API bool Reduce(TArray<class UAnimSequence*> AnimSequences, bool bOutput);
+	ENGINE_API bool Reduce(const TArray<class UAnimSequence*>& AnimSequences, bool bOutput);
 
 protected:
 	/**
diff --git a/Engine/Source/Runtime/Engine/Classes/Animation/AnimMontage.h b/Engine/Source/Runtime/Engine/Classes/Animation/AnimMontage.h
--- a/Engine/Source/Runtime/Engine/Classes/Animation/AnimMontage.h
+++ b/Engine/Source/Runtime/Engine/Classes/Animation/AnimMontage.h
@@ -369,7 +369,7 @@ public:
 
 	/** Delegate function handlers
 	 */
-	void HandleEvents(float PreviousTrackPos, float CurrentTrackPos, const FBranchingPointMarker* BranchingPointMarker);
+	ENGINE_API void HandleEvents(float PreviousTrackPos, float CurrentTrackPos, const FBranchingPointMarker* BranchingPointMarker);
 private:
 	/** Called by blueprint functions that modify the montages current position. */
 	void OnMontagePositionChanged(FName const & ToSectionName);
@@ -501,7 +501,7 @@ public:
 	bool IsValidAdditiveSlot(const FName& SlotNodeName) const;
 
 	/** Get FCompositeSection with InSectionName */
-	FCompositeSection& GetAnimCompositeSection(int32 SectionIndex);
+	ENGINE_API FCompositeSection& GetAnimCompositeSection(int32 SectionIndex);
 	const FCompositeSection& GetAnimCompositeSection(int32 SectionIndex) const;
 
 	// @todo document
diff --git a/Engine/Source/Runtime/Engine/Classes/Animation/AnimNode_StateMachine.h b/Engine/Source/Runtime/Engine/Classes/Animation/AnimNode_StateMachine.h
--- a/Engine/Source/Runtime/Engine/Classes/Animation/AnimNode_StateMachine.h
+++ b/Engine/Source/Runtime/Engine/Classes/Animation/AnimNode_StateMachine.h
@@ -172,6 +172,8 @@ private:
 
 	TArray<FPoseContext*> StateCachedPoses;
 
+	FGraphTraversalCounter UpdateCounter;
+
 	TArray<FGraphTraversalCounter> StateCacheBoneCounters;
 
 public:
diff --git a/Engine/Source/Runtime/Engine/Classes/Animation/AnimSequence.h b/Engine/Source/Runtime/Engine/Classes/Animation/AnimSequence.h
--- a/Engine/Source/Runtime/Engine/Classes/Animation/AnimSequence.h
+++ b/Engine/Source/Runtime/Engine/Classes/Animation/AnimSequence.h
@@ -123,8 +123,8 @@ struct ENGINE_API FAnimSequenceTrackContainer
 
 	void Initialize(TArray<FName> InTrackNames)
 	{
-		TrackNames = InTrackNames;
-		int32 NumNode = TrackNames.Num();
+		TrackNames = MoveTemp(InTrackNames);
+		const int32 NumNode = TrackNames.Num();
 		AnimationTracks.Empty(NumNode);
 		AnimationTracks.AddZeroed(NumNode);
 	}
diff --git a/Engine/Source/Runtime/Engine/Classes/Animation/Skeleton.h b/Engine/Source/Runtime/Engine/Classes/Animation/Skeleton.h
--- a/Engine/Source/Runtime/Engine/Classes/Animation/Skeleton.h
+++ b/Engine/Source/Runtime/Engine/Classes/Animation/Skeleton.h
@@ -315,7 +315,7 @@ private:
 	/** SlotName to GroupName TMap, only at runtime, not serialized. **/
 	TMap<FName, FName> SlotToGroupNameMap;
 
-	void BuildSlotToGroupMap();
+	void BuildSlotToGroupMap(bool bInRemoveDuplicates = false);
 
 public:
 	ENGINE_API FAnimSlotGroup* FindAnimSlotGroup(const FName& InGroupName);
diff --git a/Engine/Source/Runtime/Engine/Classes/Components/DestructibleComponent.h b/Engine/Source/Runtime/Engine/Classes/Components/DestructibleComponent.h
--- a/Engine/Source/Runtime/Engine/Classes/Components/DestructibleComponent.h
+++ b/Engine/Source/Runtime/Engine/Classes/Components/DestructibleComponent.h
@@ -147,7 +147,7 @@ public:
 	//~ Begin SkinnedMeshComponent Interface.
 	virtual bool ShouldUpdateTransform(bool bLODHasChanged) const override;
 	virtual void RefreshBoneTransforms(FActorComponentTickFunction* TickFunction = NULL) override;
-	virtual void SetSkeletalMesh(USkeletalMesh* InSkelMesh) override;
+	virtual void SetSkeletalMesh(USkeletalMesh* InSkelMesh, bool bReinitPose = true) override;
 	virtual FTransform GetSocketTransform(FName InSocketName, ERelativeTransformSpace TransformSpace = RTS_World) const override;
 	//~ End SkinnedMeshComponent Interface.
 
diff --git a/Engine/Source/Runtime/Engine/Classes/Components/PrimitiveComponent.h b/Engine/Source/Runtime/Engine/Classes/Components/PrimitiveComponent.h
--- a/Engine/Source/Runtime/Engine/Classes/Components/PrimitiveComponent.h
+++ b/Engine/Source/Runtime/Engine/Classes/Components/PrimitiveComponent.h
@@ -570,12 +570,12 @@ public:
 	
 	/** 
 	 * Finish tracking an overlap interaction that is no longer occurring between this component and the component specified. 
-	 * @param OtherComp - The component of the other actor to stop overlapping
-	 * @param bDoNotifies - True to dispatch appropriate begin/end overlap notifications when these events occur.
-	 * @param bNoNotifySelf	- True to skip end overlap notifications to this component's.  Does not affect notifications to OtherComp's actor.
+	 * @param OtherComp The component of the other actor to stop overlapping
+	 * @param bDoNotifies True to dispatch appropriate begin/end overlap notifications when these events occur.
+	 * @param bSkipNotifySelf True to skip end overlap notifications to this component's.  Does not affect notifications to OtherComp's actor.
 	 * @see [Overlap Events](https://docs.unrealengine.com/latest/INT/Engine/Physics/Collision/index.html#overlapandgenerateoverlapevents)
 	 */
-	void EndComponentOverlap(const FOverlapInfo& OtherOverlap, bool bDoNotifies=true, bool bNoNotifySelf=false);
+	void EndComponentOverlap(const FOverlapInfo& OtherOverlap, bool bDoNotifies=true, bool bSkipNotifySelf=false);
 
 	/**
 	 * Check whether this component is overlapping another component.
@@ -1081,7 +1081,7 @@ private:
 
 public:
 	void SetLODParentPrimitive(UPrimitiveComponent* InLODParentPrimitive);
-	UPrimitiveComponent* GetLODParentPrimitive();
+	UPrimitiveComponent* GetLODParentPrimitive() const;
 
 #if WITH_EDITOR
 	virtual const int32 GetNumUncachedStaticLightingInteractions() const override; // recursive function
@@ -1095,6 +1095,7 @@ public:
 	virtual bool ShouldCreatePhysicsState() const override;
 	virtual bool HasValidPhysicsState() const override;
 	virtual class FActorComponentInstanceData* GetComponentInstanceData() const override;
+	virtual void OnComponentDestroyed() override;
 	//~ End UActorComponent Interface
 
 	/** @return true if the owner is selected and this component is selectable */
@@ -1675,6 +1676,9 @@ protected:
 	/** Called when the BodyInstance ResponseToChannels, CollisionEnabled or bNotifyRigidBodyCollision changes, in case subclasses want to use that information. */
 	virtual void OnComponentCollisionSettingsChanged();
 
+	/** Ends all current component overlaps. Generally used when destroying this component or when it can no longer generate overlaps. */
+	void ClearComponentOverlaps(bool bDoNotifies, bool bSkipNotifySelf);
+
 private:
 	
 	/**
@@ -1869,13 +1873,17 @@ class ENGINE_API FPrimitiveComponentInstanceData : public FSceneComponentInstanc
 {
 public:
 	FPrimitiveComponentInstanceData(const UPrimitiveComponent* SourceComponent);
-			
 	virtual ~FPrimitiveComponentInstanceData()
 	{}
 
 	virtual void ApplyToComponent(UActorComponent* Component, const ECacheApplyPhase CacheApplyPhase) override;
+	virtual void FindAndReplaceInstances(const TMap<UObject*, UObject*>& OldToNewInstanceMap) override;
+	virtual void AddReferencedObjects(FReferenceCollector& Collector) override;
 
 	bool ContainsData() const;
+
+private:
+	UPrimitiveComponent* LODParent;
 };
 
 
diff --git a/Engine/Source/Runtime/Engine/Classes/Components/SceneComponent.h b/Engine/Source/Runtime/Engine/Classes/Components/SceneComponent.h
--- a/Engine/Source/Runtime/Engine/Classes/Components/SceneComponent.h
+++ b/Engine/Source/Runtime/Engine/Classes/Components/SceneComponent.h
@@ -16,12 +16,12 @@ struct ENGINE_API FOverlapInfo
 	FOverlapInfo()
 	{}
 
-	FOverlapInfo(const FHitResult& InSweepResult)
+	explicit FOverlapInfo(const FHitResult& InSweepResult)
 		: bFromSweep(true), OverlapInfo(InSweepResult)
 	{
 	}
 
-	FOverlapInfo(class UPrimitiveComponent* InComponent, int32 InBodyIndex = INDEX_NONE);
+	explicit FOverlapInfo(class UPrimitiveComponent* InComponent, int32 InBodyIndex = INDEX_NONE);
 	
 	int32 GetBodyIndex() const { return OverlapInfo.Item;  }
 
diff --git a/Engine/Source/Runtime/Engine/Classes/Components/ShapeComponent.h b/Engine/Source/Runtime/Engine/Classes/Components/ShapeComponent.h
--- a/Engine/Source/Runtime/Engine/Classes/Components/ShapeComponent.h
+++ b/Engine/Source/Runtime/Engine/Classes/Components/ShapeComponent.h
@@ -4,6 +4,11 @@
 #include "Components/PrimitiveComponent.h"
 #include "ShapeComponent.generated.h"
 
+namespace physx
+{
+	class PxShape;
+}
+
 /**
  * ShapeComponent is a PrimitiveComponent that is represented by a simple geometrical shape (sphere, capsule, box, etc).
  */
@@ -22,15 +27,39 @@ class ENGINE_API UShapeComponent : public UPrimitiveComponent
 
 	/** Only show this component if the actor is selected */
 	UPROPERTY()
-	uint32 bDrawOnlyIfSelected:1;
+	uint8 bDrawOnlyIfSelected:1;
 
 	/** If true it allows Collision when placing even if collision is not enabled*/
 	UPROPERTY()
-	uint32 bShouldCollideWhenPlacing:1;
+	uint8 bShouldCollideWhenPlacing:1;
 
 	/** If set, shape will be exported for navigation as dynamic modifier instead of using regular collision data */
 	UPROPERTY(EditAnywhere, Category = Navigation)
-	uint32 bDynamicObstacle : 1;
+	uint8 bDynamicObstacle : 1;
+
+protected:
+	/** If the body setup can be shared (i.e. there have been no alterations compared to the CDO)*/
+	uint8 bUseArchetypeBodySetup : 1;
+
+	/** Checks if a shared body setup is available (and if we're eligible for it). If successful you must still check for staleness */
+	template<typename ComponentType>
+	bool PrepareSharedBodySetup()
+	{
+		bool bSuccess = bUseArchetypeBodySetup;
+		if (bUseArchetypeBodySetup && ShapeBodySetup == nullptr)
+		{
+			ShapeBodySetup = CastChecked<ComponentType>(GetArchetype())->GetBodySetup();
+			bSuccess = ShapeBodySetup != nullptr;
+		}
+
+		return bSuccess;
+	}
+
+	template <typename ShapeElemType> void AddShapeToGeomArray();
+	template <typename ShapeElemType> void SetShapeToNewGeom(physx::PxShape* PShape);
+	template <typename ShapeElemType> void CreateShapeBodySetupIfNeeded();
+
+public:
 
 	/** Navigation area type (empty = default obstacle) */
 	UPROPERTY(EditAnywhere, Category = Navigation)
@@ -63,8 +92,10 @@ class ENGINE_API UShapeComponent : public UPrimitiveComponent
 
 	/** Update the body setup parameters based on shape information*/
 	virtual void UpdateBodySetup();
-
 };
 
-
-
+enum class EShapeBodySetupHelper
+{
+	InvalidateSharingIfStale,
+	UpdateBodySetup
+};
\ No newline at end of file
diff --git a/Engine/Source/Runtime/Engine/Classes/Components/SkeletalMeshComponent.h b/Engine/Source/Runtime/Engine/Classes/Components/SkeletalMeshComponent.h
--- a/Engine/Source/Runtime/Engine/Classes/Components/SkeletalMeshComponent.h
+++ b/Engine/Source/Runtime/Engine/Classes/Components/SkeletalMeshComponent.h
@@ -297,6 +297,29 @@ struct FSkeletalMeshComponentClothTickFunction : public FTickFunction
 	virtual FString DiagnosticMessage() override;
 };
 
+struct ENGINE_API FClosestPointOnPhysicsAsset
+{
+	/** The closest point in world space */
+	FVector ClosestWorldPosition;
+
+	/** The normal associated with the surface of the closest body */
+	FVector Normal;
+
+	/** The name of the bone associated with the closest body */
+	FName BoneName;
+
+	/** The distance of the closest point and the original world position. 0 Indicates world position is inside the closest body. */
+	float Distance;
+
+	FClosestPointOnPhysicsAsset()
+		: ClosestWorldPosition(FVector::ZeroVector)
+		, Normal(FVector::ZeroVector)
+		, BoneName(NAME_None)
+		, Distance(-1.f)
+	{
+	}
+};
+
 /**
  * SkeletalMeshComponent is used to create an instance of an animated SkeletalMesh asset.
  *
@@ -496,6 +519,10 @@ public:
 	UPROPERTY(EditAnywhere, AdvancedDisplay, BlueprintReadWrite, Category=Animation)
 	uint32 bPauseAnims:1;
 
+	/** On InitAnim should we set to ref pose (if false use first tick of animation data)*/
+	UPROPERTY(EditAnywhere, AdvancedDisplay, Category = Animation)
+	bool bUseRefPoseOnInitAnim;
+
 	/**
 	* Uses skinned data for collision data.
 	*/
@@ -676,7 +703,7 @@ public:
 	virtual void SkelMeshCompOnParticleSystemFinished( class UParticleSystemComponent* PSC );
 
 	class UAnimSingleNodeInstance * GetSingleNodeInstance() const;
-	void InitializeAnimScriptInstance(bool bForceReinit=true);
+	bool InitializeAnimScriptInstance(bool bForceReinit=true);
 
 	/** @return true if wind is enabled */
 	virtual bool IsWindEnabled() const;
@@ -925,6 +952,26 @@ public:
 	virtual void OnComponentCollisionSettingsChanged() override;
 	virtual void SetPhysMaterialOverride(UPhysicalMaterial* NewPhysMaterial) override;
 	virtual float GetDistanceToCollision(const FVector& Point, FVector& ClosestPointOnCollision) const override;
+
+
+	/** 
+	 *	Given a world position, find the closest point on the physics asset. Note that this is independent of collision and welding. This is based purely on animation position
+	 *  @param	WorldPosition				The point we want the closest point to (i.e. for all bodies in the physics asset, find the one that has a point closest to WorldPosition)
+	 *  @param	ClosestPointOnPhysicsAsset	The data associated with the closest point (position, normal, etc...)
+	 *  @param	bApproximate				The closest body is found using bone transform distance instead of body distance. This approximation means the final point is the closest point on a potentially not closest body. This approximation gets worse as the size of Bodies gets bigger.
+	 *  @return	true if we found a closest point
+	 */
+	bool GetClosestPointOnPhysicsAsset(const FVector& WorldPosition, FClosestPointOnPhysicsAsset& ClosestPointOnPhysicsAsset, bool bApproximate) const;
+
+	/** 
+	 *	Given a world position, find the closest point on the physics asset. Note that this is independent of collision and welding. This is based purely on animation position
+	 *  @param	WorldPosition				The point we want the closest point to (i.e. for all bodies in the physics asset, find the one that has a point closest to WorldPosition)
+	 *  @param	ClosestPointOnPhysicsAsset	The data associated with the closest point (position, normal, etc...)
+	 *  @return	true if we found a closest point
+	 */
+	UFUNCTION(BlueprintCallable, Category="Components|SkeletalMesh", meta=(DisplayName="GetClosestPointOnPhysicsAsset", Keywords="closest point"))
+	bool K2_GetClosestPointOnPhysicsAsset(const FVector& WorldPosition, FVector& ClosestWorldPosition, FVector& Normal, FName& BoneName, float& Distance) const;
+
 	virtual bool LineTraceComponent( FHitResult& OutHit, const FVector Start, const FVector End, const FCollisionQueryParams& Params ) override;
 	virtual bool SweepComponent( FHitResult& OutHit, const FVector Start, const FVector End, const FCollisionShape& CollisionShape, bool bTraceComplex=false) override;
 	virtual bool OverlapComponent(const FVector& Pos, const FQuat& Rot, const FCollisionShape& CollisionShape) override;
@@ -957,7 +1004,7 @@ public:
 	virtual void HideBone( int32 BoneIndex, EPhysBodyOp PhysBodyOption ) override;
 	virtual void UnHideBone( int32 BoneIndex ) override;
 	virtual void SetPhysicsAsset(class UPhysicsAsset* NewPhysicsAsset,bool bForceReInit = false) override;
-	virtual void SetSkeletalMesh(class USkeletalMesh* NewMesh) override;
+	virtual void SetSkeletalMesh(class USkeletalMesh* NewMesh, bool bReinitPose = true) override;
 	virtual FVector GetSkinnedVertexPosition(int32 VertexIndex) const override;
 
 	void SetSkeletalMeshWithoutResettingAnimation(class USkeletalMesh* NewMesh);
@@ -1267,7 +1314,7 @@ public:
 	void DebugDrawBones(class UCanvas* Canvas, bool bSimpleBones) const;
 
 protected:
-	bool NeedToSpawnAnimScriptInstance(bool bForceInit) const;
+	bool NeedToSpawnAnimScriptInstance() const;
 	
 private:
 
diff --git a/Engine/Source/Runtime/Engine/Classes/Components/SkinnedMeshComponent.h b/Engine/Source/Runtime/Engine/Classes/Components/SkinnedMeshComponent.h
--- a/Engine/Source/Runtime/Engine/Classes/Components/SkinnedMeshComponent.h
+++ b/Engine/Source/Runtime/Engine/Classes/Components/SkinnedMeshComponent.h
@@ -185,9 +185,6 @@ protected:
 	/** Are we using double buffered blend spaces */
 	bool bDoubleBufferedBlendSpaces;
 
-	/** Are we using double buffered blend spaces */
-	bool bReInitAnimationOnSetSkeletalMeshCalls;
-
 	/** 
 	 * If set, this component has slave pose components that are associated with this 
 	 * Note this is weak object ptr, so it will go away unless you have other strong reference
@@ -442,9 +439,10 @@ public:
 	 * Change the SkeletalMesh that is rendered for this Component. Will re-initialize the animation tree etc. 
 	 *
 	 * @param NewMesh New mesh to set for this component
+	 * @param bReinitPose Whether we should keep current pose or reinitialize.
 	 */
 	UFUNCTION(BlueprintCallable, Category="Components|SkinnedMesh")
-	virtual void SetSkeletalMesh(class USkeletalMesh* NewMesh);
+	virtual void SetSkeletalMesh(class USkeletalMesh* NewMesh, bool bReinitPose = true);
 
 	/** 
 	 * Get Parent Bone of the input bone
diff --git a/Engine/Source/Runtime/Engine/Classes/Engine/SkeletalMesh.h b/Engine/Source/Runtime/Engine/Classes/Engine/SkeletalMesh.h
--- a/Engine/Source/Runtime/Engine/Classes/Engine/SkeletalMesh.h
+++ b/Engine/Source/Runtime/Engine/Classes/Engine/SkeletalMesh.h
@@ -748,6 +748,9 @@ private:
 	UPROPERTY()
 	TArray<class USkeletalMeshSocket*> Sockets;
 
+	/** Cached matrices from GetComposedRefPoseMatrix */
+	TArray<FMatrix> CachedComposedRefPoseMatrices;
+
 public:
 	/**
 	* Initialize the mesh's render resources.
diff --git a/Engine/Source/Runtime/Engine/Classes/Engine/SkeletalMeshSocket.h b/Engine/Source/Runtime/Engine/Classes/Engine/SkeletalMeshSocket.h
--- a/Engine/Source/Runtime/Engine/Classes/Engine/SkeletalMeshSocket.h
+++ b/Engine/Source/Runtime/Engine/Classes/Engine/SkeletalMeshSocket.h
@@ -94,6 +94,9 @@ class USkeletalMeshSocket : public UObject
 	virtual void PostEditChangeProperty(FPropertyChangedEvent& PropertyChangedEvent) override;
 	//~ End UObject Interface
 
+	// Utility function to copy socket 
+	ENGINE_API void CopyFrom(const class USkeletalMeshSocket* OtherSocket);
+
 private: 
 	/** Broadcasts a notification whenever the socket property has changed. */
 	FSocketChangedEvent ChangedEvent;
diff --git a/Engine/Source/Runtime/Engine/Classes/GameFramework/CharacterMovementComponent.h b/Engine/Source/Runtime/Engine/Classes/GameFramework/CharacterMovementComponent.h
--- a/Engine/Source/Runtime/Engine/Classes/GameFramework/CharacterMovementComponent.h
+++ b/Engine/Source/Runtime/Engine/Classes/GameFramework/CharacterMovementComponent.h
@@ -947,6 +947,7 @@ public:
 	virtual void RequestDirectMove(const FVector& MoveVelocity, bool bForceMaxSpeed) override;
 	virtual bool CanStartPathFollowing() const override;
 	virtual bool CanStopPathFollowing() const override;
+	virtual float GetPathFollowingBrakingDistance(float MaxSpeed) const override;
 	//END UNaVMovementComponent Interface
 
 	//Begin UPawnMovementComponent Interface
diff --git a/Engine/Source/Runtime/Engine/Classes/GameFramework/NavMovementComponent.h b/Engine/Source/Runtime/Engine/Classes/GameFramework/NavMovementComponent.h
--- a/Engine/Source/Runtime/Engine/Classes/GameFramework/NavMovementComponent.h
+++ b/Engine/Source/Runtime/Engine/Classes/GameFramework/NavMovementComponent.h
@@ -25,14 +25,26 @@ class ENGINE_API UNavMovementComponent : public UMovementComponent
 	GENERATED_UCLASS_BODY()
 
 	/** Properties that define how the component can move. */
-	UPROPERTY(EditAnywhere, BlueprintReadWrite, Category="Movement Capabilities", meta=(DisplayName="Movement Capabilities", Keywords="Nav Agent"))
+	UPROPERTY(EditAnywhere, BlueprintReadWrite, Category = NavMovement, meta = (DisplayName = "Movement Capabilities", Keywords = "Nav Agent"))
 	FNavAgentProperties NavAgentProps;
 
 protected:
+	/** Braking distance override used with acceleration driven path following (bUseAccelerationForPaths) */
+	UPROPERTY(EditAnywhere, Category = NavMovement, meta = (EditCondition = "bUseFixedBrakingDistanceForPaths"))
+	float FixedPathBrakingDistance;
+
 	/** If set to true NavAgentProps' radius and height will be updated with Owner's collision capsule size */
-	UPROPERTY(EditAnywhere, Category=MovementComponent)
+	UPROPERTY(EditAnywhere, Category = NavMovement)
 	uint32 bUpdateNavAgentWithOwnersCollision:1;
 
+	/** If set, FixedPathBrakingDistance will be used for path following deceleration */
+	UPROPERTY(EditAnywhere, Category = NavMovement)
+	uint32 bUseAccelerationForPaths : 1;
+
+	/** If set, FixedPathBrakingDistance will be used for path following deceleration */
+	UPROPERTY(EditAnywhere, Category = NavMovement, meta = (EditCondition = "bUseAccelerationForPaths"))
+	uint32 bUseFixedBrakingDistanceForPaths : 1;
+
 private:
 	/** If set, StopActiveMovement call will abort current path following request */
 	uint32 bStopMovementAbortPaths:1;
@@ -71,13 +83,21 @@ public:
 	/** @returns navigation location of controlled actor */
 	FORCEINLINE FVector GetActorNavLocation() const { INavAgentInterface* MyOwner = Cast<INavAgentInterface>(GetOwner()); return MyOwner ? MyOwner->GetNavAgentLocation() : FNavigationSystem::InvalidLocation; }
 
-	/** request direct movement */
+	/** path following: request new velocity */
 	virtual void RequestDirectMove(const FVector& MoveVelocity, bool bForceMaxSpeed);
 
+	/** path following: request new move input (normal vector = full strength) */
+	virtual void RequestPathMove(const FVector& MoveInput);
+
 	/** check if current move target can be reached right now if positions are matching
 	 *  (e.g. performing scripted move and can't stop) */
 	virtual bool CanStopPathFollowing() const;
 
+	/** @returns braking distance for acceleration driven path following */
+	virtual float GetPathFollowingBrakingDistance(float MaxSpeed) const;
+
+	FORCEINLINE bool UseAccelerationForPathFollowing() const { return bUseAccelerationForPaths; }
+
 	/** @returns the NavAgentProps(const) */
 	FORCEINLINE const FNavAgentProperties& GetNavAgentPropertiesRef() const { return NavAgentProps; }
 	/** @returns the NavAgentProps */
diff --git a/Engine/Source/Runtime/Engine/Classes/GameFramework/PawnMovementComponent.h b/Engine/Source/Runtime/Engine/Classes/GameFramework/PawnMovementComponent.h
--- a/Engine/Source/Runtime/Engine/Classes/GameFramework/PawnMovementComponent.h
+++ b/Engine/Source/Runtime/Engine/Classes/GameFramework/PawnMovementComponent.h
@@ -72,6 +72,9 @@ public:
 	/** Notify of collision in case we want to react, such as waking up avoidance or pathing code. */
 	virtual void NotifyBumpedPawn(APawn* BumpedPawn) {}
 
+	// UNavMovementComponent override for input operations
+	virtual void RequestPathMove(const FVector& MoveInput) override;
+
 protected:
 
 	/** Pawn that owns this component. */
diff --git a/Engine/Source/Runtime/Engine/Classes/Particles/Collision/ParticleModuleCollision.h b/Engine/Source/Runtime/Engine/Classes/Particles/Collision/ParticleModuleCollision.h
--- a/Engine/Source/Runtime/Engine/Classes/Particles/Collision/ParticleModuleCollision.h
+++ b/Engine/Source/Runtime/Engine/Classes/Particles/Collision/ParticleModuleCollision.h
@@ -68,10 +68,16 @@ class UParticleModuleCollision : public UParticleModuleCollisionBase
 	 *	This is one-way - particle --> object. The particle does 
 	 *	not have physics applied to it - it just generates an 
 	 *	impulse applied to the object it collides with. 
+	 * NOTE: having this on prevents the code from running off the game thread.
 	 */
 	UPROPERTY(EditAnywhere, Category=Collision)
 	uint32 bApplyPhysics:1;
 
+	/** Any trigger volumes that are hit will be ignored. NOTE: This can be turned off if the TrigerVolume physics object type is not in the CollisionTypes array.
+	* Turning this off is strongly recommended as having it on prevents the code from running off the game thread.*/
+	UPROPERTY(EditAnywhere, Category=Collision)
+	uint32 bIgnoreTriggerVolumes:1;
+
 	/** 
 	 *	The mass of the particle - for use when bApplyPhysics is true. 
 	 *	Value is obtained using the EmitterTime at particle spawn. 
@@ -90,6 +96,7 @@ class UParticleModuleCollision : public UParticleModuleCollisionBase
 	 *	If true, then collisions with Pawns will still react, but 
 	 *	the UsedMaxCollisions count will not be decremented. 
 	 *	(ie., They don't 'count' as collisions)
+	 * NOTE: Having this on prevents the code from running in parallel.
 	 */
 	UPROPERTY(EditAnywhere, Category=Collision)
 	uint32 bPawnsDoNotDecrementCount:1;
@@ -161,10 +168,7 @@ class UParticleModuleCollision : public UParticleModuleCollisionBase
 	virtual uint32 PrepPerInstanceBlock(FParticleEmitterInstance* Owner, void* InstData) override;
 	virtual void SetToSensibleDefaults(UParticleEmitter* Owner) override;
 	virtual bool GenerateLODModuleValues(UParticleModule* SourceModule, float Percentage, UParticleLODLevel* LODLevel) override;
-	virtual bool CanTickInAnyThread() override
-	{
-		return false;
-	}
+	virtual bool CanTickInAnyThread() override;
 	//End UParticleModule Interface
 
 	/**
diff --git a/Engine/Source/Runtime/Engine/Classes/Particles/ParticleSystemComponent.h b/Engine/Source/Runtime/Engine/Classes/Particles/ParticleSystemComponent.h
--- a/Engine/Source/Runtime/Engine/Classes/Particles/ParticleSystemComponent.h
+++ b/Engine/Source/Runtime/Engine/Classes/Particles/ParticleSystemComponent.h
@@ -878,6 +878,7 @@ public:
 	virtual FBoxSphereBounds CalcBounds(const FTransform& LocalToWorld) const override;
 	virtual FPrimitiveSceneProxy* CreateSceneProxy() override;
 	virtual void GetUsedMaterials( TArray<UMaterialInterface*>& OutMaterials ) const override;
+	virtual FBodyInstance* GetBodyInstance(FName BoneName = NAME_None, bool bGetWelded = true) const override;
 	//End UPrimitiveComponent Interface
 
 	//~ Begin USceneComonent Interface
diff --git a/Engine/Source/Runtime/Engine/Classes/PhysicsEngine/BodySetup.h b/Engine/Source/Runtime/Engine/Classes/PhysicsEngine/BodySetup.h
--- a/Engine/Source/Runtime/Engine/Classes/PhysicsEngine/BodySetup.h
+++ b/Engine/Source/Runtime/Engine/Classes/PhysicsEngine/BodySetup.h
@@ -207,6 +207,28 @@ public:
 	 */
 	ENGINE_API void UpdateTriMeshVertices(const TArray<FVector> & NewPositions);
 
+	/**	
+	 * Finds the shortest distance between the body setup and a world position. Input and output are given in world space
+	 * @param	WorldPosition	The point we are trying to get close to
+	 * @param	BodyToWorldTM	The transform to convert BodySetup into world space
+	 * @return					The distance between WorldPosition and the body setup. 0 indicates WorldPosition is inside one of the shapes.
+	 *
+	 * NOTE: This function ignores convex and trimesh data
+	 */
+	ENGINE_API float GetShortestDistanceToPoint(const FVector& WorldPosition, const FTransform& BodyToWorldTM) const;
+
+	/** 
+	 * Finds the closest point in the body setup. Input and outputs are given in world space.
+	 * @param	WorldPosition			The point we are trying to get close to
+	 * @param	BodyToWorldTM			The transform to convert BodySetup into world space
+	 * @param	ClosestWorldPosition	The closest point on the body setup to WorldPosition
+	 * @param	FeatureNormal			The normal of the feature associated with ClosestWorldPosition
+	 * @return							The distance between WorldPosition and the body setup. 0 indicates WorldPosition is inside one of the shapes.
+	 *
+	 * NOTE: This function ignores convex and trimesh data
+	 */
+	ENGINE_API float GetClosestPointAndNormal(const FVector& WorldPosition, const FTransform& BodyToWorldTM, FVector& ClosestWorldPosition, FVector& FeatureNormal) const;
+
 	/**
 	 * Given a format name returns its cooked data.
 	 *
diff --git a/Engine/Source/Runtime/Engine/Classes/PhysicsEngine/BoxElem.h b/Engine/Source/Runtime/Engine/Classes/PhysicsEngine/BoxElem.h
--- a/Engine/Source/Runtime/Engine/Classes/PhysicsEngine/BoxElem.h
+++ b/Engine/Source/Runtime/Engine/Classes/PhysicsEngine/BoxElem.h
@@ -96,5 +96,26 @@ struct FKBoxElem : public FKShapeElem
 
 	ENGINE_API void ScaleElem(FVector DeltaSize, float MinSize);
 
+	ENGINE_API FKBoxElem GetFinalScaled(const FVector& Scale3D, const FTransform& RelativeTM) const;
+
 	ENGINE_API static EAggCollisionShape::Type StaticShapeType;
+
+	/**	
+	 * Finds the shortest distance between the element and a world position. Input and output are given in world space
+	 * @param	WorldPosition	The point we are trying to get close to
+	 * @param	BodyToWorldTM	The transform to convert BodySetup into world space
+	 * @return					The distance between WorldPosition and the shape. 0 indicates WorldPosition is inside one of the shapes.
+	 */
+	ENGINE_API float GetShortestDistanceToPoint(const FVector& WorldPosition, const FTransform& BodyToWorldTM) const;
+
+	
+	/**	
+	 * Finds the closest point on the shape given a world position. Input and output are given in world space
+	 * @param	WorldPosition			The point we are trying to get close to
+	 * @param	BodyToWorldTM			The transform to convert BodySetup into world space
+	 * @param	ClosestWorldPosition	The closest point on the shape in world space
+	 * @param	Normal					The normal of the feature associated with ClosestWorldPosition.
+	 * @return							The distance between WorldPosition and the shape. 0 indicates WorldPosition is inside the shape.
+	 */
+	ENGINE_API float GetClosestPointAndNormal(const FVector& WorldPosition, const FTransform& BodyToWorldTM, FVector& ClosestWorldPosition, FVector& Normal) const;
 };
diff --git a/Engine/Source/Runtime/Engine/Classes/PhysicsEngine/SphereElem.h b/Engine/Source/Runtime/Engine/Classes/PhysicsEngine/SphereElem.h
--- a/Engine/Source/Runtime/Engine/Classes/PhysicsEngine/SphereElem.h
+++ b/Engine/Source/Runtime/Engine/Classes/PhysicsEngine/SphereElem.h
@@ -68,4 +68,24 @@ struct FKSphereElem : public FKShapeElem
 	ENGINE_API void ScaleElem(FVector DeltaSize, float MinSize);
 
 	ENGINE_API static EAggCollisionShape::Type StaticShapeType;
+
+	ENGINE_API FKSphereElem GetFinalScaled(const FVector& Scale3D, const FTransform& RelativeTM) const;
+
+	/**	
+	 * Finds the shortest distance between the element and a world position. Input and output are given in world space
+	 * @param	WorldPosition	The point we are trying to get close to
+	 * @param	BodyToWorldTM	The transform to convert BodySetup into world space
+	 * @return					The distance between WorldPosition and the shape. 0 indicates WorldPosition is inside the shape.
+	 */
+	ENGINE_API float GetShortestDistanceToPoint(const FVector& WorldPosition, const FTransform& BodyToWorldTM) const;
+
+	/**	
+	 * Finds the closest point on the shape given a world position. Input and output are given in world space
+	 * @param	WorldPosition			The point we are trying to get close to
+	 * @param	BodyToWorldTM			The transform to convert BodySetup into world space
+	 * @param	ClosestWorldPosition	The closest point on the shape in world space
+	 * @param	Normal					The normal of the feature associated with ClosestWorldPosition.
+	 * @return							The distance between WorldPosition and the shape. 0 indicates WorldPosition is inside the shape.
+	 */
+	ENGINE_API float GetClosestPointAndNormal(const FVector& WorldPosition, const FTransform& BodyToWorldTM, FVector& ClosestWorldPosition, FVector& Normal) const;
 };
\ No newline at end of file
diff --git a/Engine/Source/Runtime/Engine/Classes/PhysicsEngine/SphylElem.h b/Engine/Source/Runtime/Engine/Classes/PhysicsEngine/SphylElem.h
--- a/Engine/Source/Runtime/Engine/Classes/PhysicsEngine/SphylElem.h
+++ b/Engine/Source/Runtime/Engine/Classes/PhysicsEngine/SphylElem.h
@@ -82,5 +82,26 @@ struct FKSphylElem : public FKShapeElem
 
 	ENGINE_API void ScaleElem(FVector DeltaSize, float MinSize);
 
+	ENGINE_API FKSphylElem GetFinalScaled(const FVector& Scale3D, const FTransform& RelativeTM) const;
+
+	/**	
+	 * Finds the shortest distance between the element and a world position. Input and output are given in world space
+	 * @param	WorldPosition	The point we are trying to get close to
+	 * @param	BodyToWorldTM	The transform to convert BodySetup into world space
+	 * @return					The distance between WorldPosition and the shape. 0 indicates WorldPosition is inside one of the shapes.
+	 */
+
+	ENGINE_API float GetShortestDistanceToPoint(const FVector& WorldPosition, const FTransform& BodyToWorldTM) const;
+
+	/**	
+	 * Finds the closest point on the shape given a world position. Input and output are given in world space
+	 * @param	WorldPosition			The point we are trying to get close to
+	 * @param	BodyToWorldTM			The transform to convert BodySetup into world space
+	 * @param	ClosestWorldPosition	The closest point on the shape in world space
+	 * @param	Normal					The normal of the feature associated with ClosestWorldPosition.
+	 * @return							The distance between WorldPosition and the shape. 0 indicates WorldPosition is inside the shape.
+	 */
+	ENGINE_API float GetClosestPointAndNormal(const FVector& WorldPosition, const FTransform& BodyToWorldTM, FVector& ClosestWorldPosition, FVector& Normal) const;
+
 	ENGINE_API static EAggCollisionShape::Type StaticShapeType;
 };
diff --git a/Engine/Source/Runtime/Engine/Classes/Sound/SoundConcurrency.h b/Engine/Source/Runtime/Engine/Classes/Sound/SoundConcurrency.h
--- a/Engine/Source/Runtime/Engine/Classes/Sound/SoundConcurrency.h
+++ b/Engine/Source/Runtime/Engine/Classes/Sound/SoundConcurrency.h
@@ -26,7 +26,7 @@ namespace EMaxConcurrentResolutionRule
 		StopLowestPriority,
 
 		/** Stop the sound that is quietest in the group. */
-		StopQuietist,
+		StopQuietest,
 	};
 }
 
diff --git a/Engine/Source/Runtime/Engine/Classes/Sound/SoundNode.h b/Engine/Source/Runtime/Engine/Classes/Sound/SoundNode.h
--- a/Engine/Source/Runtime/Engine/Classes/Sound/SoundNode.h
+++ b/Engine/Source/Runtime/Engine/Classes/Sound/SoundNode.h
@@ -120,6 +120,9 @@ public:
 		return 0;
 	}
 
+	/** Returns the number of simultaneous sounds this node instance plays back. */
+	virtual int32 GetNumSounds(const UPTRINT NodeWaveInstanceHash, FActiveSound& ActiveSound) const;
+
 
 	/** 
 	 * Editor interface. 
diff --git a/Engine/Source/Runtime/Engine/Classes/Sound/SoundNodeConcatenator.h b/Engine/Source/Runtime/Engine/Classes/Sound/SoundNodeConcatenator.h
--- a/Engine/Source/Runtime/Engine/Classes/Sound/SoundNodeConcatenator.h
+++ b/Engine/Source/Runtime/Engine/Classes/Sound/SoundNodeConcatenator.h
@@ -26,6 +26,7 @@ public:
 	//~ Begin USoundNode Interface. 
 	virtual bool NotifyWaveInstanceFinished( struct FWaveInstance* WaveInstance ) override;
 	virtual float GetDuration( void ) override;
+	virtual int32 GetNumSounds(const UPTRINT NodeWaveInstanceHash, FActiveSound& ActiveSound) const;
 	virtual void ParseNodes( FAudioDevice* AudioDevice, const UPTRINT NodeWaveInstanceHash, FActiveSound& ActiveSound, const FSoundParseParameters& ParseParams, TArray<FWaveInstance*>& WaveInstances ) override;
 	virtual int32 GetMaxChildNodes() const override 
 	{ 
diff --git a/Engine/Source/Runtime/Engine/Classes/Sound/SoundNodeDistanceCrossFade.h b/Engine/Source/Runtime/Engine/Classes/Sound/SoundNodeDistanceCrossFade.h
--- a/Engine/Source/Runtime/Engine/Classes/Sound/SoundNodeDistanceCrossFade.h
+++ b/Engine/Source/Runtime/Engine/Classes/Sound/SoundNodeDistanceCrossFade.h
@@ -98,6 +98,7 @@ public:
 	virtual void SetChildNodes(TArray<USoundNode*>& InChildNodes) override;
 #endif //WITH_EDITOR
 	virtual float MaxAudibleDistance( float CurrentMaxDistance ) override;
+	virtual int32 GetNumSounds(const UPTRINT NodeWaveInstanceHash, FActiveSound& ActiveSound) const;
 	//~ End USoundNode Interface. 
 
 	virtual float GetCurrentDistance(FAudioDevice* AudioDevice, FActiveSound& ActiveSound, const FSoundParseParameters& ParseParams) const;
diff --git a/Engine/Source/Runtime/Engine/Classes/Sound/SoundNodeWavePlayer.h b/Engine/Source/Runtime/Engine/Classes/Sound/SoundNodeWavePlayer.h
--- a/Engine/Source/Runtime/Engine/Classes/Sound/SoundNodeWavePlayer.h
+++ b/Engine/Source/Runtime/Engine/Classes/Sound/SoundNodeWavePlayer.h
@@ -39,7 +39,8 @@ public:
 	//~ Begin USoundNode Interface
 	virtual int32 GetMaxChildNodes() const override;
 	virtual float GetDuration() override;
-	virtual void ParseNodes( FAudioDevice* AudioDevice, const UPTRINT NodeWaveInstanceHash, FActiveSound& ActiveSound, const FSoundParseParameters& ParseParams, TArray<FWaveInstance*>& WaveInstances ) override;
+	virtual int32 GetNumSounds(const UPTRINT NodeWaveInstanceHash, FActiveSound& ActiveSound) const { return 1; }
+	virtual void ParseNodes(FAudioDevice* AudioDevice, const UPTRINT NodeWaveInstanceHash, FActiveSound& ActiveSound, const FSoundParseParameters& ParseParams, TArray<FWaveInstance*>& WaveInstances) override;
 #if WITH_EDITOR
 	virtual FText GetTitle() const override;
 #endif
diff --git a/Engine/Source/Runtime/Engine/Private/AI/Navigation/AbstractNavData.cpp b/Engine/Source/Runtime/Engine/Private/AI/Navigation/AbstractNavData.cpp
--- a/Engine/Source/Runtime/Engine/Private/AI/Navigation/AbstractNavData.cpp
+++ b/Engine/Source/Runtime/Engine/Private/AI/Navigation/AbstractNavData.cpp
@@ -48,7 +48,15 @@ FPathFindingResult AAbstractNavData::FindPathAbstract(const FNavAgentProperties&
 	}
 
 	FPathFindingResult Result;
-	Result.Path = Query.PathInstanceToFill.IsValid() ? Query.PathInstanceToFill : Self->CreatePathInstance<FAbstractNavigationPath>(Query.Owner.Get());
+	if (Query.PathInstanceToFill.IsValid())
+	{
+		Result.Path = Query.PathInstanceToFill;
+		Result.Path->ResetForRepath();
+	}
+	else
+	{
+		Result.Path = Self->CreatePathInstance<FAbstractNavigationPath>(Query.Owner.Get());
+	}
 
 	Result.Path->GetPathPoints().Reset();
 	Result.Path->GetPathPoints().Add(FNavPathPoint(Query.StartLocation));
diff --git a/Engine/Source/Runtime/Engine/Private/AI/Navigation/NavigationPath.cpp b/Engine/Source/Runtime/Engine/Private/AI/Navigation/NavigationPath.cpp
--- a/Engine/Source/Runtime/Engine/Private/AI/Navigation/NavigationPath.cpp
+++ b/Engine/Source/Runtime/Engine/Private/AI/Navigation/NavigationPath.cpp
@@ -22,25 +22,20 @@ const FNavPathType FNavigationPath::Type;
 
 FNavigationPath::FNavigationPath()
 	: PathType(FNavigationPath::Type)
-	, bUpToDate(true)
-	, bIsReady(false)
-	, bIsPartial(false)
-	, bReachedSearchLimit(false)
 	, LastUpdateTimeStamp(-1.f)	// indicates that it has not been set
 	, GoalActorLocationTetherDistanceSq(-1.f)
 {
-
+	InternalResetNavigationPath();
 }
 
 FNavigationPath::FNavigationPath(const TArray<FVector>& Points, AActor* InBase)
 	: PathType(FNavigationPath::Type)
-	, bUpToDate(true)
-	, bIsReady(true)
-	, bIsPartial(false)
-	, bReachedSearchLimit(false)
 	, LastUpdateTimeStamp(-1.f)	// indicates that it has not been set
 	, GoalActorLocationTetherDistanceSq(-1.f)
 {
+	InternalResetNavigationPath();
+	MarkReady();
+
 	Base = InBase;
 
 	PathPoints.AddZeroed(Points.Num());
@@ -51,6 +46,34 @@ FNavigationPath::FNavigationPath(const TArray<FVector>& Points, AActor* InBase)
 	}
 }
 
+void FNavigationPath::InternalResetNavigationPath()
+{
+	ShortcutNodeRefs.Reset();
+	PathPoints.Reset();
+	Base.Reset();
+
+	bUpToDate = true;
+	bIsReady = false;
+	bIsPartial = false;
+	bReachedSearchLimit = false;
+
+	// keep:
+	// - GoalActor
+	// - GoalActorAsNavAgent
+	// - SourceActor
+	// - SourceActorAsNavAgent
+	// - Querier
+	// - Filter
+	// - PathType
+	// - ObserverDelegate
+	// - bDoAutoUpdateOnInvalidation
+	// - bWaitingForRepath
+	// - NavigationDataUsed
+	// - LastUpdateTimeStamp
+	// - GoalActorLocationTetherDistanceSq
+	// - GoalActorLastLocation
+}
+
 FVector FNavigationPath::GetGoalLocation() const
 {
 	return GoalActor != NULL ? (GoalActorAsNavAgent != NULL ? GoalActorAsNavAgent->GetNavAgentLocation() : GoalActor->GetActorLocation()) : GetEndLocation();
@@ -136,6 +159,11 @@ void FNavigationPath::RePathFailed()
 	bWaitingForRepath = false;
 }
 
+void FNavigationPath::ResetForRepath()
+{
+	InternalResetNavigationPath();
+}
+
 void FNavigationPath::DebugDraw(const ANavigationData* NavData, FColor PathColor, UCanvas* Canvas, bool bPersistent, const uint32 NextPathPointIndex) const
 {
 	static const FColor Grey(100,100,100);
@@ -369,23 +397,33 @@ FString FNavigationPath::GetDescription() const
 const FNavPathType FNavMeshPath::Type;
 	
 FNavMeshPath::FNavMeshPath()
-	: bCorridorEdgesGenerated(false)
-	, bDynamic(false)
-	, bStringPulled(false)
-	, bWantsStringPulling(true)
+	: bWantsStringPulling(true)
 	, bWantsPathCorridor(false)
 {
 	PathType = FNavMeshPath::Type;
+	InternalResetNavMeshPath();
 }
 
-void FNavMeshPath::Reset()
+void FNavMeshPath::ResetForRepath()
+{
+	Super::ResetForRepath();
+	InternalResetNavMeshPath();
+}
+
+void FNavMeshPath::InternalResetNavMeshPath()
 {
-	PathPoints.Reset();
 	PathCorridor.Reset();
 	PathCorridorCost.Reset();
-	bStringPulled = false;
+	CustomLinkIds.Reset();
+	PathCorridorEdges.Reset();
+
 	bCorridorEdgesGenerated = false;
-	bIsPartial = false;
+	bDynamic = false;
+	bStringPulled = false;
+
+	// keep:
+	// - bWantsStringPulling
+	// - bWantsPathCorridor
 }
 
 float FNavMeshPath::GetStringPulledLength(const int32 StartingPoint) const
diff --git a/Engine/Source/Runtime/Engine/Private/AI/Navigation/PImplRecastNavMesh.cpp b/Engine/Source/Runtime/Engine/Private/AI/Navigation/PImplRecastNavMesh.cpp
--- a/Engine/Source/Runtime/Engine/Private/AI/Navigation/PImplRecastNavMesh.cpp
+++ b/Engine/Source/Runtime/Engine/Private/AI/Navigation/PImplRecastNavMesh.cpp
@@ -849,9 +849,6 @@ ENavigationQueryResult::Type FPImplRecastNavMesh::FindPath(const FVector& StartL
 		return ENavigationQueryResult::Error;
 	}
 
-	// initialize output
-	Path.Reset();
-
 	// get path corridor
 	dtQueryResult PathResult;
 	const dtStatus FindPathStatus = NavQuery.findPath(StartPolyID, EndPolyID, &RecastStartPos.X, &RecastEndPos.X, QueryFilter, PathResult, 0);
