Copying //UE4/Dev-Framework to Dev-Main (//UE4/Dev-Main) @ 2879625

diff --git a/Engine/Source/Runtime/AIModule/Classes/Perception/AISense.h b/Engine/Source/Runtime/AIModule/Classes/Perception/AISense.h
--- a/Engine/Source/Runtime/AIModule/Classes/Perception/AISense.h
+++ b/Engine/Source/Runtime/AIModule/Classes/Perception/AISense.h
@@ -18,13 +18,6 @@ class AIMODULE_API UAISense : public UObject
 	static const float SuspendNextUpdate;
 
 protected:
-	UPROPERTY(EditAnywhere, BlueprintReadOnly, Category = Debug)
-	FColor DebugDrawColor;
-
-	/** Sense name used in debug drawing */
-	UPROPERTY(EditAnywhere, BlueprintReadOnly, Category = Debug)
-	FString DebugName;
-
 	/** age past which stimulus of this sense are "forgotten"*/
 	UPROPERTY(EditAnywhere, BlueprintReadOnly, Category = "AI Perception", config)
 	float DefaultExpirationAge;
@@ -148,13 +141,10 @@ protected:
 	/** To be called only for BP-generated classes */
 	void ForceSenseID(FAISenseID SenseID);
 
-public:
-#if !UE_BUILD_SHIPPING
-	//----------------------------------------------------------------------//
-	// DEBUG
-	//----------------------------------------------------------------------//
-	FColor GetDebugColor() const { return DebugDrawColor; }
-	FString GetDebugName() const { return DebugName; }
-	virtual FString GetDebugLegend() const;
-#endif // !UE_BUILD_SHIPPING
+	DEPRECATED(4.12, "Function removed, use AISenseConfig.GetDebugColor() instead.")
+	FColor GetDebugColor() const { return FColor::White; }
+	DEPRECATED(4.12, "Function removed, use AISenseConfig.GetSenseName() instead.")
+	FString GetDebugName() const { return FString(); }
+	DEPRECATED(4.12, "Function removed, use AISenseConfig.DescribeSelfToGameplayDebugger() instead.")
+	virtual FString GetDebugLegend() const { return FString(); };
 };
diff --git a/Engine/Source/Runtime/AIModule/Classes/Perception/AISenseConfig.h b/Engine/Source/Runtime/AIModule/Classes/Perception/AISenseConfig.h
--- a/Engine/Source/Runtime/AIModule/Classes/Perception/AISenseConfig.h
+++ b/Engine/Source/Runtime/AIModule/Classes/Perception/AISenseConfig.h
@@ -3,14 +3,12 @@
 #pragma once
 
 #include "AIPerceptionTypes.h"
-#if !UE_BUILD_SHIPPING
-#	include "Debug/GameplayDebuggerBaseObject.h"
-#endif
 #include "AISenseConfig.generated.h"
 
 class UAISenseImplementation;
 class UCanvas;
 class UAIPerceptionComponent;
+class FGameplayDebuggerCategory;
 
 UCLASS(ClassGroup = AI, abstract, EditInlineNew, config=Game)
 class AIMODULE_API UAISenseConfig : public UObject
@@ -18,29 +16,33 @@ class AIMODULE_API UAISenseConfig : public UObject
 	GENERATED_BODY()
 
 protected:
+	UPROPERTY(EditAnywhere, BlueprintReadOnly, Category = Sense, AdvancedDisplay)
+	FColor DebugColor;
+
 	/** specifies age limit after stimuli generated by this sense become forgotten. 0 means "never" */
-	UPROPERTY(EditDefaultsOnly, BlueprintReadOnly, Category = "Sense", meta = (ClampMin = "0.0", UIMin = "0.0"))
+	UPROPERTY(EditDefaultsOnly, BlueprintReadOnly, Category = Sense, meta = (ClampMin = "0.0", UIMin = "0.0"))
 	float MaxAge;
-	
+
 	/** determines whether given sense starts in an enabled state */
-	UPROPERTY(EditDefaultsOnly, BlueprintReadOnly, Category = "Sense")
+	UPROPERTY(EditDefaultsOnly, BlueprintReadOnly, Category = Sense)
 	uint32 bStartsEnabled : 1;
 	
+	// name of sense
+	mutable FString CachedSenseName;
+
 public:
 	UAISenseConfig(const FObjectInitializer& ObjectInitializer = FObjectInitializer::Get());
 
 	virtual TSubclassOf<UAISense> GetSenseImplementation() const;
 
 	FAISenseID GetSenseID() const;
-	
+	FString GetSenseName() const;
+	FColor GetDebugColor() const { return DebugColor; }
 	void SetMaxAge(const float NewMaxAge) { MaxAge = NewMaxAge; }
 	float GetMaxAge() const { return MaxAge == 0.f ? FAIStimulus::NeverHappenedAge : MaxAge; }	
 	bool IsEnabled() const { return bStartsEnabled; }
 
-#if !UE_BUILD_SHIPPING
-	//----------------------------------------------------------------------//
-	// DEBUG	
-	//----------------------------------------------------------------------//
-	virtual void GetDebugData(TArray<FString>& OnScreenStrings, TArray<FGameplayDebuggerShapeElement>& DebugShapes, const UAIPerceptionComponent& PerceptionComponent) const {}
-#endif // !UE_BUILD_SHIPPING
-};
\ No newline at end of file
+#if WITH_GAMEPLAY_DEBUGGER
+	virtual void DescribeSelfToGameplayDebugger(const UAIPerceptionComponent* PerceptionComponent, FGameplayDebuggerCategory* DebuggerCategory) const;
+#endif // WITH_GAMEPLAY_DEBUGGER
+};
diff --git a/Engine/Source/Runtime/AIModule/Classes/Perception/AISenseConfig_Blueprint.h b/Engine/Source/Runtime/AIModule/Classes/Perception/AISenseConfig_Blueprint.h
--- a/Engine/Source/Runtime/AIModule/Classes/Perception/AISenseConfig_Blueprint.h
+++ b/Engine/Source/Runtime/AIModule/Classes/Perception/AISenseConfig_Blueprint.h
@@ -10,11 +10,11 @@ class UAISense_Blueprint;
 UCLASS(Blueprintable, Abstract)
 class AIMODULE_API UAISenseConfig_Blueprint : public UAISenseConfig
 {
-	GENERATED_BODY()
+	GENERATED_UCLASS_BODY()
 
 public:
 	UPROPERTY(EditDefaultsOnly, BlueprintReadOnly, Category = "Sense", NoClear, config)
 	TSubclassOf<UAISense_Blueprint> Implementation;
 
-	virtual TSubclassOf<UAISense> GetSenseImplementation() const override { return Implementation; }
+	virtual TSubclassOf<UAISense> GetSenseImplementation() const override;
 };
diff --git a/Engine/Source/Runtime/AIModule/Classes/Perception/AISenseConfig_Damage.h b/Engine/Source/Runtime/AIModule/Classes/Perception/AISenseConfig_Damage.h
--- a/Engine/Source/Runtime/AIModule/Classes/Perception/AISenseConfig_Damage.h
+++ b/Engine/Source/Runtime/AIModule/Classes/Perception/AISenseConfig_Damage.h
@@ -11,7 +11,7 @@ class UAISense_Damage;
 UCLASS(meta = (DisplayName = "AI Damage sense config"))
 class AIMODULE_API UAISenseConfig_Damage : public UAISenseConfig
 {
-	GENERATED_BODY()
+	GENERATED_UCLASS_BODY()
 
 public:
 	UPROPERTY(EditDefaultsOnly, BlueprintReadOnly, Category = "Sense", NoClear, config)
diff --git a/Engine/Source/Runtime/AIModule/Classes/Perception/AISenseConfig_Hearing.h b/Engine/Source/Runtime/AIModule/Classes/Perception/AISenseConfig_Hearing.h
--- a/Engine/Source/Runtime/AIModule/Classes/Perception/AISenseConfig_Hearing.h
+++ b/Engine/Source/Runtime/AIModule/Classes/Perception/AISenseConfig_Hearing.h
@@ -10,7 +10,7 @@ class UAISense_Hearing;
 UCLASS(meta = (DisplayName = "AI Hearing config"))
 class AIMODULE_API UAISenseConfig_Hearing : public UAISenseConfig
 {
-	GENERATED_BODY()
+	GENERATED_UCLASS_BODY()
 
 public:
 	UPROPERTY(EditDefaultsOnly, BlueprintReadOnly, Category = "Sense", NoClear, config)
@@ -29,13 +29,9 @@ public:
 	UPROPERTY(EditDefaultsOnly, BlueprintReadOnly, Category = "Sense", config)
 	FAISenseAffiliationFilter DetectionByAffiliation;
 
-	UAISenseConfig_Hearing(const FObjectInitializer& ObjectInitializer = FObjectInitializer::Get());
-
 	virtual TSubclassOf<UAISense> GetSenseImplementation() const override;
-#if !UE_BUILD_SHIPPING
-	//----------------------------------------------------------------------//
-	// DEBUG
-	//----------------------------------------------------------------------//
-	virtual void GetDebugData(TArray<FString>& OnScreenStrings, TArray<FGameplayDebuggerShapeElement>& DebugShapes, const UAIPerceptionComponent& PerceptionComponent) const override;
-#endif // !UE_BUILD_SHIPPING
+
+#if WITH_GAMEPLAY_DEBUGGER
+	virtual void DescribeSelfToGameplayDebugger(const UAIPerceptionComponent* PerceptionComponent, FGameplayDebuggerCategory* DebuggerCategory) const;
+#endif // WITH_GAMEPLAY_DEBUGGER
 };
diff --git a/Engine/Source/Runtime/AIModule/Classes/Perception/AISenseConfig_Prediction.h b/Engine/Source/Runtime/AIModule/Classes/Perception/AISenseConfig_Prediction.h
--- a/Engine/Source/Runtime/AIModule/Classes/Perception/AISenseConfig_Prediction.h
+++ b/Engine/Source/Runtime/AIModule/Classes/Perception/AISenseConfig_Prediction.h
@@ -10,7 +10,7 @@ class UAISense_Prediction;
 UCLASS(meta = (DisplayName = "AI Prediction sense config"))
 class AIMODULE_API UAISenseConfig_Prediction : public UAISenseConfig
 {
-	GENERATED_BODY()
+	GENERATED_UCLASS_BODY()
 public:	
 	virtual TSubclassOf<UAISense> GetSenseImplementation() const override;
 };
\ No newline at end of file
diff --git a/Engine/Source/Runtime/AIModule/Classes/Perception/AISenseConfig_Sight.h b/Engine/Source/Runtime/AIModule/Classes/Perception/AISenseConfig_Sight.h
--- a/Engine/Source/Runtime/AIModule/Classes/Perception/AISenseConfig_Sight.h
+++ b/Engine/Source/Runtime/AIModule/Classes/Perception/AISenseConfig_Sight.h
@@ -10,7 +10,7 @@ class UAISense_Sight;
 UCLASS(meta = (DisplayName = "AI Sight config"))
 class AIMODULE_API UAISenseConfig_Sight : public UAISenseConfig
 {
-	GENERATED_BODY()
+	GENERATED_UCLASS_BODY()
 
 public:
 	UPROPERTY(EditDefaultsOnly, BlueprintReadOnly, Category = "Sense", NoClear, config)
@@ -37,10 +37,8 @@ public:
 	float AutoSuccessRangeFromLastSeenLocation;
 		
 	virtual TSubclassOf<UAISense> GetSenseImplementation() const override;
-#if !UE_BUILD_SHIPPING
-	//----------------------------------------------------------------------//
-	// DEBUG
-	//----------------------------------------------------------------------//
-	virtual void GetDebugData(TArray<FString>& OnScreenStrings, TArray<FGameplayDebuggerShapeElement>& DebugShapes, const UAIPerceptionComponent& PerceptionComponent) const override;
-#endif // !UE_BUILD_SHIPPING
-};
\ No newline at end of file
+
+#if WITH_GAMEPLAY_DEBUGGER
+	virtual void DescribeSelfToGameplayDebugger(const UAIPerceptionComponent* PerceptionComponent, FGameplayDebuggerCategory* DebuggerCategory) const;
+#endif // WITH_GAMEPLAY_DEBUGGER
+};
diff --git a/Engine/Source/Runtime/AIModule/Classes/Perception/AISenseConfig_Touch.h b/Engine/Source/Runtime/AIModule/Classes/Perception/AISenseConfig_Touch.h
--- a/Engine/Source/Runtime/AIModule/Classes/Perception/AISenseConfig_Touch.h
+++ b/Engine/Source/Runtime/AIModule/Classes/Perception/AISenseConfig_Touch.h
@@ -10,7 +10,7 @@ class UAISense_Touch;
 UCLASS(meta = (DisplayName = "AI Touch config"))
 class AIMODULE_API UAISenseConfig_Touch : public UAISenseConfig
 {
-	GENERATED_BODY()
+	GENERATED_UCLASS_BODY()
 public:	
-	virtual TSubclassOf<UAISense> GetSenseImplementation() const override { return UAISense_Touch::StaticClass(); }
+	virtual TSubclassOf<UAISense> GetSenseImplementation() const override;
 };
\ No newline at end of file
diff --git a/Engine/Source/Runtime/AIModule/Classes/Perception/AISense_Hearing.h b/Engine/Source/Runtime/AIModule/Classes/Perception/AISense_Hearing.h
--- a/Engine/Source/Runtime/AIModule/Classes/Perception/AISense_Hearing.h
+++ b/Engine/Source/Runtime/AIModule/Classes/Perception/AISense_Hearing.h
@@ -106,14 +106,4 @@ protected:
 	void OnNewListenerImpl(const FPerceptionListener& NewListener);
 	void OnListenerUpdateImpl(const FPerceptionListener& UpdatedListener);
 	void OnListenerRemovedImpl(const FPerceptionListener& UpdatedListener);
-
-public:
-#if !UE_BUILD_SHIPPING
-	//----------------------------------------------------------------------//
-	// DEBUG
-	//----------------------------------------------------------------------//
-	virtual FString GetDebugLegend() const override;
-	static FColor GetDebugHearingRangeColor() { return FColor::Yellow; }
-	static FColor GetDebugLoSHearingRangeeColor() { return FColorList::Cyan; }
-#endif // !UE_BUILD_SHIPPING
 };
diff --git a/Engine/Source/Runtime/AIModule/Classes/Perception/AISense_Sight.h b/Engine/Source/Runtime/AIModule/Classes/Perception/AISense_Sight.h
--- a/Engine/Source/Runtime/AIModule/Classes/Perception/AISense_Sight.h
+++ b/Engine/Source/Runtime/AIModule/Classes/Perception/AISense_Sight.h
@@ -190,14 +190,4 @@ protected:
 	FORCEINLINE void SortQueries() { SightQueryQueue.Sort(FAISightQuery::FSortPredicate()); }
 
 	float CalcQueryImportance(const FPerceptionListener& Listener, const FVector& TargetLocation, const float SightRadiusSq) const;
-
-public:
-#if !UE_BUILD_SHIPPING
-	//----------------------------------------------------------------------//
-	// DEBUG
-	//----------------------------------------------------------------------//
-	FString GetDebugLegend() const override;
-	static FColor GetDebugSightRangeColor() { return FColor::Green; }
-	static FColor GetDebugLoseSightColor() { return FColorList::NeonPink; }
-#endif // !UE_BUILD_SHIPPING
 };
diff --git a/Engine/Source/Runtime/AIModule/Private/AIModule.cpp b/Engine/Source/Runtime/AIModule/Private/AIModule.cpp
--- a/Engine/Source/Runtime/AIModule/Private/AIModule.cpp
+++ b/Engine/Source/Runtime/AIModule/Private/AIModule.cpp
@@ -12,6 +12,15 @@
 
 #include "AIModule.h"
 
+#if WITH_GAMEPLAY_DEBUGGER
+#include "GameplayDebugger.h"
+#include "GameplayDebugger/GameplayDebuggerCategory_AI.h"
+#include "GameplayDebugger/GameplayDebuggerCategory_BehaviorTree.h"
+#include "GameplayDebugger/GameplayDebuggerCategory_EQS.h"
+#include "GameplayDebugger/GameplayDebuggerCategory_Navmesh.h"
+#include "GameplayDebugger/GameplayDebuggerCategory_Perception.h"
+#endif // WITH_GAMEPLAY_DEBUGGER
+
 #define LOCTEXT_NAMESPACE "AIModule"
 
 DEFINE_LOG_CATEGORY_STATIC(LogAIModule, Log, All);
@@ -38,16 +47,6 @@ IMPLEMENT_MODULE(FAIModule, AIModule)
 
 void FAIModule::StartupModule()
 { 
-	// This code will execute after your module is loaded into memory and after global variables initialization. We needs some place to load GameplayDebugger module so it's best place for it now.
-#if !(UE_BUILD_SHIPPING || UE_BUILD_TEST)
-	bool bEnableGameplayDebuggerPlugin = false;
-	GConfig->GetBool(TEXT("/Script/GameplayDebuggerModule.GameplayDebuggerModuleSettings"), TEXT("EnableGameplayDebuggerPlugin"), bEnableGameplayDebuggerPlugin, GEngineIni);
-	if (!bEnableGameplayDebuggerPlugin)
-	{
-		FModuleManager::LoadModulePtr< IModuleInterface >("GameplayDebugger");
-	}
-#endif
-
 #if WITH_EDITOR 
 	FModuleManager::LoadModulePtr< IModuleInterface >("AITestSuite");
 
@@ -64,6 +63,16 @@ void FAIModule::StartupModule()
 #if WITH_EDITOR && ENABLE_VISUAL_LOG
 	FVisualLogger::Get().RegisterExtension(*EVisLogTags::TAG_EQS, &VisualLoggerExtension);
 #endif
+
+#if WITH_GAMEPLAY_DEBUGGER
+	IGameplayDebugger& GameplayDebuggerModule = IGameplayDebugger::Get();
+	GameplayDebuggerModule.RegisterCategory("AI", IGameplayDebugger::FOnGetCategory::CreateStatic(&FGameplayDebuggerCategory_AI::MakeInstance), EGameplayDebuggerCategoryState::EnabledInGameAndSimulate, 1);
+	GameplayDebuggerModule.RegisterCategory("BehaviorTree", IGameplayDebugger::FOnGetCategory::CreateStatic(&FGameplayDebuggerCategory_BehaviorTree::MakeInstance), EGameplayDebuggerCategoryState::EnabledInGame, 2);
+	GameplayDebuggerModule.RegisterCategory("EQS", IGameplayDebugger::FOnGetCategory::CreateStatic(&FGameplayDebuggerCategory_EQS::MakeInstance));
+	GameplayDebuggerModule.RegisterCategory("Navmesh", IGameplayDebugger::FOnGetCategory::CreateStatic(&FGameplayDebuggerCategory_Navmesh::MakeInstance), EGameplayDebuggerCategoryState::Disabled, 0);
+	GameplayDebuggerModule.RegisterCategory("Perception", IGameplayDebugger::FOnGetCategory::CreateStatic(&FGameplayDebuggerCategory_Perception::MakeInstance));
+	GameplayDebuggerModule.NotifyCategoriesChanged();
+#endif
 }
 
 void FAIModule::ShutdownModule()
@@ -73,6 +82,19 @@ void FAIModule::ShutdownModule()
 #if WITH_EDITOR && ENABLE_VISUAL_LOG
 	FVisualLogger::Get().UnregisterExtension(*EVisLogTags::TAG_EQS, &VisualLoggerExtension);
 #endif
+
+#if WITH_GAMEPLAY_DEBUGGER
+	if (IGameplayDebugger::IsAvailable())
+	{
+		IGameplayDebugger& GameplayDebuggerModule = IGameplayDebugger::Get();
+		GameplayDebuggerModule.UnregisterCategory("AI");
+		GameplayDebuggerModule.UnregisterCategory("BehaviorTree");
+		GameplayDebuggerModule.UnregisterCategory("EQS");
+		GameplayDebuggerModule.UnregisterCategory("Navmesh");
+		GameplayDebuggerModule.UnregisterCategory("Perception");
+		GameplayDebuggerModule.NotifyCategoriesChanged();
+	}
+#endif
 }
 
 UAISystemBase* FAIModule::CreateAISystemInstance(UWorld* World)
diff --git a/Engine/Source/Runtime/AIModule/Private/AISystem.cpp b/Engine/Source/Runtime/AIModule/Private/AISystem.cpp
--- a/Engine/Source/Runtime/AIModule/Private/AISystem.cpp
+++ b/Engine/Source/Runtime/AIModule/Private/AISystem.cpp
@@ -66,6 +66,8 @@ void UAISystem::PostInitProperties()
 			FOnActorSpawned::FDelegate ActorSpawnedDelegate = FOnActorSpawned::FDelegate::CreateUObject(this, &UAISystem::OnActorSpawned);
 			ActorSpawnedDelegateHandle = WorldOuter->AddOnActorSpawnedHandler(ActorSpawnedDelegate);
 		}
+
+		ConditionalLoadDebuggerPlugin();
 	}
 }
 
@@ -203,3 +205,18 @@ UAISystem::FBlackboardDataToComponentsIterator UAISystem::CreateBlackboardDataTo
 {
 	return UAISystem::FBlackboardDataToComponentsIterator(BlackboardDataToComponentsMap, &BlackboardAsset);
 }
+
+void UAISystem::ConditionalLoadDebuggerPlugin()
+{
+#if ENABLED_GAMEPLAY_DEBUGGER
+	if (bEnableDebuggerPlugin)
+	{
+		LoadDebuggerPlugin();
+	}
+#endif
+}
+
+void UAISystem::LoadDebuggerPlugin()
+{
+	FModuleManager::LoadModulePtr< IModuleInterface >("GameplayDebugger");
+}
diff --git a/Engine/Source/Runtime/AIModule/Private/EnvironmentQuery/EnvQueryInstance.cpp b/Engine/Source/Runtime/AIModule/Private/EnvironmentQuery/EnvQueryInstance.cpp
--- a/Engine/Source/Runtime/AIModule/Private/EnvironmentQuery/EnvQueryInstance.cpp
+++ b/Engine/Source/Runtime/AIModule/Private/EnvironmentQuery/EnvQueryInstance.cpp
@@ -321,6 +321,12 @@ void FEnvQueryInstance::ExecuteOneStep(double InCurrentStepTimeLimit)
 				// not doing it always for debugging purposes
 				OptionIndex++;
 				CurrentTest = -1;
+#if USE_EQS_DEBUGGER
+				if (bStoreDebugInfo)
+				{
+					DebugData.Reset();
+				}
+#endif // USE_EQS_DEBUGGER
 			}
 		}
 	}
diff --git a/Engine/Source/Runtime/AIModule/Private/EnvironmentQuery/Tests/EnvQueryTest_Distance.cpp b/Engine/Source/Runtime/AIModule/Private/EnvironmentQuery/Tests/EnvQueryTest_Distance.cpp
--- a/Engine/Source/Runtime/AIModule/Private/EnvironmentQuery/Tests/EnvQueryTest_Distance.cpp
+++ b/Engine/Source/Runtime/AIModule/Private/EnvironmentQuery/Tests/EnvQueryTest_Distance.cpp
@@ -5,12 +5,6 @@
 #include "EnvironmentQuery/Items/EnvQueryItemType_VectorBase.h"
 #include "EnvironmentQuery/Tests/EnvQueryTest_Distance.h"
 
-#if !WITH_DEV_AUTOMATION_TESTS
-
-UEnvQueryTest_Distance::UEnvQueryTest_Distance(const FObjectInitializer& ObjectInitializer) : Super(ObjectInitializer) {}
-
-#else
-
 namespace
 {
 	FORCEINLINE float CalcDistance3D(const FVector& PosA, const FVector& PosB)
@@ -148,5 +142,3 @@ FText UEnvQueryTest_Distance::GetDescriptionDetails() const
 {
 	return DescribeFloatTestParams();
 }
-
-#endif //WITH_DEV_AUTOMATION_TESTS
\ No newline at end of file
diff --git a/Engine/Source/Runtime/AIModule/Private/EnvironmentQuery/Tests/EnvQueryTest_Dot.cpp b/Engine/Source/Runtime/AIModule/Private/EnvironmentQuery/Tests/EnvQueryTest_Dot.cpp
--- a/Engine/Source/Runtime/AIModule/Private/EnvironmentQuery/Tests/EnvQueryTest_Dot.cpp
+++ b/Engine/Source/Runtime/AIModule/Private/EnvironmentQuery/Tests/EnvQueryTest_Dot.cpp
@@ -6,12 +6,6 @@
 #include "EnvironmentQuery/Contexts/EnvQueryContext_Item.h"
 #include "EnvironmentQuery/Tests/EnvQueryTest_Dot.h"
 
-#if !WITH_DEV_AUTOMATION_TESTS
-
-UEnvQueryTest_Dot::UEnvQueryTest_Dot(const FObjectInitializer& ObjectInitializer) : Super(ObjectInitializer) {}
-
-#else
-
 UEnvQueryTest_Dot::UEnvQueryTest_Dot(const FObjectInitializer& ObjectInitializer) : Super(ObjectInitializer)
 {
 	Cost = EEnvTestCost::Low;
@@ -222,5 +216,3 @@ FText UEnvQueryTest_Dot::GetDescriptionDetails() const
 {
 	return DescribeFloatTestParams();
 }
-
-#endif //WITH_DEV_AUTOMATION_TESTS
\ No newline at end of file
diff --git a/Engine/Source/Runtime/AIModule/Private/EnvironmentQuery/Tests/EnvQueryTest_GameplayTags.cpp b/Engine/Source/Runtime/AIModule/Private/EnvironmentQuery/Tests/EnvQueryTest_GameplayTags.cpp
--- a/Engine/Source/Runtime/AIModule/Private/EnvironmentQuery/Tests/EnvQueryTest_GameplayTags.cpp
+++ b/Engine/Source/Runtime/AIModule/Private/EnvironmentQuery/Tests/EnvQueryTest_GameplayTags.cpp
@@ -5,12 +5,6 @@
 #include "EnvironmentQuery/Items/EnvQueryItemType_ActorBase.h"
 #include "GameplayTagAssetInterface.h"
 
-#if !WITH_DEV_AUTOMATION_TESTS
-
-UEnvQueryTest_GameplayTags::UEnvQueryTest_GameplayTags(const FObjectInitializer& ObjectInitializer) : Super(ObjectInitializer) {}
-
-#else
-
 UEnvQueryTest_GameplayTags::UEnvQueryTest_GameplayTags(const FObjectInitializer& ObjectInitializer) : Super(ObjectInitializer)
 {
 	Cost = EEnvTestCost::Low;
@@ -91,5 +85,3 @@ FText UEnvQueryTest_GameplayTags::GetDescriptionDetails() const
 {
 	return GameplayTags.ToMatchingText(TagsToMatch, !BoolValue.DefaultValue);
 }
-
-#endif //WITH_DEV_AUTOMATION_TESTS
\ No newline at end of file
diff --git a/Engine/Source/Runtime/AIModule/Private/EnvironmentQuery/Tests/EnvQueryTest_Overlap.cpp b/Engine/Source/Runtime/AIModule/Private/EnvironmentQuery/Tests/EnvQueryTest_Overlap.cpp
--- a/Engine/Source/Runtime/AIModule/Private/EnvironmentQuery/Tests/EnvQueryTest_Overlap.cpp
+++ b/Engine/Source/Runtime/AIModule/Private/EnvironmentQuery/Tests/EnvQueryTest_Overlap.cpp
@@ -5,12 +5,6 @@
 #include "EnvironmentQuery/Items/EnvQueryItemType_VectorBase.h"
 #include "EnvironmentQuery/Tests/EnvQueryTest_Overlap.h"
 
-#if !WITH_DEV_AUTOMATION_TESTS
-
-UEnvQueryTest_Overlap::UEnvQueryTest_Overlap(const FObjectInitializer& ObjectInitializer) : Super(ObjectInitializer) {}
-
-#else
-
 #define LOCTEXT_NAMESPACE "EnvQueryGenerator"
 
 UEnvQueryTest_Overlap::UEnvQueryTest_Overlap(const FObjectInitializer& ObjectInitializer) : Super(ObjectInitializer)
@@ -132,5 +126,3 @@ bool UEnvQueryTest_Overlap::RunOverlapBlocking(const FVector& ItemPos, const FCo
 }
 
 #undef LOCTEXT_NAMESPACE
-
-#endif //WITH_DEV_AUTOMATION_TESTS
\ No newline at end of file
diff --git a/Engine/Source/Runtime/AIModule/Private/EnvironmentQuery/Tests/EnvQueryTest_Pathfinding.cpp b/Engine/Source/Runtime/AIModule/Private/EnvironmentQuery/Tests/EnvQueryTest_Pathfinding.cpp
--- a/Engine/Source/Runtime/AIModule/Private/EnvironmentQuery/Tests/EnvQueryTest_Pathfinding.cpp
+++ b/Engine/Source/Runtime/AIModule/Private/EnvironmentQuery/Tests/EnvQueryTest_Pathfinding.cpp
@@ -9,12 +9,6 @@
 
 #define LOCTEXT_NAMESPACE "EnvQueryGenerator"
 
-#if !WITH_DEV_AUTOMATION_TESTS
-
-UEnvQueryTest_Pathfinding::UEnvQueryTest_Pathfinding(const FObjectInitializer& ObjectInitializer) : Super(ObjectInitializer) {}
-
-#else
-
 UEnvQueryTest_Pathfinding::UEnvQueryTest_Pathfinding(const FObjectInitializer& ObjectInitializer) : Super(ObjectInitializer)
 {
 	Context = UEnvQueryContext_Querier::StaticClass();
@@ -234,6 +228,4 @@ ANavigationData* UEnvQueryTest_Pathfinding::FindNavigationData(UNavigationSystem
 	return NavSys.GetMainNavData(FNavigationSystem::DontCreate);
 }
 
-#endif //WITH_DEV_AUTOMATION_TESTS
-
 #undef LOCTEXT_NAMESPACE
diff --git a/Engine/Source/Runtime/AIModule/Private/EnvironmentQuery/Tests/EnvQueryTest_PathfindingBatch.cpp b/Engine/Source/Runtime/AIModule/Private/EnvironmentQuery/Tests/EnvQueryTest_PathfindingBatch.cpp
--- a/Engine/Source/Runtime/AIModule/Private/EnvironmentQuery/Tests/EnvQueryTest_PathfindingBatch.cpp
+++ b/Engine/Source/Runtime/AIModule/Private/EnvironmentQuery/Tests/EnvQueryTest_PathfindingBatch.cpp
@@ -12,12 +12,6 @@
 #include "AI/Navigation/PImplRecastNavMesh.h"
 #endif
 
-#if !WITH_DEV_AUTOMATION_TESTS
-
-UEnvQueryTest_PathfindingBatch::UEnvQueryTest_PathfindingBatch(const FObjectInitializer& ObjectInitializer) : Super(ObjectInitializer) {}
-
-#else
-
 #define LOCTEXT_NAMESPACE "EnvQueryGenerator"
 
 UEnvQueryTest_PathfindingBatch::UEnvQueryTest_PathfindingBatch(const FObjectInitializer& ObjectInitializer) : Super(ObjectInitializer)
@@ -196,5 +190,3 @@ FText UEnvQueryTest_PathfindingBatch::GetDescriptionTitle() const
 }
 
 #undef LOCTEXT_NAMESPACE
-
-#endif //WITH_DEV_AUTOMATION_TESTS
\ No newline at end of file
diff --git a/Engine/Source/Runtime/AIModule/Private/EnvironmentQuery/Tests/EnvQueryTest_Project.cpp b/Engine/Source/Runtime/AIModule/Private/EnvironmentQuery/Tests/EnvQueryTest_Project.cpp
--- a/Engine/Source/Runtime/AIModule/Private/EnvironmentQuery/Tests/EnvQueryTest_Project.cpp
+++ b/Engine/Source/Runtime/AIModule/Private/EnvironmentQuery/Tests/EnvQueryTest_Project.cpp
@@ -5,12 +5,6 @@
 #include "EnvironmentQuery/Tests/EnvQueryTest_Project.h"
 #include "EnvironmentQuery/EnvQueryTraceHelpers.h"
 
-#if !WITH_DEV_AUTOMATION_TESTS
-
-UEnvQueryTest_Project::UEnvQueryTest_Project(const FObjectInitializer& ObjectInitializer) : Super(ObjectInitializer) {}
-
-#else
-
 UEnvQueryTest_Project::UEnvQueryTest_Project(const FObjectInitializer& ObjectInitializer) : Super(ObjectInitializer)
 {
 	Cost = EEnvTestCost::Medium;
@@ -106,5 +100,3 @@ FText UEnvQueryTest_Project::GetDescriptionDetails() const
 {
 	return DescribeBoolTestParams(TEXT("projected"));
 }
-
-#endif //WITH_DEV_AUTOMATION_TESTS
\ No newline at end of file
diff --git a/Engine/Source/Runtime/AIModule/Private/EnvironmentQuery/Tests/EnvQueryTest_Random.cpp b/Engine/Source/Runtime/AIModule/Private/EnvironmentQuery/Tests/EnvQueryTest_Random.cpp
--- a/Engine/Source/Runtime/AIModule/Private/EnvironmentQuery/Tests/EnvQueryTest_Random.cpp
+++ b/Engine/Source/Runtime/AIModule/Private/EnvironmentQuery/Tests/EnvQueryTest_Random.cpp
@@ -4,12 +4,6 @@
 #include "EnvironmentQuery/Tests/EnvQueryTest_Random.h"
 #include "EnvironmentQuery/Items/EnvQueryItemType.h"
 
-#if !WITH_DEV_AUTOMATION_TESTS
-
-UEnvQueryTest_Random::UEnvQueryTest_Random(const FObjectInitializer& ObjectInitializer) : Super(ObjectInitializer) {}
-
-#else
-
 UEnvQueryTest_Random::UEnvQueryTest_Random(const FObjectInitializer& ObjectInitializer) : Super(ObjectInitializer)
 {
 	Cost = EEnvTestCost::Low;
@@ -43,5 +37,3 @@ FText UEnvQueryTest_Random::GetDescriptionDetails() const
 {
 	return DescribeFloatTestParams();
 }
-
-#endif //WITH_DEV_AUTOMATION_TESTS
\ No newline at end of file
diff --git a/Engine/Source/Runtime/AIModule/Private/EnvironmentQuery/Tests/EnvQueryTest_Trace.cpp b/Engine/Source/Runtime/AIModule/Private/EnvironmentQuery/Tests/EnvQueryTest_Trace.cpp
--- a/Engine/Source/Runtime/AIModule/Private/EnvironmentQuery/Tests/EnvQueryTest_Trace.cpp
+++ b/Engine/Source/Runtime/AIModule/Private/EnvironmentQuery/Tests/EnvQueryTest_Trace.cpp
@@ -5,12 +5,6 @@
 #include "EnvironmentQuery/Items/EnvQueryItemType_VectorBase.h"
 #include "EnvironmentQuery/Tests/EnvQueryTest_Trace.h"
 
-#if !WITH_DEV_AUTOMATION_TESTS
-
-UEnvQueryTest_Trace::UEnvQueryTest_Trace(const FObjectInitializer& ObjectInitializer) : Super(ObjectInitializer) {}
-
-#else
-
 #define LOCTEXT_NAMESPACE "EnvQueryGenerator"
 
 UEnvQueryTest_Trace::UEnvQueryTest_Trace(const FObjectInitializer& ObjectInitializer) : Super(ObjectInitializer)
@@ -192,5 +186,3 @@ bool UEnvQueryTest_Trace::RunCapsuleTraceFrom(const FVector& ItemPos, const FVec
 }
 
 #undef LOCTEXT_NAMESPACE
-
-#endif //WITH_DEV_AUTOMATION_TESTS
\ No newline at end of file
diff --git a/Engine/Source/Runtime/AIModule/Private/GameplayDebugger/GameplayDebuggerCategory_BehaviorTree.cpp b/Engine/Source/Runtime/AIModule/Private/GameplayDebugger/GameplayDebuggerCategory_BehaviorTree.cpp
new file mode 100644
--- /dev/null
+++ b/Engine/Source/Runtime/AIModule/Private/GameplayDebugger/GameplayDebuggerCategory_BehaviorTree.cpp
@@ -0,0 +1,79 @@
+// Copyright 1998-2016 Epic Games, Inc. All Rights Reserved.
+
+#include "AIModulePrivate.h"
+#include "GameplayDebuggerCategory_BehaviorTree.h"
+
+#if WITH_GAMEPLAY_DEBUGGER
+
+FGameplayDebuggerCategory_BehaviorTree::FGameplayDebuggerCategory_BehaviorTree()
+{
+	SetDataPackReplication<FRepData>(&DataPack);
+}
+
+TSharedRef<FGameplayDebuggerCategory> FGameplayDebuggerCategory_BehaviorTree::MakeInstance()
+{
+	return MakeShareable(new FGameplayDebuggerCategory_BehaviorTree());
+}
+
+void FGameplayDebuggerCategory_BehaviorTree::FRepData::Serialize(FArchive& Ar)
+{
+	Ar << TreeName;
+	Ar << TreeDesc;
+	Ar << BlackboardDesc;
+}
+
+void FGameplayDebuggerCategory_BehaviorTree::CollectData(APlayerController* OwnerPC, AActor* DebugActor)
+{
+	APawn* MyPawn = Cast<APawn>(DebugActor);
+	AAIController* MyController = MyPawn ? Cast<AAIController>(MyPawn->Controller) : nullptr;
+	UBrainComponent* BrainComp = MyController ? MyController->GetBrainComponent() : nullptr;
+	
+	if (IsValid(BrainComp))
+	{
+		DataPack.TreeName = BrainComp->GetName();
+		DataPack.TreeDesc = BrainComp->GetDebugInfoString();
+
+		if (BrainComp->GetBlackboardComponent())
+		{
+			DataPack.BlackboardDesc = BrainComp->GetBlackboardComponent()->GetDebugInfoString(EBlackboardDescription::KeyWithValue);
+		}
+	}
+}
+
+void FGameplayDebuggerCategory_BehaviorTree::DrawData(APlayerController* OwnerPC, FGameplayDebuggerCanvasContext& CanvasContext)
+{
+	if (!DataPack.TreeName.IsEmpty())
+	{
+		CanvasContext.Printf(TEXT("Brain Component: {yellow}%s"), *DataPack.TreeName);
+		CanvasContext.Print(DataPack.TreeDesc);
+
+		TArray<FString> BlackboardLines;
+		DataPack.BlackboardDesc.ParseIntoArrayLines(BlackboardLines, true);
+
+		int32 SavedDefX = CanvasContext.DefaultX;
+		int32 SavedPosY = CanvasContext.CursorY;
+		CanvasContext.DefaultX = CanvasContext.CursorX = 600.0f;
+		CanvasContext.CursorY = CanvasContext.DefaultY;
+
+		for (int32 Idx = 0; Idx < BlackboardLines.Num(); Idx++)
+		{
+			int32 SeparatorIndex = INDEX_NONE;
+			BlackboardLines[Idx].FindChar(TEXT(':'), SeparatorIndex);
+
+			if (SeparatorIndex != INDEX_NONE && Idx)
+			{
+				FString ColoredLine = BlackboardLines[Idx].Left(SeparatorIndex + 1) + FString(TEXT("{yellow}")) + BlackboardLines[Idx].Mid(SeparatorIndex + 1);
+				CanvasContext.Print(ColoredLine);
+			}
+			else
+			{
+				CanvasContext.Print(BlackboardLines[Idx]);
+			}
+		}
+
+		CanvasContext.DefaultX = CanvasContext.CursorX = SavedDefX;
+		CanvasContext.CursorY = SavedPosY;
+	}
+}
+
+#endif // WITH_GAMEPLAY_DEBUGGER
diff --git a/Engine/Source/Runtime/AIModule/Private/GameplayDebugger/GameplayDebuggerCategory_EQS.cpp b/Engine/Source/Runtime/AIModule/Private/GameplayDebugger/GameplayDebuggerCategory_EQS.cpp
new file mode 100644
--- /dev/null
+++ b/Engine/Source/Runtime/AIModule/Private/GameplayDebugger/GameplayDebuggerCategory_EQS.cpp
@@ -0,0 +1,344 @@
+// Copyright 1998-2016 Epic Games, Inc. All Rights Reserved.
+
+#include "AIModulePrivate.h"
+
+#include "EnvironmentQuery/EnvQueryManager.h"
+#include "EnvironmentQuery/EQSRenderingComponent.h"
+#include "GameplayDebuggerCategory_EQS.h"
+
+#if WITH_GAMEPLAY_DEBUGGER
+
+static TAutoConsoleVariable<int32> CVarEQSDetailsOnHUD(
+	TEXT("ai.vd.EQSOnHUD"),
+	0,
+	TEXT("Enable or disable EQS details table on screen.\n")
+	TEXT(" 0: Disable details\n")
+	TEXT(" 1: Enable details (default)"),
+	ECVF_Cheat);
+
+FGameplayDebuggerCategory_EQS::FGameplayDebuggerCategory_EQS()
+{
+	MaxQueries = 5;
+	MaxItemTableRows = 20;
+	ShownQueryIndex = 0;
+	CollectDataInterval = 2.0f;
+
+	BindKeyPress(TEXT("Multiply"), this, &FGameplayDebuggerCategory_EQS::CycleShownQueries);
+
+#if USE_EQS_DEBUGGER
+	SetDataPackReplication<FRepData>(&DataPack);
+#endif
+}
+
+TSharedRef<FGameplayDebuggerCategory> FGameplayDebuggerCategory_EQS::MakeInstance()
+{
+	return MakeShareable(new FGameplayDebuggerCategory_EQS());
+}
+
+#if USE_EQS_DEBUGGER
+void FGameplayDebuggerCategory_EQS::FRepData::Serialize(FArchive& Ar)
+{
+	Ar << QueryDebugData;
+}
+#endif // USE_EQS_DEBUGGER
+
+void FGameplayDebuggerCategory_EQS::CollectData(APlayerController* OwnerPC, AActor* DebugActor)
+{
+#if USE_EQS_DEBUGGER
+	UWorld* World = OwnerPC->GetWorld();
+	UEnvQueryManager* QueryManager = World ? UEnvQueryManager::GetCurrent(World) : nullptr;
+	if (QueryManager == nullptr)
+	{
+		return;
+	}
+
+	TArray<FEQSDebugger::FEnvQueryInfo> AuthQueryData = QueryManager->GetDebugger().GetAllQueriesForOwner(DebugActor);
+
+	APawn* DebugPawn = Cast<APawn>(DebugActor);
+	if (DebugPawn && DebugPawn->GetController())
+	{
+		TArray<FEQSDebugger::FEnvQueryInfo>& AuthControllerQueryData = QueryManager->GetDebugger().GetAllQueriesForOwner(DebugPawn->GetController());
+		AuthQueryData.Append(AuthControllerQueryData);
+	}
+
+	struct FEnvQuerySortNewFirst
+	{
+		bool operator()(const FEQSDebugger::FEnvQueryInfo& A, const FEQSDebugger::FEnvQueryInfo& B) const
+		{
+			return (A.Timestamp < B.Timestamp);
+		}
+	};
+
+	AuthQueryData.Sort(FEnvQuerySortNewFirst());
+
+	for (int32 Idx = 0; Idx < AuthQueryData.Num() && DataPack.QueryDebugData.Num() < MaxQueries; Idx++)
+	{
+		FEnvQueryInstance* QueryInstance = AuthQueryData[Idx].Instance.Get();
+		if (QueryInstance)
+		{
+			EQSDebug::FQueryData DebugItem;
+			UEnvQueryDebugHelpers::QueryToDebugData(*QueryInstance, DebugItem);
+
+			DataPack.QueryDebugData.Add(DebugItem);
+		}
+	}
+#endif // USE_EQS_DEBUGGER
+}
+
+void FGameplayDebuggerCategory_EQS::OnDataPackReplicated(int32 DataPackId)
+{
+	MarkRenderStateDirty();
+
+#if USE_EQS_DEBUGGER
+	if (!DataPack.QueryDebugData.IsValidIndex(ShownQueryIndex))
+	{
+		ShownQueryIndex = 0;
+	}
+#endif
+}
+
+FDebugRenderSceneProxy* FGameplayDebuggerCategory_EQS::CreateSceneProxy(const UPrimitiveComponent* InComponent)
+{
+#if USE_EQS_DEBUGGER
+	if (DataPack.QueryDebugData.IsValidIndex(ShownQueryIndex))
+	{
+		const EQSDebug::FQueryData& QueryData = DataPack.QueryDebugData[ShownQueryIndex];
+		const FString ViewFlagName = GetSceneProxyViewFlag();
+		return new FEQSSceneProxy(InComponent, ViewFlagName, QueryData.SolidSpheres, QueryData.Texts);
+	}
+#endif // USE_EQS_DEBUGGER
+
+	return nullptr;
+}
+
+void FGameplayDebuggerCategory_EQS::DrawData(APlayerController* OwnerPC, FGameplayDebuggerCanvasContext& CanvasContext)
+{
+#if USE_EQS_DEBUGGER
+
+	const FString HeaderDesc = (DataPack.QueryDebugData.Num() > 1) ?
+		FString::Printf(TEXT("Queries: {yellow}%d{white}, press {yellow}[*]{white} to cycle through"), DataPack.QueryDebugData.Num()) :
+		FString::Printf(TEXT("Queries: {yellow}%d"), DataPack.QueryDebugData.Num());
+
+	CanvasContext.Print(HeaderDesc);
+	if (DataPack.QueryDebugData.Num() == 0)
+	{
+		return;
+	}
+
+	FString QueryDesc(TEXT("\t"));
+	for (int32 Idx = 0; Idx < DataPack.QueryDebugData.Num(); Idx++)
+	{
+		QueryDesc += (Idx == ShownQueryIndex) ? TEXT("[{green}") : TEXT("[");
+		QueryDesc += DataPack.QueryDebugData[Idx].Name;
+		QueryDesc += (Idx == ShownQueryIndex) ? TEXT("{white}] ") : TEXT("] ");
+	}
+	CanvasContext.Print(QueryDesc);
+
+	if (DataPack.QueryDebugData.IsValidIndex(ShownQueryIndex))
+	{
+		const EQSDebug::FQueryData& ShownQueryData = DataPack.QueryDebugData[ShownQueryIndex];
+
+		CanvasContext.MoveToNewLine();
+		CanvasContext.Printf(TEXT("Timestamp: {yellow}%.3f (~ %.2fs ago)"), ShownQueryData.Timestamp, OwnerPC->GetWorld()->TimeSince(ShownQueryData.Timestamp));
+		CanvasContext.Printf(TEXT("Id: {yellow}%d"), ShownQueryData.Id);
+
+		FString OptionsDesc(TEXT("Options: "));
+		for (int32 Idx = 0; Idx < ShownQueryData.Options.Num(); Idx++)
+		{
+			QueryDesc += (Idx == ShownQueryData.UsedOption) ? TEXT("[{green}") : TEXT("[");
+			QueryDesc += ShownQueryData.Options[Idx];
+			QueryDesc += (Idx == ShownQueryData.UsedOption) ? TEXT("{white}] ") : TEXT("] ");
+		}
+		CanvasContext.Print(OptionsDesc);
+
+		DrawLookedAtItem(ShownQueryData, OwnerPC, CanvasContext);
+
+		CanvasContext.MoveToNewLine();
+		DrawDetailedItemTable(ShownQueryData, CanvasContext);
+	}
+#else // USE_EQS_DEBUGGER
+	CanvasContext.Print(FColor::Red, TEXT("Unable to gather EQS debug data, use build with USE_EQS_DEBUGGER enabled."));
+#endif // USE_EQS_DEBUGGER
+}
+
+#if USE_EQS_DEBUGGER
+void FGameplayDebuggerCategory_EQS::DrawLookedAtItem(const EQSDebug::FQueryData& QueryData, APlayerController* OwnerPC, FGameplayDebuggerCanvasContext& CanvasContext) const
+{
+	FVector CameraLocation;
+	FVector CameraDirection;
+	if (!OwnerPC->GetSpectatorPawn())
+	{
+		FRotator CameraRotation;
+		OwnerPC->GetPlayerViewPoint(CameraLocation, CameraRotation);
+		CameraDirection = CameraRotation.Vector();
+	}
+	else
+	{
+		CameraDirection = CanvasContext.Canvas->SceneView->GetViewDirection();
+		CameraLocation = CanvasContext.Canvas->SceneView->ViewMatrices.ViewOrigin;
+	}
+
+	int32 BestItemIndex = INDEX_NONE;
+	float BestScore = -FLT_MAX;
+	for (int32 Idx = 0; Idx < QueryData.RenderDebugHelpers.Num(); Idx++)
+	{
+		const EQSDebug::FDebugHelper& ItemInfo = QueryData.RenderDebugHelpers[Idx];
+		const FVector DirToItem = ItemInfo.Location - CameraLocation;
+		float DistToItem = DirToItem.Size();
+		if (FMath::IsNearlyZero(DistToItem))
+		{
+			DistToItem = 1.0f;
+		}
+
+		const float ItemScore = FVector::DotProduct(DirToItem, CameraDirection) / DistToItem;
+		if (ItemScore > BestScore)
+		{
+			BestItemIndex = Idx;
+			BestScore = ItemScore;
+		}
+	}
+
+	if (BestItemIndex != INDEX_NONE)
+	{
+		DrawDebugSphere(OwnerPC->GetWorld(), QueryData.RenderDebugHelpers[BestItemIndex].Location, QueryData.RenderDebugHelpers[BestItemIndex].Radius, 8, FColor::Red);
+
+		FString ItemDesc;
+
+		const int32 FailedTestIndex = QueryData.RenderDebugHelpers[BestItemIndex].FailedTestIndex;
+		if (FailedTestIndex != INDEX_NONE)
+		{
+			ItemDesc = FString::Printf(TEXT("{red}Selected item:%d failed test[%d]: {yellow}%s {LightBlue}(%s)\n{white}'%s' with score: %3.3f"),
+				BestItemIndex, FailedTestIndex,
+				*QueryData.Tests[FailedTestIndex].ShortName,
+				*QueryData.Tests[FailedTestIndex].Detailed,
+				*QueryData.RenderDebugHelpers[BestItemIndex].AdditionalInformation,
+				QueryData.RenderDebugHelpers[BestItemIndex].FailedScore);
+		}
+		else
+		{
+			ItemDesc = FString::Printf(TEXT("Selected item:{yellow}%d"), BestItemIndex);
+		}
+
+		CanvasContext.Print(ItemDesc);
+	}
+}
+
+namespace FEQSDebugTable
+{
+	const float RowHeight = 20.0f;
+	const float ItemDescriptionWidth = 312.0f;
+	const float ItemScoreWidth = 50.0f;
+	const float TestScoreWidth = 100.0f;
+}
+
+void FGameplayDebuggerCategory_EQS::DrawDetailedItemTable(const EQSDebug::FQueryData& QueryData, FGameplayDebuggerCanvasContext& CanvasContext) const
+{
+	const int32 EnableTableView = CVarEQSDetailsOnHUD.GetValueOnGameThread();
+	CanvasContext.Printf(TEXT("Detailed table view: {%s}%s{white}, use '{green}ai.vd.EQSOnHUD %d{white}' to toggle"),
+		EnableTableView ? TEXT("green") : TEXT("red"),
+		EnableTableView ? TEXT("active") : TEXT("disabled"),
+		EnableTableView ? 0 : 1);
+
+	if (EnableTableView && QueryData.NumValidItems > 0)
+	{
+		FCanvasTileItem TileItem(FVector2D(0, 0), GWhiteTexture, FVector2D(CanvasContext.Canvas->SizeX, FEQSDebugTable::RowHeight), FLinearColor::Black);
+		FLinearColor ColorOdd(0, 0, 0, 0.6f);
+		FLinearColor ColorEven(0, 0, 0.4f, 0.4f);
+		TileItem.BlendMode = SE_BLEND_Translucent;
+
+		// table header
+		CanvasContext.CursorY += FEQSDebugTable::RowHeight;
+		const float HeaderY = CanvasContext.CursorY + 3.0f;
+		TileItem.SetColor(ColorOdd);
+		CanvasContext.DrawItem(TileItem, 0, CanvasContext.CursorY);
+
+		float HeaderX = CanvasContext.CursorX;
+		CanvasContext.PrintfAt(HeaderX, HeaderY, FColor::Yellow, TEXT("Num items: %d"), QueryData.NumValidItems);
+		HeaderX += FEQSDebugTable::ItemDescriptionWidth;
+
+		CanvasContext.PrintAt(HeaderX, HeaderY, FColor::White, TEXT("Score"));
+		HeaderX += FEQSDebugTable::ItemScoreWidth;
+
+		for (int32 Idx = 0; Idx < QueryData.Tests.Num(); Idx++)
+		{
+			CanvasContext.PrintfAt(HeaderX, HeaderY, FColor::White, TEXT("Test %d"), Idx);
+			HeaderX += FEQSDebugTable::TestScoreWidth;
+		}
+
+		CanvasContext.CursorY += FEQSDebugTable::RowHeight;
+
+		// item rows
+		const int32 MaxShownItems = FMath::Min(MaxItemTableRows, QueryData.Items.Num());
+		for (int32 Idx = 0; Idx < MaxShownItems; Idx++)
+		{
+			TileItem.SetColor((Idx % 2) ? ColorOdd : ColorEven);
+			CanvasContext.DrawItem(TileItem, 0, CanvasContext.CursorY);
+
+			DrawDetailedItemRow(QueryData.Items[Idx], CanvasContext);
+			CanvasContext.CursorY += FEQSDebugTable::RowHeight;
+		}
+
+		// test description
+		CanvasContext.Print(TEXT("All tests from used option:"));
+		for (int32 Idx = 0; Idx < QueryData.Tests.Num(); Idx++)
+		{
+			CanvasContext.Printf(TEXT("Test %d = {yellow}%s {LightBlue}(%s)"), Idx, *QueryData.Tests[Idx].ShortName, *QueryData.Tests[Idx].Detailed);
+		}
+	}
+}
+
+void FGameplayDebuggerCategory_EQS::DrawDetailedItemRow(const EQSDebug::FItemData& ItemData, FGameplayDebuggerCanvasContext& CanvasContext) const
+{
+	const float PosY = CanvasContext.CursorY + 1.0f;
+	float PosX = CanvasContext.CursorX;
+
+	CanvasContext.PrintAt(PosX, PosY, FColor::White, ItemData.Desc);
+	PosX += FEQSDebugTable::ItemDescriptionWidth;
+
+	CanvasContext.PrintfAt(PosX, PosY, FColor::Yellow, TEXT("%.2f"), ItemData.TotalScore);
+	PosX += FEQSDebugTable::ItemScoreWidth;
+
+	FCanvasTileItem ActiveTileItem(FVector2D(0, PosY + 15.0f), GWhiteTexture, FVector2D(0, 2.0f), FLinearColor::Yellow);
+	FCanvasTileItem BackTileItem(FVector2D(0, PosY + 15.0f), GWhiteTexture, FVector2D(0, 2.0f), FLinearColor(0.1f, 0.1f, 0.1f));
+	const float BarWidth = 80.0f;
+
+	const int32 NumTests = ItemData.TestScores.Num();
+
+	float TotalWeightedScore = 0.0f;
+	for (int32 Idx = 0; Idx < NumTests; Idx++)
+	{
+		TotalWeightedScore += ItemData.TestScores[Idx];
+	}
+
+	for (int32 Idx = 0; Idx < NumTests; Idx++)
+	{
+		const float ScoreW = ItemData.TestScores[Idx];
+		const float ScoreN = ItemData.TestValues[Idx];
+		FString DescScoreW = FString::Printf(TEXT("%.2f"), ScoreW);
+		FString DescScoreN = (ScoreN == UEnvQueryTypes::SkippedItemValue) ? TEXT("SKIP") : FString::Printf(TEXT("%.2f"), ScoreN);
+		FString TestDesc = DescScoreW + FString(" {LightBlue}") + DescScoreN;
+
+		float Pct = (TotalWeightedScore > KINDA_SMALL_NUMBER) ? (ScoreW / TotalWeightedScore) : 0.0f;
+		ActiveTileItem.Position.X = PosX;
+		ActiveTileItem.Size.X = BarWidth * Pct;
+		BackTileItem.Position.X = PosX + ActiveTileItem.Size.X;
+		BackTileItem.Size.X = FMath::Max(BarWidth * (1.0f - Pct), 0.0f);
+
+		CanvasContext.DrawItem(ActiveTileItem, ActiveTileItem.Position.X, ActiveTileItem.Position.Y);
+		CanvasContext.DrawItem(BackTileItem, BackTileItem.Position.X, BackTileItem.Position.Y);
+
+		CanvasContext.PrintAt(PosX, PosY, FColor::Green, TestDesc);
+		PosX += FEQSDebugTable::TestScoreWidth;
+	}
+}
+#endif // USE_EQS_DEBUGGER
+
+void FGameplayDebuggerCategory_EQS::CycleShownQueries()
+{
+#if USE_EQS_DEBUGGER
+	ShownQueryIndex = (DataPack.QueryDebugData.Num() > 0) ? ((ShownQueryIndex + 1) % DataPack.QueryDebugData.Num()) : 0;
+#endif // USE_EQS_DEBUGGER
+	MarkRenderStateDirty();
+}
+
+#endif // ENABLE_GAMEPLAY_DEBUGGER
diff --git a/Engine/Source/Runtime/AIModule/Private/GameplayDebugger/GameplayDebuggerCategory_EQS.h b/Engine/Source/Runtime/AIModule/Private/GameplayDebugger/GameplayDebuggerCategory_EQS.h
new file mode 100644
--- /dev/null
+++ b/Engine/Source/Runtime/AIModule/Private/GameplayDebugger/GameplayDebuggerCategory_EQS.h
@@ -0,0 +1,47 @@
+// Copyright 1998-2016 Epic Games, Inc. All Rights Reserved.
+
+#if WITH_GAMEPLAY_DEBUGGER
+
+#pragma once
+#include "GameplayDebuggerCategory.h"
+#include "EnvironmentQuery/EnvQueryTypes.h"
+#include "EnvironmentQuery/EnvQueryDebugHelpers.h"
+
+class FGameplayDebuggerCategory_EQS : public FGameplayDebuggerCategory
+{
+public:
+	FGameplayDebuggerCategory_EQS();
+
+	virtual void CollectData(APlayerController* OwnerPC, AActor* DebugActor) override;
+	virtual void OnDataPackReplicated(int32 DataPackId) override;
+	virtual void DrawData(APlayerController* OwnerPC, FGameplayDebuggerCanvasContext& CanvasContext) override;
+	virtual FDebugRenderSceneProxy* CreateSceneProxy(const UPrimitiveComponent* InComponent) override;
+
+	static TSharedRef<FGameplayDebuggerCategory> MakeInstance();
+
+protected:
+#if USE_EQS_DEBUGGER
+	struct FRepData
+	{
+		TArray<EQSDebug::FQueryData> QueryDebugData;
+		
+		void Serialize(FArchive& Ar);
+	};
+	FRepData DataPack;
+
+	void DrawLookedAtItem(const EQSDebug::FQueryData& QueryData, APlayerController* OwnerPC, FGameplayDebuggerCanvasContext& CanvasContext) const;
+	void DrawDetailedItemTable(const EQSDebug::FQueryData& QueryData, FGameplayDebuggerCanvasContext& CanvasContext) const;
+	void DrawDetailedItemRow(const EQSDebug::FItemData& ItemData, FGameplayDebuggerCanvasContext& CanvasContext) const;
+#endif
+
+	void CycleShownQueries();
+
+	uint32 bDrawLabels : 1;
+	uint32 bDrawFailedItems : 1;
+
+	int32 MaxItemTableRows;
+	int32 MaxQueries;
+	int32 ShownQueryIndex;
+};
+
+#endif // WITH_GAMEPLAY_DEBUGGER
diff --git a/Engine/Source/Runtime/AIModule/Private/GameplayDebugger/GameplayDebuggerCategory_Navmesh.cpp b/Engine/Source/Runtime/AIModule/Private/GameplayDebugger/GameplayDebuggerCategory_Navmesh.cpp
new file mode 100644
--- /dev/null
+++ b/Engine/Source/Runtime/AIModule/Private/GameplayDebugger/GameplayDebuggerCategory_Navmesh.cpp
@@ -0,0 +1,78 @@
+// Copyright 1998-2016 Epic Games, Inc. All Rights Reserved.
+
+#include "AIModulePrivate.h"
+#include "GameplayDebuggerCategory_Navmesh.h"
+
+#if WITH_GAMEPLAY_DEBUGGER
+
+FGameplayDebuggerCategory_Navmesh::FGameplayDebuggerCategory_Navmesh()
+{
+	bShowOnlyWithDebugActor = false;
+	bShowCategoryName = false;
+	bShowDataPackReplication = true;
+	CollectDataInterval = 5.0f;
+	SetDataPackReplication<FNavMeshSceneProxyData>(&NavmeshRenderData);
+}
+
+TSharedRef<FGameplayDebuggerCategory> FGameplayDebuggerCategory_Navmesh::MakeInstance()
+{
+	return MakeShareable(new FGameplayDebuggerCategory_Navmesh());
+}
+
+void FGameplayDebuggerCategory_Navmesh::CollectData(APlayerController* OwnerPC, AActor* DebugActor)
+{
+#if WITH_RECAST
+	ARecastNavMesh* NavData = nullptr;
+
+	APawn* PlayerPawn = OwnerPC ? OwnerPC->GetPawnOrSpectator() : nullptr;
+	APawn* DebugActorAsPawn = Cast<APawn>(DebugActor);
+	APawn* DestPawn = DebugActorAsPawn ? DebugActorAsPawn : PlayerPawn;
+	if (DestPawn)
+	{
+		UNavigationSystem* NavSys = UNavigationSystem::GetCurrent(OwnerPC->GetWorld());
+		const FNavAgentProperties& NavAgentProperties = DestPawn->GetNavAgentPropertiesRef();
+		NavData = Cast<ARecastNavMesh>(NavSys->GetNavDataForProps(NavAgentProperties));
+	}
+
+	if (NavData)
+	{
+		// add 3x3 neighborhood of target
+		const FVector TargetLocation = DestPawn->GetActorLocation();
+
+		TArray<int32> TileSet;
+		int32 TileX = 0;
+		int32 TileY = 0;
+		const int32 DeltaX[] = { 0, 1, 1, 0, -1, -1, -1, 0, 1 };
+		const int32 DeltaY[] = { 0, 0, 1, 1, 1, 0, -1, -1, -1 };
+
+		int32 TargetTileX = 0;
+		int32 TargetTileY = 0;
+		NavData->GetNavMeshTileXY(TargetLocation, TargetTileX, TargetTileY);
+		for (int32 Idx = 0; Idx < ARRAY_COUNT(DeltaX); Idx++)
+		{
+			const int32 NeiX = TargetTileX + DeltaX[Idx];
+			const int32 NeiY = TargetTileY + DeltaY[Idx];
+			NavData->GetNavMeshTilesAt(NeiX, NeiY, TileSet);
+		}
+
+		const int32 DetailFlags =
+			(1 << static_cast<int32>(ENavMeshDetailFlags::PolyEdges)) |
+			(1 << static_cast<int32>(ENavMeshDetailFlags::FilledPolys)) |
+			(1 << static_cast<int32>(ENavMeshDetailFlags::NavLinks));
+
+		NavmeshRenderData.GatherData(NavData, DetailFlags, TileSet);
+	}
+#endif // WITH_RECAST
+}
+
+void FGameplayDebuggerCategory_Navmesh::OnDataPackReplicated(int32 DataPackId)
+{
+	MarkRenderStateDirty();
+}
+
+FDebugRenderSceneProxy* FGameplayDebuggerCategory_Navmesh::CreateSceneProxy(const UPrimitiveComponent* InComponent)
+{
+	return new FNavMeshSceneProxy(InComponent, &NavmeshRenderData, true);
+}
+
+#endif // WITH_GAMEPLAY_DEBUGGER
diff --git a/Engine/Source/Runtime/AIModule/Private/GameplayDebugger/GameplayDebuggerCategory_Navmesh.h b/Engine/Source/Runtime/AIModule/Private/GameplayDebugger/GameplayDebuggerCategory_Navmesh.h
new file mode 100644
--- /dev/null
+++ b/Engine/Source/Runtime/AIModule/Private/GameplayDebugger/GameplayDebuggerCategory_Navmesh.h
@@ -0,0 +1,24 @@
+// Copyright 1998-2016 Epic Games, Inc. All Rights Reserved.
+
+#if WITH_GAMEPLAY_DEBUGGER
+
+#pragma once
+#include "GameplayDebuggerCategory.h"
+#include "AI/Navigation/NavMeshRenderingComponent.h"
+
+class FGameplayDebuggerCategory_Navmesh : public FGameplayDebuggerCategory
+{
+public:
+	FGameplayDebuggerCategory_Navmesh();
+
+	virtual void CollectData(APlayerController* OwnerPC, AActor* DebugActor) override;
+	virtual FDebugRenderSceneProxy* CreateSceneProxy(const UPrimitiveComponent* InComponent) override;
+	virtual void OnDataPackReplicated(int32 DataPackId) override;
+
+	static TSharedRef<FGameplayDebuggerCategory> MakeInstance();
+
+protected:
+	FNavMeshSceneProxyData NavmeshRenderData;
+};
+
+#endif // WITH_GAMEPLAY_DEBUGGER
diff --git a/Engine/Source/Runtime/AIModule/Private/GameplayDebugger/GameplayDebuggerCategory_Perception.cpp b/Engine/Source/Runtime/AIModule/Private/GameplayDebugger/GameplayDebuggerCategory_Perception.cpp
new file mode 100644
--- /dev/null
+++ b/Engine/Source/Runtime/AIModule/Private/GameplayDebugger/GameplayDebuggerCategory_Perception.cpp
@@ -0,0 +1,39 @@
+// Copyright 1998-2016 Epic Games, Inc. All Rights Reserved.
+
+#include "AIModulePrivate.h"
+#include "GameplayDebuggerCategory_Perception.h"
+
+#if WITH_GAMEPLAY_DEBUGGER
+
+FGameplayDebuggerCategory_Perception::FGameplayDebuggerCategory_Perception()
+{
+}
+
+TSharedRef<FGameplayDebuggerCategory> FGameplayDebuggerCategory_Perception::MakeInstance()
+{
+	return MakeShareable(new FGameplayDebuggerCategory_Perception());
+}
+
+void FGameplayDebuggerCategory_Perception::CollectData(APlayerController* OwnerPC, AActor* DebugActor)
+{
+	APawn* MyPawn = Cast<APawn>(DebugActor);
+	if (MyPawn)
+	{
+		AAIController* BTAI = Cast<AAIController>(MyPawn->GetController());
+		if (BTAI)
+		{
+			UAIPerceptionComponent* PerceptionComponent = BTAI->GetPerceptionComponent();
+			if (PerceptionComponent == nullptr)
+			{
+				PerceptionComponent = MyPawn->FindComponentByClass<UAIPerceptionComponent>();
+			}
+
+			if (PerceptionComponent)
+			{
+				PerceptionComponent->DescribeSelfToGameplayDebugger(this);
+			}
+		}
+	}
+}
+
+#endif // WITH_GAMEPLAY_DEBUGGER
diff --git a/Engine/Source/Runtime/AIModule/Private/GameplayDebugger/GameplayDebuggerCategory_Perception.h b/Engine/Source/Runtime/AIModule/Private/GameplayDebugger/GameplayDebuggerCategory_Perception.h
new file mode 100644
--- /dev/null
+++ b/Engine/Source/Runtime/AIModule/Private/GameplayDebugger/GameplayDebuggerCategory_Perception.h
@@ -0,0 +1,18 @@
+// Copyright 1998-2016 Epic Games, Inc. All Rights Reserved.
+
+#if WITH_GAMEPLAY_DEBUGGER
+
+#pragma once
+#include "GameplayDebuggerCategory.h"
+
+class FGameplayDebuggerCategory_Perception : public FGameplayDebuggerCategory
+{
+public:
+	FGameplayDebuggerCategory_Perception();
+
+	virtual void CollectData(APlayerController* OwnerPC, AActor* DebugActor) override;
+
+	static TSharedRef<FGameplayDebuggerCategory> MakeInstance();
+};
+
+#endif // WITH_GAMEPLAY_DEBUGGER
diff --git a/Engine/Source/Runtime/AIModule/Private/Navigation/PathFollowingComponent.cpp b/Engine/Source/Runtime/AIModule/Private/Navigation/PathFollowingComponent.cpp
--- a/Engine/Source/Runtime/AIModule/Private/Navigation/PathFollowingComponent.cpp
+++ b/Engine/Source/Runtime/AIModule/Private/Navigation/PathFollowingComponent.cpp
@@ -1052,7 +1052,7 @@ bool UPathFollowingComponent::HasReachedInternal(const FVector& GoalLocation, fl
 
 void UPathFollowingComponent::DebugReachTest(float& CurrentDot, float& CurrentDistance, float& CurrentHeight, uint8& bDotFailed, uint8& bDistanceFailed, uint8& bHeightFailed) const
 {
-	if (!Path.IsValid() || MovementComp == NULL)
+	if (!Path.IsValid() || MovementComp == nullptr || MovementComp->UpdatedComponent == nullptr)
 	{
 		return;
 	}
diff --git a/Engine/Source/Runtime/AIModule/Private/Perception/AIPerceptionComponent.cpp b/Engine/Source/Runtime/AIModule/Private/Perception/AIPerceptionComponent.cpp
--- a/Engine/Source/Runtime/AIModule/Private/Perception/AIPerceptionComponent.cpp
+++ b/Engine/Source/Runtime/AIModule/Private/Perception/AIPerceptionComponent.cpp
@@ -5,6 +5,10 @@
 #include "DrawDebugHelpers.h"
 #include "Engine/Canvas.h"
 
+#if WITH_GAMEPLAY_DEBUGGER
+#include "GameplayDebuggerCategory.h"
+#endif
+
 #include "Perception/AISense_Sight.h"
 #include "Perception/AISense_Hearing.h"
 #include "Perception/AISense_Damage.h"
@@ -657,49 +661,42 @@ bool UAIPerceptionComponent::GetActorsPerception(AActor* Actor, FActorPerception
 //----------------------------------------------------------------------//
 // debug
 //----------------------------------------------------------------------//
-#if !UE_BUILD_SHIPPING
-void UAIPerceptionComponent::GrabGameplayDebuggerData(TArray<FString>& OnScreenStrings, TArray<FGameplayDebuggerShapeElement>& DebugShapes) const
+#if WITH_GAMEPLAY_DEBUGGER
+void UAIPerceptionComponent::DescribeSelfToGameplayDebugger(FGameplayDebuggerCategory* DebuggerCategory) const
 {
-#if ENABLED_GAMEPLAY_DEBUGGER
-	UAIPerceptionSystem* PerceptionSys = UAIPerceptionSystem::GetCurrent(GetWorld());
-	check(PerceptionSys);
+	if (DebuggerCategory == nullptr)
+	{
+		return;
+	}
 
 	for (UAIPerceptionComponent::TActorPerceptionContainer::TConstIterator It(GetPerceptualDataConstIterator()); It; ++It)
 	{
-		if (It->Key == NULL)
-		{
-			continue;
-		}
-
 		const FActorPerceptionInfo& ActorPerceptionInfo = It->Value;
-
-		if (ActorPerceptionInfo.Target.IsValid())
+		if (ActorPerceptionInfo.Target.IsValid() && It->Key)
 		{
 			const FVector TargetLocation = ActorPerceptionInfo.Target->GetActorLocation();
-			for (const auto& Stimulus : ActorPerceptionInfo.LastSensedStimuli)
+			for (const FAIStimulus& Stimulus : ActorPerceptionInfo.LastSensedStimuli)
 			{
-				if (Stimulus.Strength >= 0)
+				const UAISenseConfig* SenseConfig = GetSenseConfig(Stimulus.Type);
+				if (Stimulus.Strength >= 0 && SenseConfig)
 				{
-					const FString Description = FString::Printf(TEXT("%s: %.2f a:%.2f"), *PerceptionSys->GetSenseName(Stimulus.Type), Stimulus.Strength, Stimulus.GetAge());
-					DebugShapes.Add(UGameplayDebuggerHelper::MakeString(Description, Stimulus.StimulusLocation + FVector(0, 0, 30)));
+					const FString Description = FString::Printf(TEXT("%s: %.2f age:%.2f"), *SenseConfig->GetSenseName(), Stimulus.Strength, Stimulus.GetAge());
+					const FColor DebugColor = SenseConfig->GetDebugColor();
 
-					const FColor DebugColor = PerceptionSys->GetSenseDebugColor(Stimulus.Type);
-					DebugShapes.Add(UGameplayDebuggerHelper::MakePoint(Stimulus.StimulusLocation, DebugColor, 30));
-					DebugShapes.Add(UGameplayDebuggerHelper::MakeLine(Stimulus.ReceiverLocation, Stimulus.StimulusLocation, DebugColor));
-					DebugShapes.Add(UGameplayDebuggerHelper::MakeLine(TargetLocation, Stimulus.StimulusLocation, FColor::Black));
+					DebuggerCategory->AddShape(FGameplayDebuggerShape::MakePoint(Stimulus.StimulusLocation + FVector(0, 0, 30), 30.0f, DebugColor, Description));
+					DebuggerCategory->AddShape(FGameplayDebuggerShape::MakeSegment(Stimulus.ReceiverLocation, Stimulus.StimulusLocation, DebugColor));
+					DebuggerCategory->AddShape(FGameplayDebuggerShape::MakeSegment(TargetLocation, Stimulus.StimulusLocation, FColor::Black));
 				}
 			}
 		}
 	}
 
-	for (auto Sense : SensesConfig)
+	for (UAISenseConfig* SenseConfig : SensesConfig)
 	{
-		Sense->GetDebugData(OnScreenStrings, DebugShapes, *this);
+		SenseConfig->DescribeSelfToGameplayDebugger(this, DebuggerCategory);
 	}
-#endif
 }
-
-#endif // !UE_BUILD_SHIPPING
+#endif // WITH_GAMEPLAY_DEBUGGER
 
 #if ENABLE_VISUAL_LOG
 void UAIPerceptionComponent::DescribeSelfToVisLog(FVisualLogEntry* Snapshot) const
diff --git a/Engine/Source/Runtime/AIModule/Private/Perception/AIPerceptionSystem.cpp b/Engine/Source/Runtime/AIModule/Private/Perception/AIPerceptionSystem.cpp
--- a/Engine/Source/Runtime/AIModule/Private/Perception/AIPerceptionSystem.cpp
+++ b/Engine/Source/Runtime/AIModule/Private/Perception/AIPerceptionSystem.cpp
@@ -58,9 +58,6 @@ FAISenseID UAIPerceptionSystem::RegisterSenseClass(TSubclassOf<UAISense> SenseCl
 	{
 		const int32 ItemsToAdd = SenseID.Index - Senses.Num() + 1;
 		Senses.AddZeroed(ItemsToAdd);
-#if !UE_BUILD_SHIPPING
-		DebugSenseColors.AddZeroed(ItemsToAdd);
-#endif
 	}
 
 	if (Senses[SenseID] == nullptr)
@@ -79,14 +76,9 @@ FAISenseID UAIPerceptionSystem::RegisterSenseClass(TSubclassOf<UAISense> SenseCl
 			// otherwise it will get called in StartPlay()
 		}
 
-#if !UE_BUILD_SHIPPING
-		DebugSenseColors[SenseID] = Senses[SenseID]->GetDebugColor();
-		PerceptionDebugLegend += Senses[SenseID]->GetDebugLegend();
-
 		// make senses v-log to perception system's log
 		REDIRECT_OBJECT_TO_VLOG(Senses[SenseID], this);
 		UE_VLOG(this, LogAIPerception, Log, TEXT("Registering sense %s"), *Senses[SenseID]->GetName());
-#endif
 	}
 
 	return SenseID;
@@ -649,15 +641,3 @@ void UAIPerceptionSystem::ReportPerceptionEvent(UObject* WorldContext, UAISenseE
 		PerceptionSys->ReportEvent(PerceptionEvent);
 	}
 }
-
-#if !UE_BUILD_SHIPPING
-FColor UAIPerceptionSystem::GetSenseDebugColor(FAISenseID SenseID) const
-{
-	return SenseID.IsValid() && Senses.IsValidIndex(SenseID) && Senses[SenseID] != nullptr ? Senses[SenseID]->GetDebugColor() : FColor::Black;
-}
-
-FString UAIPerceptionSystem::GetSenseName(FAISenseID SenseID) const
-{
-	return SenseID.IsValid() && Senses.IsValidIndex(SenseID) && Senses[SenseID] != nullptr ? *Senses[SenseID]->GetDebugName() : TEXT("Invalid");
-}
-#endif 
diff --git a/Engine/Source/Runtime/AIModule/Private/Perception/AISense.cpp b/Engine/Source/Runtime/AIModule/Private/Perception/AISense.cpp
--- a/Engine/Source/Runtime/AIModule/Private/Perception/AISense.cpp
+++ b/Engine/Source/Runtime/AIModule/Private/Perception/AISense.cpp
@@ -2,14 +2,19 @@
 
 #include "AIModulePrivate.h"
 #include "Perception/AISense.h"
-#include "Perception/AISenseConfig_Sight.h"
+#include "Perception/AISenseConfig_Blueprint.h"
 #include "Perception/AISenseConfig_Hearing.h"
 #include "Perception/AISenseConfig_Prediction.h"
+#include "Perception/AISenseConfig_Sight.h"
+#include "Perception/AISenseConfig_Team.h"
+#include "Perception/AISenseConfig_Touch.h"
+#include "PErception/AISense_Blueprint.h"
 #include "Perception/AISense_Prediction.h"
+#include "Perception/AISense_Touch.h"
 
-#if !UE_BUILD_SHIPPING
-#include "DrawDebugHelpers.h"
-#endif // !UE_BUILD_SHIPPING
+#if WITH_GAMEPLAY_DEBUGGER
+#include "GameplayDebuggerCategory.h"
+#endif
 
 const float UAISense::SuspendNextUpdate = FLT_MAX;
 
@@ -18,7 +23,6 @@ UAISense::UAISense(const FObjectInitializer& ObjectInitializer)
 	, TimeUntilNextUpdate(SuspendNextUpdate)
 	, SenseID(FAISenseID::InvalidID())
 {
-	DebugName = GetName();
 	DefaultExpirationAge = FAIStimulus::NeverHappenedAge;
 
 	bNeedsForgettingNotification = false;
@@ -97,21 +101,11 @@ void UAISense::RegisterWrappedEvent(UAISenseEvent& PerceptionEvent)
 	UE_VLOG(GetPerceptionSystem(), LogAIPerception, Error, TEXT("%s did not override UAISense::RegisterWrappedEvent!"), *GetName());
 }
 
-#if !UE_BUILD_SHIPPING
-//----------------------------------------------------------------------//
-// DEBUG
-//----------------------------------------------------------------------//
-FString UAISense::GetDebugLegend() const
-{
-	return FString::Printf(TEXT("{%s} %s,"), *DebugDrawColor.ToString(), *GetName());
-}
-#endif // !UE_BUILD_SHIPPING
-
 //----------------------------------------------------------------------//
 // 
 //----------------------------------------------------------------------//
 UAISenseConfig::UAISenseConfig(const FObjectInitializer& ObjectInitializer)
-	: Super(ObjectInitializer), bStartsEnabled(true)
+	: Super(ObjectInitializer), DebugColor(FColor::White), bStartsEnabled(true)
 {
 }
 
@@ -120,35 +114,98 @@ TSubclassOf<UAISense> UAISenseConfig::GetSenseImplementation() const
 	return UAISense::StaticClass(); 
 }
 
+FString UAISenseConfig::GetSenseName() const
+{
+	if (CachedSenseName.Len() == 0)
+	{
+		CachedSenseName = GetSenseImplementation()->GetName();
+		CachedSenseName.RemoveFromEnd(TEXT("_C"));
+
+		int32 SeparatorIdx = INDEX_NONE;
+		const bool bHasSeparator = CachedSenseName.FindLastChar(TEXT('_'), SeparatorIdx);
+		if (bHasSeparator)
+		{
+			CachedSenseName = CachedSenseName.Mid(SeparatorIdx + 1);
+		}
+	}
+
+	return CachedSenseName;
+}
+
+#if WITH_GAMEPLAY_DEBUGGER
+static FString DescribeColorHelper(const FColor& Color)
+{
+	int32 MaxColors = GColorList.GetColorsNum();
+	for (int32 Idx = 0; Idx < MaxColors; Idx++)
+	{
+		if (Color == GColorList.GetFColorByIndex(Idx))
+		{
+			return GColorList.GetColorNameByIndex(Idx);
+		}
+	}
+
+	return FString(TEXT("color"));
+}
+
+void UAISenseConfig::DescribeSelfToGameplayDebugger(const UAIPerceptionComponent* PerceptionComponent, FGameplayDebuggerCategory* DebuggerCategory) const
+{
+	if (DebuggerCategory)
+	{
+		DebuggerCategory->AddTextLine(
+			FString::Printf(TEXT("%s: {%s}%s"), *GetSenseName(), *GetDebugColor().ToString(), *DescribeColorHelper(GetDebugColor()))
+			);
+	}
+}
+#endif // WITH_GAMEPLAY_DEBUGGER
+
 //----------------------------------------------------------------------//
 // 
 //----------------------------------------------------------------------//
+UAISenseConfig_Sight::UAISenseConfig_Sight(const FObjectInitializer& ObjectInitializer) : Super(ObjectInitializer)
+{
+	DebugColor = FColor::Green;
+}
+
 TSubclassOf<UAISense> UAISenseConfig_Sight::GetSenseImplementation() const 
 { 
 	return *Implementation; 
 }
 
-#if !UE_BUILD_SHIPPING
-void UAISenseConfig_Sight::GetDebugData(TArray<FString>& OnScreenStrings, TArray<FGameplayDebuggerShapeElement>& DebugShapes, const UAIPerceptionComponent& PerceptionComponent) const
+#if WITH_GAMEPLAY_DEBUGGER
+void UAISenseConfig_Sight::DescribeSelfToGameplayDebugger(const UAIPerceptionComponent* PerceptionComponent, FGameplayDebuggerCategory* DebuggerCategory) const
 {
-	/*PeripheralVisionAngleDegrees*/
-	const AActor* BodyActor = PerceptionComponent.GetBodyActor();
+	if (PerceptionComponent == nullptr || DebuggerCategory == nullptr)
+	{
+		return;
+	}
+
+	FColor SightRangeColor = FColor::Green;
+	FColor LoseSightRangeColor = FColorList::NeonPink;
+
+	// don't call Super implementation on purpose, replace color description line
+	DebuggerCategory->AddTextLine(
+		FString::Printf(TEXT("%s: {%s}%s {white}rangeIN:{%s}%s {white} rangeOUT:{%s}%s"), *GetSenseName(),
+			*GetDebugColor().ToString(), *DescribeColorHelper(GetDebugColor()),
+			*SightRangeColor.ToString(), *DescribeColorHelper(SightRangeColor),
+			*LoseSightRangeColor.ToString(), *DescribeColorHelper(LoseSightRangeColor))
+		);
+
+	const AActor* BodyActor = PerceptionComponent->GetBodyActor();
 	if (BodyActor != nullptr)
 	{
-		UWorld* World = BodyActor->GetWorld();
 		FVector BodyLocation, BodyFacing;
-		PerceptionComponent.GetLocationAndDirection(BodyLocation, BodyFacing);
+		PerceptionComponent->GetLocationAndDirection(BodyLocation, BodyFacing);
 
-		DebugShapes.Add(UGameplayDebuggerHelper::MakeCylinder(BodyLocation, BodyLocation + FVector(0, 0, -50), LoseSightRadius, 32, UAISense_Sight::GetDebugLoseSightColor()));
-		DebugShapes.Add(UGameplayDebuggerHelper::MakeCylinder(BodyLocation, BodyLocation + FVector(0, 0, -50), SightRadius, 32, UAISense_Sight::GetDebugSightRangeColor()));
+		DebuggerCategory->AddShape(FGameplayDebuggerShape::MakeCylinder(BodyLocation, LoseSightRadius, 25.0f, LoseSightRangeColor));
+		DebuggerCategory->AddShape(FGameplayDebuggerShape::MakeCylinder(BodyLocation, SightRadius, 25.0f, SightRangeColor));
 
 		const float SightPieLength = FMath::Max(LoseSightRadius, SightRadius);
-		DebugShapes.Add(UGameplayDebuggerHelper::MakeLine(BodyLocation, BodyLocation + (BodyFacing * SightPieLength), UAISense_Sight::GetDebugSightRangeColor()));
-		DebugShapes.Add(UGameplayDebuggerHelper::MakeLine(BodyLocation, BodyLocation + (BodyFacing.RotateAngleAxis(PeripheralVisionAngleDegrees, FVector::UpVector) * SightPieLength), UAISense_Sight::GetDebugSightRangeColor()));
-		DebugShapes.Add(UGameplayDebuggerHelper::MakeLine(BodyLocation, BodyLocation + (BodyFacing.RotateAngleAxis(-PeripheralVisionAngleDegrees, FVector::UpVector) * SightPieLength), UAISense_Sight::GetDebugSightRangeColor()));
+		DebuggerCategory->AddShape(FGameplayDebuggerShape::MakeSegment(BodyLocation, BodyLocation + (BodyFacing * SightPieLength), SightRangeColor));
+		DebuggerCategory->AddShape(FGameplayDebuggerShape::MakeSegment(BodyLocation, BodyLocation + (BodyFacing.RotateAngleAxis(PeripheralVisionAngleDegrees, FVector::UpVector) * SightPieLength), SightRangeColor));
+		DebuggerCategory->AddShape(FGameplayDebuggerShape::MakeSegment(BodyLocation, BodyLocation + (BodyFacing.RotateAngleAxis(-PeripheralVisionAngleDegrees, FVector::UpVector) * SightPieLength), SightRangeColor));
 	}
 }
-#endif // !UE_BUILD_SHIPPING
+#endif // WITH_GAMEPLAY_DEBUGGER
 
 //----------------------------------------------------------------------//
 // UAISenseConfig_Hearing
@@ -157,33 +214,82 @@ void UAISenseConfig_Sight::GetDebugData(TArray<FString>& OnScreenStrings, TArray
 UAISenseConfig_Hearing::UAISenseConfig_Hearing(const FObjectInitializer& ObjectInitializer) 
 	: Super(ObjectInitializer), HearingRange(3000.f)
 {
+	DebugColor = FColor::Yellow;
 }
 
 TSubclassOf<UAISense> UAISenseConfig_Hearing::GetSenseImplementation() const 
 { 
 	return *Implementation; 
 }
-#if !UE_BUILD_SHIPPING
-void UAISenseConfig_Hearing::GetDebugData(TArray<FString>& OnScreenStrings, TArray<FGameplayDebuggerShapeElement>& DebugShapes, const UAIPerceptionComponent& PerceptionComponent) const
+
+#if WITH_GAMEPLAY_DEBUGGER
+void UAISenseConfig_Hearing::DescribeSelfToGameplayDebugger(const UAIPerceptionComponent* PerceptionComponent, FGameplayDebuggerCategory* DebuggerCategory) const
 {
-	const AActor* BodyActor = PerceptionComponent.GetBodyActor();
+	if (PerceptionComponent == nullptr || DebuggerCategory == nullptr)
+	{
+		return;
+	}
+
+	FColor HearingRangeColor = FColor::Yellow;
+	FColor LoSHearingRangeColor = FColorList::Cyan;
+
+	// don't call Super implementation on purpose, replace color description line
+	DebuggerCategory->AddTextLine(
+		FString::Printf(TEXT("%s: {%s}%s {white}range:{%s}%s {white} rangeLoS:{%s}%s"), *GetSenseName(),
+			*GetDebugColor().ToString(), *DescribeColorHelper(GetDebugColor()),
+			*HearingRangeColor.ToString(), *DescribeColorHelper(HearingRangeColor),
+			*LoSHearingRangeColor.ToString(), *DescribeColorHelper(LoSHearingRangeColor))
+		);
+
+	const AActor* BodyActor = PerceptionComponent->GetBodyActor();
 	if (BodyActor != nullptr)
 	{
-		UWorld* World = BodyActor->GetWorld();
 		FVector OwnerLocation = BodyActor->GetActorLocation();
-		DebugShapes.Add(UGameplayDebuggerHelper::MakeCylinder(OwnerLocation, OwnerLocation + FVector(0, 0, -50), HearingRange, 32, UAISense_Hearing::GetDebugHearingRangeColor()));
+		
+		DebuggerCategory->AddShape(FGameplayDebuggerShape::MakeCylinder(OwnerLocation, HearingRange, 25.0f, HearingRangeColor));
 		if (bUseLoSHearing)
 		{
-			DebugShapes.Add(UGameplayDebuggerHelper::MakeCylinder(OwnerLocation, OwnerLocation + FVector(0, 0, -50), LoSHearingRange, 32, UAISense_Hearing::GetDebugLoSHearingRangeeColor()));
+			DebuggerCategory->AddShape(FGameplayDebuggerShape::MakeCylinder(OwnerLocation, LoSHearingRange, 25.0f, LoSHearingRangeColor));
 		}
 	}
 }
-#endif // !UE_BUILD_SHIPPING
+#endif // WITH_GAMEPLAY_DEBUGGER
 
 //----------------------------------------------------------------------//
 // UAISenseConfig_Prediction
 //----------------------------------------------------------------------//
+UAISenseConfig_Prediction::UAISenseConfig_Prediction(const FObjectInitializer& ObjectInitializer) : Super(ObjectInitializer)
+{
+	DebugColor = FColorList::Grey;
+}
+
 TSubclassOf<UAISense> UAISenseConfig_Prediction::GetSenseImplementation() const 
 { 
 	return UAISense_Prediction::StaticClass(); 
-}
\ No newline at end of file
+}
+
+
+UAISenseConfig_Blueprint::UAISenseConfig_Blueprint(const FObjectInitializer& ObjectInitializer) : Super(ObjectInitializer)
+{
+	DebugColor = FColorList::Grey;
+}
+
+TSubclassOf<UAISense> UAISenseConfig_Blueprint::GetSenseImplementation() const
+{
+	return *Implementation;
+}
+
+UAISenseConfig_Team::UAISenseConfig_Team(const FObjectInitializer& ObjectInitializer) : Super(ObjectInitializer)
+{
+	DebugColor = FColor::Blue;
+}
+
+UAISenseConfig_Touch::UAISenseConfig_Touch(const FObjectInitializer& ObjectInitializer) : Super(ObjectInitializer)
+{
+	DebugColor = FColor::Cyan;
+}
+
+TSubclassOf<UAISense> UAISenseConfig_Touch::GetSenseImplementation() const
+{
+	return UAISense_Touch::StaticClass();
+}
diff --git a/Engine/Source/Runtime/AIModule/Private/Perception/AISenseConfig_Damage.cpp b/Engine/Source/Runtime/AIModule/Private/Perception/AISenseConfig_Damage.cpp
--- a/Engine/Source/Runtime/AIModule/Private/Perception/AISenseConfig_Damage.cpp
+++ b/Engine/Source/Runtime/AIModule/Private/Perception/AISenseConfig_Damage.cpp
@@ -4,6 +4,11 @@
 #include "Perception/AISenseConfig_Damage.h"
 #include "Perception/AISense_Damage.h"
 
+UAISenseConfig_Damage::UAISenseConfig_Damage(const FObjectInitializer& ObjectInitializer) : Super(ObjectInitializer) 
+{
+	DebugColor = FColor::Red;
+}
+
 TSubclassOf<UAISense> UAISenseConfig_Damage::GetSenseImplementation() const
 {
 	return Implementation;
diff --git a/Engine/Source/Runtime/AIModule/Private/Perception/AISense_Damage.cpp b/Engine/Source/Runtime/AIModule/Private/Perception/AISense_Damage.cpp
--- a/Engine/Source/Runtime/AIModule/Private/Perception/AISense_Damage.cpp
+++ b/Engine/Source/Runtime/AIModule/Private/Perception/AISense_Damage.cpp
@@ -68,11 +68,8 @@ IAIPerceptionListenerInterface* FAIDamageEvent::GetDamagedActorAsPerceptionListe
 //----------------------------------------------------------------------//
 // 
 //----------------------------------------------------------------------//
-UAISense_Damage::UAISense_Damage(const FObjectInitializer& ObjectInitializer) :
-	Super(ObjectInitializer)
+UAISense_Damage::UAISense_Damage(const FObjectInitializer& ObjectInitializer) : Super(ObjectInitializer)
 {
-	DebugName = TEXT("Damage");
-	DebugDrawColor = FColor::Red;
 }
 
 float UAISense_Damage::Update()
diff --git a/Engine/Source/Runtime/AIModule/Private/Perception/AISense_Hearing.cpp b/Engine/Source/Runtime/AIModule/Private/Perception/AISense_Hearing.cpp
--- a/Engine/Source/Runtime/AIModule/Private/Perception/AISense_Hearing.cpp
+++ b/Engine/Source/Runtime/AIModule/Private/Perception/AISense_Hearing.cpp
@@ -53,9 +53,6 @@ UAISense_Hearing::FDigestedHearingProperties::FDigestedHearingProperties()
 UAISense_Hearing::UAISense_Hearing(const FObjectInitializer& ObjectInitializer) 
 	: Super(ObjectInitializer)
 {
-	DebugDrawColor = FColor::Yellow;
-	DebugName = TEXT("Hearing");
-
 	if (HasAnyFlags(RF_ClassDefaultObject) == false)
 	{
 		OnNewListenerDelegate.BindUObject(this, &UAISense_Hearing::OnNewListenerImpl);
@@ -188,17 +185,3 @@ void UAISense_Hearing::RegisterWrappedEvent(UAISenseEvent& PerceptionEvent)
 		RegisterEvent(HearingEvent->GetNoiseEvent());
 	}
 }
-
-
-#if !UE_BUILD_SHIPPING
-//----------------------------------------------------------------------//
-// DEBUG
-//----------------------------------------------------------------------//
-FString UAISense_Hearing::GetDebugLegend() const
-{
-	static const FColor HearingColor = GetDebugHearingRangeColor();
-	static const FColor LoSHearingColor = GetDebugLoSHearingRangeeColor();
-
-	return FString::Printf(TEXT("{%s} Hearing, {%s} LoS hearing,"), *HearingColor.ToString(), *LoSHearingColor.ToString());
-}
-#endif // !UE_BUILD_SHIPPING
diff --git a/Engine/Source/Runtime/AIModule/Private/Perception/AISense_Prediction.cpp b/Engine/Source/Runtime/AIModule/Private/Perception/AISense_Prediction.cpp
--- a/Engine/Source/Runtime/AIModule/Private/Perception/AISense_Prediction.cpp
+++ b/Engine/Source/Runtime/AIModule/Private/Perception/AISense_Prediction.cpp
@@ -6,8 +6,6 @@
 UAISense_Prediction::UAISense_Prediction(const FObjectInitializer& ObjectInitializer) :
 	Super(ObjectInitializer)
 {
-	DebugName = TEXT("Prediction");
-	DebugDrawColor = FColorList::Grey;
 }
 
 float UAISense_Prediction::Update()
diff --git a/Engine/Source/Runtime/AIModule/Private/Perception/AISense_Sight.cpp b/Engine/Source/Runtime/AIModule/Private/Perception/AISense_Sight.cpp
--- a/Engine/Source/Runtime/AIModule/Private/Perception/AISense_Sight.cpp
+++ b/Engine/Source/Runtime/AIModule/Private/Perception/AISense_Sight.cpp
@@ -97,8 +97,6 @@ UAISense_Sight::UAISense_Sight(const FObjectInitializer& ObjectInitializer)
 		OnListenerRemovedDelegate.BindUObject(this, &UAISense_Sight::OnListenerRemovedImpl);
 	}
 
-	DebugDrawColor = FColor::Green;
-	DebugName = TEXT("Sight");
 	NotifyType = EAISenseNotifyType::OnPerceptionChange;
 	
 	bAutoRegisterAllPawnsAsSources = true;
@@ -605,16 +603,3 @@ void UAISense_Sight::OnListenerForgetsAll(const FPerceptionListener& Listener)
 		}
 	}
 }
-
-#if !UE_BUILD_SHIPPING
-//----------------------------------------------------------------------//
-// DEBUG
-//----------------------------------------------------------------------//
-FString UAISense_Sight::GetDebugLegend() const
-{
-	static const FColor SightColor = GetDebugSightRangeColor(); 
-	static const FColor LoseSightColor = GetDebugLoseSightColor();
-
-	return FString::Printf(TEXT("{%s} Sight, {%s} Lose Sight,"), *SightColor.ToString(), *LoseSightColor.ToString());
-}
-#endif // !UE_BUILD_SHIPPING
\ No newline at end of file
diff --git a/Engine/Source/Runtime/AIModule/Private/Perception/AISense_Team.cpp b/Engine/Source/Runtime/AIModule/Private/Perception/AISense_Team.cpp
--- a/Engine/Source/Runtime/AIModule/Private/Perception/AISense_Team.cpp
+++ b/Engine/Source/Runtime/AIModule/Private/Perception/AISense_Team.cpp
@@ -20,8 +20,6 @@ FAITeamStimulusEvent::FAITeamStimulusEvent(AActor* InBroadcaster, AActor* InEnem
 UAISense_Team::UAISense_Team(const FObjectInitializer& ObjectInitializer) :
 	Super(ObjectInitializer)
 {
-	DebugName = TEXT("Team");
-	DebugDrawColor = FColor::Blue;
 }
 
 float UAISense_Team::Update()
diff --git a/Engine/Source/Runtime/AIModule/Private/Perception/AISense_Touch.cpp b/Engine/Source/Runtime/AIModule/Private/Perception/AISense_Touch.cpp
--- a/Engine/Source/Runtime/AIModule/Private/Perception/AISense_Touch.cpp
+++ b/Engine/Source/Runtime/AIModule/Private/Perception/AISense_Touch.cpp
@@ -6,8 +6,6 @@
 UAISense_Touch::UAISense_Touch(const FObjectInitializer& ObjectInitializer) :
 	Super(ObjectInitializer)
 {
-	DebugName = TEXT("Touch");
-	DebugDrawColor = FColor::Cyan;
 }
 
 float UAISense_Touch::Update()
diff --git a/Engine/Source/Runtime/AnimGraphRuntime/Private/BoneControllers/AnimNode_BoneDrivenController.cpp b/Engine/Source/Runtime/AnimGraphRuntime/Private/BoneControllers/AnimNode_BoneDrivenController.cpp
--- a/Engine/Source/Runtime/AnimGraphRuntime/Private/BoneControllers/AnimNode_BoneDrivenController.cpp
+++ b/Engine/Source/Runtime/AnimGraphRuntime/Private/BoneControllers/AnimNode_BoneDrivenController.cpp
@@ -3,6 +3,8 @@
 #include "AnimGraphRuntimePrivatePCH.h"
 #include "BoneControllers/AnimNode_BoneDrivenController.h"
 
+#include "AnimInstanceProxy.h"
+
 /////////////////////////////////////////////////////
 // FAnimNode_BoneDrivenController
 
@@ -14,6 +16,7 @@
 	, RangeMax(1.0f)
 	, RemappedMin(0.0f)
 	, RemappedMax(1.0f)
+	, DestinationMode(EDrivenDestinationMode::Bone)
 	, TargetComponent_DEPRECATED(EComponentType::None)
 	, bAffectTargetTranslationX(false)
 	, bAffectTargetTranslationY(false)
@@ -24,7 +27,7 @@
 	, bAffectTargetScaleX(false)
 	, bAffectTargetScaleY(false)
 	, bAffectTargetScaleZ(false)
-	, ModificationMode(EDrivenBoneModificationMode::AddToInput)
+	, ModificationMode(EDrivenBoneModificationMode::AddToInput)	
 {
 }
 
@@ -33,7 +36,15 @@
 	FString DebugLine = DebugData.GetNodeName(this);
 	DebugLine += "(";
 	AddDebugNodeData(DebugLine);
-	DebugLine += FString::Printf(TEXT("  DrivingBone: %s\nDrivenBone: %s"), *SourceBone.BoneName.ToString(), *TargetBone.BoneName.ToString());
+	if (DestinationMode == EDrivenDestinationMode::Bone)
+	{
+		DebugLine += FString::Printf(TEXT("  DrivingBone: %s\nDrivenBone: %s"), *SourceBone.BoneName.ToString(), *TargetBone.BoneName.ToString());
+	}
+	else
+	{
+		DebugLine += FString::Printf(TEXT("  DrivingBone: %s\nDrivenParameter: %s"), *SourceBone.BoneName.ToString(), *ParameterName.ToString());
+	}
+	
 	DebugData.AddDebugItem(DebugLine);
 
 	ComponentPose.GatherDebugData(DebugData);
@@ -42,62 +53,21 @@
 void FAnimNode_BoneDrivenController::EvaluateBoneTransforms(USkeletalMeshComponent* SkelComp, FCSPose<FCompactPose>& MeshBases, TArray<FBoneTransform>& OutCSBoneTransforms)
 {
 	check(OutCSBoneTransforms.Num() == 0);
-	
+
 	// Early out if we're not driving from or to anything
-	if (SourceComponent == EComponentType::None)
+	if (SourceComponent == EComponentType::None || DestinationMode != EDrivenDestinationMode::Bone)
 	{
 		return;
 	}
 
 	// Get the Local space transform and the ref pose transform to see how the transform for the source bone has changed
 	const FBoneContainer& BoneContainer = MeshBases.GetPose().GetBoneContainer();
-	const FTransform& SourceOrigRef = BoneContainer.GetRefPoseArray()[SourceBone.BoneIndex];
-	const FTransform SourceCurr = MeshBases.GetLocalSpaceTransform(SourceBone.GetCompactPoseIndex(BoneContainer));
+	const FTransform& SourceRefPoseBoneTransform = BoneContainer.GetRefPoseArray()[SourceBone.BoneIndex];
+	const FTransform SourceCurrentBoneTransform = MeshBases.GetLocalSpaceTransform(SourceBone.GetCompactPoseIndex(BoneContainer));
 
-	// Resolve source value
-	float SourceValue = 0.0f;
-	if (SourceComponent < EComponentType::RotationX)
-	{
-		const FVector TranslationDiff = SourceCurr.GetLocation() - SourceOrigRef.GetLocation();
-		SourceValue = TranslationDiff[(int32)(SourceComponent - EComponentType::TranslationX)];
-	}
-	else if (SourceComponent < EComponentType::Scale)
-	{
-		const FVector RotationDiff = (SourceCurr.GetRotation() * SourceOrigRef.GetRotation().Inverse()).Euler();
-		SourceValue = RotationDiff[(int32)(SourceComponent - EComponentType::RotationX)];
-	}
-	else if (SourceComponent == EComponentType::Scale)
-	{
-		const FVector CurrentScale = SourceCurr.GetScale3D();
-		const FVector RefScale = SourceOrigRef.GetScale3D();
-		const float ScaleDiff = FMath::Max3(CurrentScale[0], CurrentScale[1], CurrentScale[2]) - FMath::Max3(RefScale[0], RefScale[1], RefScale[2]);
-		SourceValue = ScaleDiff;
-	}
-	else
-	{
-		const FVector ScaleDiff = SourceCurr.GetScale3D() - SourceOrigRef.GetScale3D();
-		SourceValue = ScaleDiff[(int32)(SourceComponent - EComponentType::ScaleX)];
-	}
+	const float FinalDriverValue = ExtractSourceValue(SourceCurrentBoneTransform, SourceRefPoseBoneTransform);
 	
-	// Determine the resulting value
-	float FinalDriverValue = SourceValue;
-	if (DrivingCurve != nullptr)
-	{
-		// Remap thru the curve if set
-		FinalDriverValue = DrivingCurve->GetFloatValue(FinalDriverValue);
-	}
-	else
-	{
-		// Apply the fixed function remapping/clamping
-		if (bUseRange)
-		{
-			const float Alpha = FMath::Clamp(FMath::GetRangePct(RangeMin, RangeMax, FinalDriverValue), 0.0f, 1.0f);
-			FinalDriverValue = FMath::Lerp(RemappedMin, RemappedMax, Alpha);
-		}
-
-		FinalDriverValue *= Multiplier;
-	}
-
+	
 	// Calculate a new local-space bone position by adding or replacing target components in the current local space position
 	const FCompactPoseBoneIndex TargetBoneIndex = TargetBone.GetCompactPoseIndex(BoneContainer);
 
@@ -187,7 +157,7 @@
 	if (ParentIndex != INDEX_NONE)
 	{
 		const FTransform& ParentTM = MeshBases.GetComponentSpaceTransform(ParentIndex);
-		
+
 		OutCSBoneTransforms.Add(FBoneTransform(TargetBoneIndex, ModifiedLocalTM * ParentTM));
 	}
 	else
@@ -196,9 +166,83 @@
 	}
 }
 
+void FAnimNode_BoneDrivenController::EvaluateComponentSpaceInternal(FComponentSpacePoseContext& Context)
+{
+	// Early out if we're not driving from or to anything
+	if (SourceComponent == EComponentType::None || DestinationMode == EDrivenDestinationMode::Bone)
+	{
+		return;
+	}
+
+	// Get the Local space transform and the ref pose transform to see how the transform for the source bone has changed
+	const FBoneContainer& BoneContainer = Context.Pose.GetPose().GetBoneContainer();
+	const FTransform& SourceRefPoseBoneTransform = BoneContainer.GetRefPoseArray()[SourceBone.BoneIndex];
+	const FTransform SourceCurrentBoneTransform = Context.Pose.GetLocalSpaceTransform(SourceBone.GetCompactPoseIndex(BoneContainer));
+
+	const float FinalDriverValue = ExtractSourceValue(SourceCurrentBoneTransform, SourceRefPoseBoneTransform);
+
+	if (DestinationMode == EDrivenDestinationMode::MorphTarget || DestinationMode == EDrivenDestinationMode::MaterialParameter)
+	{
+		//	Morph target and Material parameter curves
+		USkeleton* Skeleton = Context.AnimInstanceProxy->GetSkeleton();
+		FSmartNameMapping::UID NameUID;
+		Skeleton->AddSmartNameAndModify(USkeleton::AnimCurveMappingName, ParameterName, NameUID);
+		Context.Curve.Set(NameUID, FinalDriverValue, (DestinationMode == EDrivenDestinationMode::MorphTarget) ? EAnimCurveFlags::ACF_DrivesMorphTarget : EAnimCurveFlags::ACF_DrivesMaterial);
+	}
+}
+
+const float FAnimNode_BoneDrivenController::ExtractSourceValue(const FTransform &InCurrentBoneTransform, const FTransform &InRefPoseBoneTransform)
+{
+	// Resolve source value
+	float SourceValue = 0.0f;
+	if (SourceComponent < EComponentType::RotationX)
+	{
+		const FVector TranslationDiff = InCurrentBoneTransform.GetLocation() - InRefPoseBoneTransform.GetLocation();
+		SourceValue = TranslationDiff[(int32)(SourceComponent - EComponentType::TranslationX)];
+	}
+	else if (SourceComponent < EComponentType::Scale)
+	{
+		const FVector RotationDiff = (InCurrentBoneTransform.GetRotation() * InRefPoseBoneTransform.GetRotation().Inverse()).Euler();
+		SourceValue = RotationDiff[(int32)(SourceComponent - EComponentType::RotationX)];
+	}
+	else if (SourceComponent == EComponentType::Scale)
+	{
+		const FVector CurrentScale = InCurrentBoneTransform.GetScale3D();
+		const FVector RefScale = InRefPoseBoneTransform.GetScale3D();
+		const float ScaleDiff = FMath::Max3(CurrentScale[0], CurrentScale[1], CurrentScale[2]) - FMath::Max3(RefScale[0], RefScale[1], RefScale[2]);
+		SourceValue = ScaleDiff;
+	}
+	else
+	{
+		const FVector ScaleDiff = InCurrentBoneTransform.GetScale3D() - InRefPoseBoneTransform.GetScale3D();
+		SourceValue = ScaleDiff[(int32)(SourceComponent - EComponentType::ScaleX)];
+	}
+
+	// Determine the resulting value
+	float FinalDriverValue = SourceValue;
+	if (DrivingCurve != nullptr)
+	{
+		// Remap thru the curve if set
+		FinalDriverValue = DrivingCurve->GetFloatValue(FinalDriverValue);
+	}
+	else
+	{
+		// Apply the fixed function remapping/clamping
+		if (bUseRange)
+		{
+			const float Alpha = FMath::Clamp(FMath::GetRangePct(RangeMin, RangeMax, FinalDriverValue), 0.0f, 1.0f);
+			FinalDriverValue = FMath::Lerp(RemappedMin, RemappedMax, Alpha);
+		}
+
+		FinalDriverValue *= Multiplier;
+	}
+
+	return FinalDriverValue;
+}
+
 bool FAnimNode_BoneDrivenController::IsValidToEvaluate(const USkeleton* Skeleton, const FBoneContainer& RequiredBones)
 {
-	return SourceBone.IsValid(RequiredBones) && TargetBone.IsValid(RequiredBones);
+	return SourceBone.IsValid(RequiredBones) && ( TargetBone.IsValid(RequiredBones) || DestinationMode != EDrivenDestinationMode::Bone );
 }
 
 void FAnimNode_BoneDrivenController::ConvertTargetComponentToBits()
diff --git a/Engine/Source/Runtime/AnimGraphRuntime/Public/BoneControllers/AnimNode_BoneDrivenController.h b/Engine/Source/Runtime/AnimGraphRuntime/Public/BoneControllers/AnimNode_BoneDrivenController.h
--- a/Engine/Source/Runtime/AnimGraphRuntime/Public/BoneControllers/AnimNode_BoneDrivenController.h
+++ b/Engine/Source/Runtime/AnimGraphRuntime/Public/BoneControllers/AnimNode_BoneDrivenController.h
@@ -42,6 +42,15 @@ enum class EDrivenBoneModificationMode : uint8
 	AddToRefPose
 };
 
+// Type of destination value to drive
+UENUM()
+enum class EDrivenDestinationMode : uint8
+{	
+	Bone,
+	MorphTarget,
+	MaterialParameter
+};
+
 /**
  * This is the runtime version of a bone driven controller, which maps part of the state from one bone to another (e.g., 2 * source.x -> target.z)
  */
@@ -86,6 +95,14 @@ struct ANIMGRAPHRUNTIME_API FAnimNode_BoneDrivenController : public FAnimNode_Sk
 	UPROPERTY(EditAnywhere, Category = Mapping, meta = (EditCondition = bUseRange, DisplayName="Mapped Range Max"))
 	float RemappedMax;
 
+	// Type of destination to drive, currently either bone or morph target
+	UPROPERTY(EditAnywhere, Category = "Destination (driven)")
+	EDrivenDestinationMode DestinationMode;
+
+	/** Name of Morph Target to drive using the source attribute */
+	UPROPERTY(EditAnywhere, Category = "Destination (driven)")
+	FName ParameterName;
+
 	// Bone to drive using controller input
 	UPROPERTY(EditAnywhere, Category="Destination (driven)")
 	FBoneReference TargetBone;
@@ -143,7 +160,8 @@ public:
 	// End of FAnimNode_Base interface
 
 	// FAnimNode_SkeletalControlBase interface
-	virtual void EvaluateBoneTransforms(USkeletalMeshComponent* SkelComp, FCSPose<FCompactPose>& MeshBases, TArray<FBoneTransform>& OutBoneTransforms) override;
+	virtual void EvaluateBoneTransforms(USkeletalMeshComponent* SkelComp, FCSPose<FCompactPose>& MeshBases, TArray<FBoneTransform>& OutBoneTransforms) override;	
+	virtual void EvaluateComponentSpaceInternal(FComponentSpacePoseContext& Context);
 	virtual bool IsValidToEvaluate(const USkeleton* Skeleton, const FBoneContainer& RequiredBones) override;
 	// End of FAnimNode_SkeletalControlBase interface
 
@@ -153,5 +171,8 @@ protected:
 	
 	// FAnimNode_SkeletalControlBase protected interface
 	virtual void InitializeBoneReferences(const FBoneContainer& RequiredBones) override;
+
+	/** Extracts the value used to drive the target bone or parameter */
+	const float ExtractSourceValue(const FTransform& InCurrentBoneTransform, const FTransform& InRefPoseBoneTransform);
 	// End of FAnimNode_SkeletalControlBase protected interface
 };
\ No newline at end of file
diff --git a/Engine/Source/Runtime/Core/Public/Containers/Array.h b/Engine/Source/Runtime/Core/Public/Containers/Array.h
--- a/Engine/Source/Runtime/Core/Public/Containers/Array.h
+++ b/Engine/Source/Runtime/Core/Public/Containers/Array.h
@@ -1703,7 +1703,7 @@ public:
 	 *
 	 * @param NewNum New size of the array.
 	 */
-	void SetNumZeroed(int32 NewNum)
+	void SetNumZeroed(int32 NewNum, bool bAllowShrinking = true)
 	{
 		if (NewNum > Num())
 		{
@@ -1711,7 +1711,7 @@ public:
 		}
 		else if (NewNum < Num())
 		{
-			RemoveAt(NewNum, Num() - NewNum);
+			RemoveAt(NewNum, Num() - NewNum, bAllowShrinking);
 		}
 	}
 
@@ -1720,7 +1720,7 @@ public:
 	 *
 	 * @param NewNum New size of the array.
 	 */
-	void SetNumUninitialized(int32 NewNum)
+	void SetNumUninitialized(int32 NewNum, bool bAllowShrinking = true)
 	{
 		if (NewNum > Num())
 		{
@@ -1728,7 +1728,7 @@ public:
 		}
 		else if (NewNum < Num())
 		{
-			RemoveAt(NewNum, Num() - NewNum);
+			RemoveAt(NewNum, Num() - NewNum, bAllowShrinking);
 		}
 	}
 
diff --git a/Engine/Source/Runtime/CoreUObject/Private/UObject/UObjectHash.cpp b/Engine/Source/Runtime/CoreUObject/Private/UObject/UObjectHash.cpp
--- a/Engine/Source/Runtime/CoreUObject/Private/UObject/UObjectHash.cpp
+++ b/Engine/Source/Runtime/CoreUObject/Private/UObject/UObjectHash.cpp
@@ -671,12 +671,12 @@ void ForEachObjectWithOuter(const class UObjectBase* Outer, TFunctionRef<void(UO
 			if (!Object->HasAnyFlags(ExclusionFlags) && !Object->HasAnyInternalFlags(ExclusionInternalFlags))
 			{
 				Operation(Object);
-				if (bIncludeNestedObjects)
+			}
+			if (bIncludeNestedObjects)
+			{
+				if (TSet<UObjectBase*> const* ObjectInners = ThreadHash.ObjectOuterMap.Find(Object))
 				{
-					if (TSet<UObjectBase*> const* ObjectInners = ThreadHash.ObjectOuterMap.Find(Object))
-					{
-						AllInners.Add(ObjectInners);
-					}
+					AllInners.Add(ObjectInners);
 				}
 			}
 		}
diff --git a/Engine/Source/Runtime/Engine/Classes/AI/Navigation/NavMeshRenderingComponent.h b/Engine/Source/Runtime/Engine/Classes/AI/Navigation/NavMeshRenderingComponent.h
--- a/Engine/Source/Runtime/Engine/Classes/AI/Navigation/NavMeshRenderingComponent.h
+++ b/Engine/Source/Runtime/Engine/Classes/AI/Navigation/NavMeshRenderingComponent.h
@@ -1,9 +1,142 @@
 // Copyright 1998-2016 Epic Games, Inc. All Rights Reserved.
 
 #pragma once
+
+#include "DebugRenderSceneProxy.h"
+#include "DynamicMeshBuilder.h"
+#include "LocalVertexFactory.h"
+#include "GenericOctree.h"
+#include "Components/PrimitiveComponent.h"
 #include "NavMeshRenderingComponent.generated.h"
 
-struct FNavMeshSceneProxyData;
+class UNavMeshRenderingComponent;
+class ARecastNavMesh;
+
+enum class ENavMeshDetailFlags : uint8
+{
+	TriangleEdges,
+	PolyEdges,
+	BoundaryEdges,
+	FilledPolys,
+	TileBounds,
+	PathCollidingGeometry,
+	TileLabels,
+	PolygonLabels,
+	PolygonCost,
+	PathLabels,
+	NavLinks,
+	FailedNavLinks,
+	Clusters,
+	NavOctree,
+};
+
+// exported to API for GameplayDebugger module
+struct ENGINE_API FNavMeshSceneProxyData : public TSharedFromThis<FNavMeshSceneProxyData, ESPMode::ThreadSafe>
+{
+	struct FDebugMeshData
+	{
+		TArray<FDynamicMeshVertex> Vertices;
+		TArray<int32> Indices;
+		FColor ClusterColor;
+	};
+	TArray<FDebugMeshData> MeshBuilders;
+
+	TArray<FDebugRenderSceneProxy::FDebugLine> ThickLineItems;
+	TArray<FDebugRenderSceneProxy::FDebugLine> TileEdgeLines;
+	TArray<FDebugRenderSceneProxy::FDebugLine> NavMeshEdgeLines;
+	TArray<FDebugRenderSceneProxy::FDebugLine> NavLinkLines;
+	TArray<FDebugRenderSceneProxy::FDebugLine> ClusterLinkLines;
+
+	struct FDebugText
+	{
+		FVector Location;
+		FString Text;
+
+		FDebugText() {}
+		FDebugText(const FVector& InLocation, const FString& InText) : Location(InLocation), Text(InText) {}
+	};
+	TArray<FDebugText> DebugLabels;
+	
+	TArray<int32> PathCollidingGeomIndices;
+	TArray<FDynamicMeshVertex> PathCollidingGeomVerts;
+	TArray<FBoxCenterAndExtent>	OctreeBounds;
+
+	FBox Bounds;
+	FVector NavMeshDrawOffset;
+	uint32 bDataGathered : 1;
+	uint32 bNeedsNewData : 1;
+	int32 NavDetailFlags;
+
+	FNavMeshSceneProxyData() : NavMeshDrawOffset(0, 0, 10.f),
+		bDataGathered(false), bNeedsNewData(true), NavDetailFlags(0) {}
+
+	void Reset();
+	void Serialize(FArchive& Ar);
+	uint32 GetAllocatedSize() const;
+
+#if WITH_RECAST
+	int32 GetDetailFlags(const ARecastNavMesh* NavMesh) const;
+	void GatherData(const ARecastNavMesh* NavMesh, int32 InNavDetailFlags, const TArray<int32>& TileSet);
+#endif
+};
+
+// exported to API for GameplayDebugger module
+class ENGINE_API FNavMeshSceneProxy : public FDebugRenderSceneProxy
+{
+public:
+	FNavMeshSceneProxy(const UPrimitiveComponent* InComponent, FNavMeshSceneProxyData* InProxyData, bool ForceToRender = false);
+	virtual ~FNavMeshSceneProxy();
+
+	virtual void RegisterDebugDrawDelgate() override;
+	virtual void UnregisterDebugDrawDelgate() override;
+	virtual void GetDynamicMeshElements(const TArray<const FSceneView*>& Views, const FSceneViewFamily& ViewFamily, uint32 VisibilityMap, FMeshElementCollector& Collector) const override;
+
+protected:
+	void DrawDebugBox(FPrimitiveDrawInterface* PDI, FVector const& Center, FVector const& Box, FColor const& Color) const;
+	virtual void DrawDebugLabels(UCanvas* Canvas, APlayerController*) override;
+	virtual FPrimitiveViewRelevance GetViewRelevance(const FSceneView* View) const override;
+
+	virtual uint32 GetMemoryFootprint(void) const override { return sizeof(*this) + GetAllocatedSize(); }
+	uint32 GetAllocatedSize(void) const;
+
+private:
+	class FNavMeshIndexBuffer : public FIndexBuffer
+	{
+	public:
+		TArray<int32> Indices;
+		virtual void InitRHI() override;
+	};
+
+	class FNavMeshVertexBuffer : public FVertexBuffer
+	{
+	public:
+		TArray<FDynamicMeshVertex> Vertices;
+		virtual void InitRHI() override;
+	};
+
+	class FNavMeshVertexFactory : public FLocalVertexFactory
+	{
+	public:
+		void Init(const FNavMeshVertexBuffer* VertexBuffer);
+	};
+			
+	FNavMeshSceneProxyData ProxyData;
+
+	FNavMeshIndexBuffer IndexBuffer;
+	FNavMeshVertexBuffer VertexBuffer;
+	FNavMeshVertexFactory VertexFactory;
+
+	TArray<FColoredMaterialRenderProxy> MeshColors;
+	TArray<FMeshBatchElement> MeshBatchElements;
+
+	FDebugDrawDelegate DebugTextDrawingDelegate;
+	FDelegateHandle DebugTextDrawingDelegateHandle;
+	TWeakObjectPtr<UNavMeshRenderingComponent> RenderingComponent;
+	uint32 bRequestedData : 1;
+	uint32 bForceRendering : 1;
+	uint32 bSkipDistanceCheck : 1;
+	uint32 bUseThickLines : 1;
+};
 
 UCLASS(hidecategories=Object, editinlinenew)
 class ENGINE_API UNavMeshRenderingComponent : public UPrimitiveComponent
@@ -32,7 +165,6 @@ public:
 
 	void ForceUpdate() { bForceUpdate = true; }
 	bool IsForcingUpdate() const { return bForceUpdate; }
-	void GatherData(FNavMeshSceneProxyData& DebugDrawData) const;
 
 	static bool IsNavigationShowFlagSet(const UWorld* World);
 
diff --git a/Engine/Source/Runtime/Engine/Classes/AI/Navigation/NavigationData.h b/Engine/Source/Runtime/Engine/Classes/AI/Navigation/NavigationData.h
--- a/Engine/Source/Runtime/Engine/Classes/AI/Navigation/NavigationData.h
+++ b/Engine/Source/Runtime/Engine/Classes/AI/Navigation/NavigationData.h
@@ -281,6 +281,8 @@ public:
 
 	FVector GetLastRepathGoalLocation() const { return GoalActorLastLocation; }
 	void UpdateLastRepathGoalLocation();
+	
+	float GetLastUpdateTime() const { return LastUpdateTimeStamp; }
 
 	/** if enabled path will request recalculation if it gets invalidated due to a change to underlying navigation */
 	void EnableRecalculationOnInvalidation(bool bShouldAutoUpdate)
diff --git a/Engine/Source/Runtime/Engine/Classes/Animation/AnimStateMachineTypes.h b/Engine/Source/Runtime/Engine/Classes/Animation/AnimStateMachineTypes.h
--- a/Engine/Source/Runtime/Engine/Classes/Animation/AnimStateMachineTypes.h
+++ b/Engine/Source/Runtime/Engine/Classes/Animation/AnimStateMachineTypes.h
@@ -96,6 +96,7 @@ struct FAnimationState : public FAnimationStateBase
 	
 	FAnimationState()
 		: FAnimationStateBase()
+		, StateRootNodeIndex(INDEX_NONE)
 		, StartNotify(INDEX_NONE)
 		, EndNotify(INDEX_NONE)
 		, FullyBlendedNotify(INDEX_NONE)
@@ -245,7 +246,8 @@ struct FBakedAnimationState
 
 public:
 	FBakedAnimationState()
-		: StartNotify(INDEX_NONE)
+		: StateRootNodeIndex(INDEX_NONE)
+		, StartNotify(INDEX_NONE)
 		, EndNotify(INDEX_NONE)
 		, FullyBlendedNotify(INDEX_NONE)
 		, bIsAConduit(false)
diff --git a/Engine/Source/Runtime/Engine/Classes/Components/StaticMeshComponent.h b/Engine/Source/Runtime/Engine/Classes/Components/StaticMeshComponent.h
--- a/Engine/Source/Runtime/Engine/Classes/Components/StaticMeshComponent.h
+++ b/Engine/Source/Runtime/Engine/Classes/Components/StaticMeshComponent.h
@@ -185,7 +185,7 @@ class ENGINE_API UStaticMeshComponent : public UMeshComponent
 	uint32 bForceNavigationObstacle : 1;
 
 	/** Use the collision profile specified in the StaticMesh asset.*/
-	UPROPERTY(EditAnywhere, Category = Collision)
+	UPROPERTY()
 	uint32 bUseDefaultCollision : 1;
 
 	/** If true, mesh painting is disallowed on this instance. Set if vertex colors are overridden in a construction script. */
diff --git a/Engine/Source/Runtime/Engine/Classes/Distributions/DistributionFloat.h b/Engine/Source/Runtime/Engine/Classes/Distributions/DistributionFloat.h
--- a/Engine/Source/Runtime/Engine/Classes/Distributions/DistributionFloat.h
+++ b/Engine/Source/Runtime/Engine/Classes/Distributions/DistributionFloat.h
@@ -41,6 +41,8 @@ public:
 	{
 	}
 
+	/** Whether the distribution data has been cooked or the object itself is available */
+	bool IsCreated();
 
 #if WITH_EDITOR
 	/**`
@@ -71,10 +73,13 @@ public:
 
 	void InitLookupTable();
 
-	FORCEINLINE bool HasLookupTable()
+	FORCEINLINE bool HasLookupTable(bool bInitializeIfNeeded = true)
 	{
 #if WITH_EDITOR
-		InitLookupTable();
+		if(bInitializeIfNeeded)
+		{
+			InitLookupTable();
+		}
 #endif
 		return GDistributionType != 0 && !LookupTable.IsEmpty();
 	}
@@ -85,7 +90,6 @@ public:
 		return true; // even if they stay distributions, this should probably be ok as long as nobody is changing them at runtime
 		//return !Distribution || HasLookupTable();
 	}
-	
 };
 
 UCLASS(abstract, customconstructor,MinimalAPI)
@@ -99,6 +103,11 @@ class UDistributionFloat : public UDistribution
 
 	/** Set internally when the distribution is updated so that that FRawDistribution can know to update itself*/
 	uint32 bIsDirty:1;
+	
+protected:
+	UPROPERTY()
+	uint32 bBakedDataSuccesfully:1;	//It's possible that even though we want to bake we are not able to because of content or code.
+public:
 
 	/** Script-accessible way to query a float distribution */
 	virtual float GetFloatValue(float F = 0);
@@ -145,6 +154,11 @@ class UDistributionFloat : public UDistribution
 		return bCanBeBaked; 
 	}
 
+	bool HasBakedSuccesfully() const
+	{
+		return bBakedDataSuccesfully;
+	}
+
 	/**
 	 * Returns the number of values in the distribution. 1 for float.
 	 */
@@ -159,6 +173,7 @@ class UDistributionFloat : public UDistribution
 #endif	// WITH_EDITOR
 	virtual bool NeedsLoadForClient() const override;
 	virtual bool NeedsLoadForServer() const override;
+	virtual void Serialize(FArchive& Ar) override;
 	/** End UObject interface */
 };
 
diff --git a/Engine/Source/Runtime/Engine/Classes/Distributions/DistributionVector.h b/Engine/Source/Runtime/Engine/Classes/Distributions/DistributionVector.h
--- a/Engine/Source/Runtime/Engine/Classes/Distributions/DistributionVector.h
+++ b/Engine/Source/Runtime/Engine/Classes/Distributions/DistributionVector.h
@@ -59,14 +59,24 @@ private:
 	UPROPERTY()
 	float MaxValue;
 
+	UPROPERTY()
+	FVector MinValueVec;
+
+	UPROPERTY()
+	FVector MaxValueVec;
+
 public:
 	UPROPERTY(EditAnywhere, export, noclear, Category=RawDistributionVector)
 	class UDistributionVector* Distribution;
 
+	/** Whether the distribution data has been cooked or the object itself is available */
+	bool IsCreated();
 
 	FRawDistributionVector()
 		: MinValue(0)
 		, MaxValue(0)
+		, MinValueVec(FVector::ZeroVector)
+		, MaxValueVec(FVector::ZeroVector)
 		, Distribution(NULL)
 	{
 	}
@@ -94,6 +104,11 @@ public:
 	*/
 	void GetOutRange(float& MinOut, float& MaxOut);
 
+	/**
+	* Get the min and max values
+	*/
+	void GetRange(FVector& MinOut, FVector& MaxOut);
+
 	/**
 	* Is this distribution a uniform type? (ie, does it have two values per entry?)
 	*/
@@ -101,10 +116,13 @@ public:
 
 	void InitLookupTable();
 
-	FORCEINLINE bool HasLookupTable()
+	FORCEINLINE bool HasLookupTable(bool bInitializeIfNeeded = true)
 	{
 #if WITH_EDITOR
-		InitLookupTable();
+		if(bInitializeIfNeeded)
+		{
+			InitLookupTable();
+		}
 #endif
 		return GDistributionType != 0 && !LookupTable.IsEmpty();
 	}
@@ -117,6 +135,7 @@ public:
 	}
 };
 
+
 UCLASS(abstract, customconstructor)
 class ENGINE_API UDistributionVector : public UDistribution
 {
@@ -129,6 +148,10 @@ class ENGINE_API UDistributionVector : public UDistribution
 	/** Set internally when the distribution is updated so that that FRawDistribution can know to update itself*/
 	UPROPERTY()
 	uint32 bIsDirty:1;
+protected:
+	UPROPERTY()
+	uint32 bBakedDataSuccesfully:1;	//It's possible that even though we want to bake we are not able to because of content or code.
+public:
 
 	/** Script-accessible way to query a FVector distribution */
 	virtual FVector GetVectorValue(float F = 0);
@@ -175,6 +198,11 @@ class ENGINE_API UDistributionVector : public UDistribution
 		return bCanBeBaked; 
 	}
 
+	bool HasBakedSuccesfully() const
+	{
+		return bBakedDataSuccesfully;
+	}
+
 	/**
 	 * Returns the number of values in the distribution. 3 for vector.
 	 */
@@ -189,6 +217,7 @@ class ENGINE_API UDistributionVector : public UDistribution
 #endif // WITH_EDITOR
 	virtual bool NeedsLoadForClient() const override;
 	virtual bool NeedsLoadForServer() const override;
+	virtual void Serialize(FArchive& Ar) override;
 	/** End UObject interface */
 
 };
