Fix missing includes on Linux and Mac non-unity

diff --git a/Engine/Source/Runtime/AIModule/Private/GameplayDebugger/GameplayDebuggerCategory_AI.cpp b/Engine/Source/Runtime/AIModule/Private/GameplayDebugger/GameplayDebuggerCategory_AI.cpp
--- a/Engine/Source/Runtime/AIModule/Private/GameplayDebugger/GameplayDebuggerCategory_AI.cpp
+++ b/Engine/Source/Runtime/AIModule/Private/GameplayDebugger/GameplayDebuggerCategory_AI.cpp
@@ -7,6 +7,12 @@
 #include "DebugRenderSceneProxy.h"
 #include "GameplayTasksComponent.h"
 #include "GameplayDebuggerCategory_AI.h"
+#include "GameFramework/Character.h"
+#include "GameFramework/CharacterMovementComponent.h"
+#include "Animation/AnimMontage.h"
+#include "Engine/Canvas.h"
+#include "DrawDebugHelpers.h"
+
 
 FGameplayDebuggerCategory_AI::FGameplayDebuggerCategory_AI()
 {
diff --git a/Engine/Source/Runtime/AIModule/Private/GameplayDebugger/GameplayDebuggerCategory_EQS.cpp b/Engine/Source/Runtime/AIModule/Private/GameplayDebugger/GameplayDebuggerCategory_EQS.cpp
--- a/Engine/Source/Runtime/AIModule/Private/GameplayDebugger/GameplayDebuggerCategory_EQS.cpp
+++ b/Engine/Source/Runtime/AIModule/Private/GameplayDebugger/GameplayDebuggerCategory_EQS.cpp
@@ -5,6 +5,8 @@
 #include "EnvironmentQuery/EnvQueryManager.h"
 #include "EnvironmentQuery/EQSRenderingComponent.h"
 #include "GameplayDebuggerCategory_EQS.h"
+#include "Engine/Canvas.h"
+#include "DrawDebugHelpers.h"
 
 #if WITH_GAMEPLAY_DEBUGGER
 
