Copying //UE4/Dev-Framework to Dev-Main (//UE4/Dev-Main) @ 2926658

diff --git a/Engine/Source/Editor/PhAT/Private/PhATPreviewViewportClient.cpp b/Engine/Source/Editor/PhAT/Private/PhATPreviewViewportClient.cpp
--- a/Engine/Source/Editor/PhAT/Private/PhATPreviewViewportClient.cpp
+++ b/Engine/Source/Editor/PhAT/Private/PhATPreviewViewportClient.cpp
@@ -1,6 +1,6 @@
 // Copyright 1998-2016 Epic Games, Inc. All Rights Reserved.
 
-#include "PhATModule.h"
+#include "PhATPrivatePCH.h"
 #include "PhATEdSkeletalMeshComponent.h"
 #include "PhAT.h"
 #include "MouseDeltaTracker.h"
@@ -70,7 +70,7 @@ FPhATEdPreviewViewportClient::FPhATEdPreviewViewportClient(TWeakPtr<FPhAT> InPhA
 
 	// Take into account internal mesh translation/rotation/scaling etc.
 	FTransform LocalToWorld = SharedData->EditorSkelComp->ComponentToWorld;
-	FSphere WorldSphere = SharedData->EditorSkelMesh->Bounds.GetSphere().TransformBy(LocalToWorld);
+	FSphere WorldSphere = SharedData->EditorSkelMesh->GetImportedBounds().GetSphere().TransformBy(LocalToWorld);
 
 	CollBoxExtent = CollBox.GetExtent();
 	if (CollBoxExtent.X > CollBoxExtent.Y)
diff --git a/Engine/Source/Editor/PhAT/Private/PhATPrivatePCH.h b/Engine/Source/Editor/PhAT/Private/PhATPrivatePCH.h
new file mode 100644
--- /dev/null
+++ b/Engine/Source/Editor/PhAT/Private/PhATPrivatePCH.h
@@ -0,0 +1,8 @@
+// Copyright 1998-2016 Epic Games, Inc. All Rights Reserved.
+
+#pragma once
+
+#include "UnrealEd.h"
+#include "Toolkits/AssetEditorToolkit.h"
+#include "Toolkits/IToolkit.h"
+#include "IPhAT.h"
\ No newline at end of file
diff --git a/Engine/Source/Editor/PhAT/Private/PhATRender.cpp b/Engine/Source/Editor/PhAT/Private/PhATRender.cpp
--- a/Engine/Source/Editor/PhAT/Private/PhATRender.cpp
+++ b/Engine/Source/Editor/PhAT/Private/PhATRender.cpp
@@ -1,6 +1,6 @@
 // Copyright 1998-2016 Epic Games, Inc. All Rights Reserved.
 
-#include "PhATModule.h"
+#include "PhATPrivatePCH.h"
 #include "PhATSharedData.h"
 #include "PhATHitProxies.h"
 #include "PhATEdSkeletalMeshComponent.h"
diff --git a/Engine/Source/Editor/PhAT/Private/PhATSharedData.cpp b/Engine/Source/Editor/PhAT/Private/PhATSharedData.cpp
--- a/Engine/Source/Editor/PhAT/Private/PhATSharedData.cpp
+++ b/Engine/Source/Editor/PhAT/Private/PhATSharedData.cpp
@@ -1,5 +1,6 @@
 // Copyright 1998-2016 Epic Games, Inc. All Rights Reserved.
 
+#include "PhATPrivatePCH.h"
 #include "PhATModule.h"
 #include "PhysicsPublic.h"
 #include "EditorSupportDelegates.h"
@@ -931,8 +932,8 @@ void FPhATSharedData::InitConstraintSetup(UPhysicsConstraintTemplate* Constraint
 
 	ConstraintSetup->DefaultInstance.ConstraintBone2 = ParentBodySetup->BoneName;
 	ConstraintSetup->DefaultInstance.Pos2 = RelTM.GetOrigin();
-	ConstraintSetup->DefaultInstance.PriAxis2 = RelTM.GetScaledAxis( EAxis::X );
-	ConstraintSetup->DefaultInstance.SecAxis2 = RelTM.GetScaledAxis( EAxis::Y );
+	ConstraintSetup->DefaultInstance.PriAxis2 = RelTM.GetUnitAxis( EAxis::X );
+	ConstraintSetup->DefaultInstance.SecAxis2 = RelTM.GetUnitAxis( EAxis::Y );
 
 	// Disable collision between constrained bodies by default.
 	SetCollisionBetween(ChildBodyIndex, ParentBodyIndex, false);
diff --git a/Engine/Source/Editor/PhAT/Public/PhATModule.h b/Engine/Source/Editor/PhAT/Public/PhATModule.h
--- a/Engine/Source/Editor/PhAT/Public/PhATModule.h
+++ b/Engine/Source/Editor/PhAT/Public/PhATModule.h
@@ -1,17 +1,9 @@
 // Copyright 1998-2016 Epic Games, Inc. All Rights Reserved.
 
-#ifndef __PhATModule_h__
-#define __PhATModule_h__
+#pragma once
 
-#include "UnrealEd.h"
 #include "ModuleInterface.h"
-#include "Toolkits/AssetEditorToolkit.h"
-#include "Toolkits/IToolkit.h"
-#include "IPhAT.h"
-
-extern const FName PhATAppIdentifier;
-
-class FPhAT;
+#include "AssetEditorToolkit.h"
 
 DECLARE_LOG_CATEGORY_EXTERN(LogPhAT, Log, All);
 
@@ -25,10 +17,9 @@ class IPhATModule : public IModuleInterface,
 {
 public:
 	/** Creates a new PhAT instance */
-	virtual TSharedRef<IPhAT> CreatePhAT(const EToolkitMode::Type Mode, const TSharedPtr< IToolkitHost >& InitToolkitHost, UPhysicsAsset* PhysicsAsset) = 0;
+	virtual TSharedRef<class IPhAT> CreatePhAT(const EToolkitMode::Type Mode, const TSharedPtr< IToolkitHost >& InitToolkitHost, UPhysicsAsset* PhysicsAsset) = 0;
 
 	/** Opens a "New Asset/Body" modal dialog window */
 	virtual void OpenNewBodyDlg(struct FPhysAssetCreateParams* NewBodyData, EAppReturnType::Type* NewBodyResponse) = 0;
 };
 
-#endif // __PhATModule_h__
diff --git a/Engine/Source/Editor/PropertyEditor/Private/UserInterface/PropertyEditor/SPropertyEditorAsset.cpp b/Engine/Source/Editor/PropertyEditor/Private/UserInterface/PropertyEditor/SPropertyEditorAsset.cpp
--- a/Engine/Source/Editor/PropertyEditor/Private/UserInterface/PropertyEditor/SPropertyEditorAsset.cpp
+++ b/Engine/Source/Editor/PropertyEditor/Private/UserInterface/PropertyEditor/SPropertyEditorAsset.cpp
@@ -201,7 +201,7 @@ void SPropertyEditorAsset::Construct( const FArguments& InArgs, const TSharedPtr
 				if (Obj->HasAllFlags(RF_ClassDefaultObject))
 				{
 					IsEnabledAttribute.Set(false);
-					TooltipAttribute.Set(LOCTEXT("VariableHasDisableEditOnTemplate", "Editing this value in a Class Default Object is not allowed"));
+					TooltipAttribute.Set(LOCTEXT("VariableHasDisableEditOnTemplateTooltip", "Editing this value in a Class Default Object is not allowed"));
 					break;
 				}
 
diff --git a/Engine/Source/Editor/SceneOutliner/Private/FolderTreeItem.cpp b/Engine/Source/Editor/SceneOutliner/Private/FolderTreeItem.cpp
--- a/Engine/Source/Editor/SceneOutliner/Private/FolderTreeItem.cpp
+++ b/Engine/Source/Editor/SceneOutliner/Private/FolderTreeItem.cpp
@@ -262,7 +262,7 @@ void FFolderTreeItem::Delete()
 		FName ParentPath;
 	};
 
-	const FScopedTransaction Transaction( LOCTEXT("DeleteFolder", "Delete Folder") );
+	const FScopedTransaction Transaction( LOCTEXT("DeleteFolderTransaction", "Delete Folder") );
 
 	FResetActorFolders ResetFolders(GetParentPath(Path));
 	for (auto& Child : GetChildren())
diff --git a/Engine/Source/Editor/SceneOutliner/Private/SSceneOutliner.cpp b/Engine/Source/Editor/SceneOutliner/Private/SSceneOutliner.cpp
--- a/Engine/Source/Editor/SceneOutliner/Private/SSceneOutliner.cpp
+++ b/Engine/Source/Editor/SceneOutliner/Private/SSceneOutliner.cpp
@@ -56,13 +56,13 @@ namespace SceneOutliner
 				switch(World->GetNetMode())
 				{
 					case NM_Client:
-						if(WorldContext)
+						if (WorldContext)
 						{
 							PostFix = FText::Format(LOCTEXT("ClientPostfixFormat", "(Client {0})"), FText::AsNumber(WorldContext->PIEInstance - 1));
 						}
 						else
 						{
-							PostFix = LOCTEXT("ClientPostfixFormat", "(Client)");
+							PostFix = LOCTEXT("ClientPostfix", "(Client)");
 						}
 						break;
 					case NM_DedicatedServer:
diff --git a/Engine/Source/Editor/SequenceRecorder/Private/ActorRecording.cpp b/Engine/Source/Editor/SequenceRecorder/Private/ActorRecording.cpp
--- a/Engine/Source/Editor/SequenceRecorder/Private/ActorRecording.cpp
+++ b/Engine/Source/Editor/SequenceRecorder/Private/ActorRecording.cpp
@@ -140,14 +140,17 @@ void UActorRecording::GetSceneComponents(TArray<USceneComponent*>& OutArray, boo
 			auto ShouldRemovePredicate = [&](UActorComponent* PossiblyRemovedComponent)
 				{
 					// try to find a component with this name in the CDO
-					UActorComponent* const* FoundComponent = CDO->GetComponents().FindByPredicate([&](UActorComponent* SearchComponent) 
-					{ 
-						return SearchComponent->GetClass() == PossiblyRemovedComponent->GetClass() &&
-							   SearchComponent->GetFName() == PossiblyRemovedComponent->GetFName(); 
-					} );
+					for (UActorComponent* SearchComponent : CDO->GetComponents())
+					{
+						if (SearchComponent->GetClass() == PossiblyRemovedComponent->GetClass() &&
+							SearchComponent->GetFName() == PossiblyRemovedComponent->GetFName())
+						{
+							return false;
+						}
+					}
 
 					// remove if its not found
-					return FoundComponent == nullptr;
+					return true;
 				};
 
 			OutArray.RemoveAllSwap(ShouldRemovePredicate);
diff --git a/Engine/Source/Editor/Sequencer/Private/Sequencer.cpp b/Engine/Source/Editor/Sequencer/Private/Sequencer.cpp
--- a/Engine/Source/Editor/Sequencer/Private/Sequencer.cpp
+++ b/Engine/Source/Editor/Sequencer/Private/Sequencer.cpp
@@ -952,7 +952,7 @@ TRange<float> FSequencer::GetInOutRange() const
 
 void FSequencer::SetInOutRange(TRange<float> Range)
 {
-	const FScopedTransaction Transaction(LOCTEXT("SetPlaybackRange_Transaction", "Set In/Out Range"));
+	const FScopedTransaction Transaction(LOCTEXT("SetInOutRange_Transaction", "Set In/Out Range"));
 	UMovieScene* FocusedMovieScene = GetFocusedMovieSceneSequence()->GetMovieScene();
 
 	FocusedMovieScene->SetInOutRange(Range);
@@ -2411,7 +2411,7 @@ void FSequencer::OnEndScrubbing()
 
 void FSequencer::OnBeginPlaybackRangeDrag()
 {
-	GEditor->BeginTransaction(LOCTEXT("SetPlaybackRange_Transaction", "Set playback range"));
+	GEditor->BeginTransaction(LOCTEXT("SetPlaybackRange_Transaction", "Set Playback Range"));
 }
 
 
@@ -2423,7 +2423,7 @@ void FSequencer::OnEndPlaybackRangeDrag()
 
 void FSequencer::OnBeginInOutRangeDrag()
 {
-	GEditor->BeginTransaction(LOCTEXT("SetSelectionRange_Transaction", "Set selection range"));
+	GEditor->BeginTransaction(LOCTEXT("SetSelectionRange_Transaction", "Set Selection Range"));
 }
 
 
@@ -3732,7 +3732,7 @@ bool FSequencer::IsNodeLocked() const
 
 void FSequencer::ConvertToSpawnable(TSharedRef<FSequencerObjectBindingNode> NodeToBeConverted)
 {
-	const FScopedTransaction Transaction( LOCTEXT("ConvertSelectedSpawnable", "Convert Node to Spawnables") );
+	const FScopedTransaction Transaction( LOCTEXT("ConvertSelectedNodeSpawnable", "Convert Node to Spawnables") );
 
 	// Ensure we're in a non-possessed state
 	RootMovieSceneSequenceInstance->RestoreState(*this);
@@ -3744,7 +3744,7 @@ void FSequencer::ConvertToSpawnable(TSharedRef<FSequencerObjectBindingNode> Node
 void FSequencer::ConvertSelectedNodesToSpawnables()
 {
 	// @todo sequencer: Undo doesn't seem to be working at all
-	const FScopedTransaction Transaction( LOCTEXT("ConvertSelectedSpawnable", "Convert Selected Nodes to Spawnables") );
+	const FScopedTransaction Transaction( LOCTEXT("ConvertSelectedNodesSpawnable", "Convert Selected Nodes to Spawnables") );
 
 	UMovieScene* MovieScene = GetFocusedMovieSceneSequence()->GetMovieScene();
 
diff --git a/Engine/Source/Editor/Sequencer/Private/SequencerContextMenus.cpp b/Engine/Source/Editor/Sequencer/Private/SequencerContextMenus.cpp
--- a/Engine/Source/Editor/Sequencer/Private/SequencerContextMenus.cpp
+++ b/Engine/Source/Editor/Sequencer/Private/SequencerContextMenus.cpp
@@ -1474,7 +1474,7 @@ void FPasteFromHistoryContextMenu::PopulateMenu(FMenuBuilder& MenuBuilder)
 	// Copy a reference to the context menu by value into each lambda handler to ensure the type stays alive until the menu is closed
 	TSharedRef<FPasteFromHistoryContextMenu> Shared = AsShared();
 
-	MenuBuilder.BeginSection("SequencerPasteHistory", LOCTEXT("PasteFromHistory", "Paste History"));
+	MenuBuilder.BeginSection("SequencerPasteHistory", LOCTEXT("PasteFromHistory", "Paste From History"));
 
 	for (int32 Index = Sequencer->GetClipboardStack().Num() - 1; Index >= 0; --Index)
 	{
diff --git a/Engine/Source/Editor/UnrealEd/Classes/Editor/EditorEngine.h b/Engine/Source/Editor/UnrealEd/Classes/Editor/EditorEngine.h
--- a/Engine/Source/Editor/UnrealEd/Classes/Editor/EditorEngine.h
+++ b/Engine/Source/Editor/UnrealEd/Classes/Editor/EditorEngine.h
@@ -394,10 +394,6 @@ public:
 	UPROPERTY()
 	TArray<class UActorFactory*> ActorFactories;
 
-	/** Actors that are being deleted and should processed in the global re-attach */
-	UPROPERTY()
-	TArray<class AActor*> ActorsForGlobalReregister;
-
 	/** The name of the file currently being opened in the editor. "" if no file is being opened. */
 	UPROPERTY()
 	FString UserOpenedFile;
diff --git a/Engine/Source/Editor/UnrealEd/Classes/Preferences/PersonaOptions.h b/Engine/Source/Editor/UnrealEd/Classes/Preferences/PersonaOptions.h
--- a/Engine/Source/Editor/UnrealEd/Classes/Preferences/PersonaOptions.h
+++ b/Engine/Source/Editor/UnrealEd/Classes/Preferences/PersonaOptions.h
@@ -21,6 +21,9 @@ class UNREALED_API UPersonaOptions : public UObject
 	UPROPERTY(EditAnywhere, config, Category=Options)
 	uint32 bShowSky:1;
 
+	UPROPERTY(EditAnywhere, config, Category = Options)
+	uint32 bAutoAlignFloorToMesh : 1;
+
 	UPROPERTY(EditAnywhere, config, Category=Options)
 	uint32 bShowGrid:1;
 
@@ -31,7 +34,7 @@ class UNREALED_API UPersonaOptions : public UObject
 	uint32 bMuteAudio:1;
 
 	// currently Stats can have None, Basic and Detailed. Please refer to EDisplayInfoMode.
-	UPROPERTY(EditAnywhere, config, Category = Options, meta=(ClampMin ="0", ClampMax = "2", UIMin = "0", UIMax = "2"))
+	UPROPERTY(EditAnywhere, config, Category = Options, meta=(ClampMin ="0", ClampMax = "3", UIMin = "0", UIMax = "3"))
 	int32 ShowMeshStats;
 
 	UPROPERTY(EditAnywhere, config, Category=Options)
@@ -63,6 +66,7 @@ public:
 	void SetShowGrid( bool bInShowGrid );
 	void SetHighlightOrigin( bool bInHighlightOrigin );
 	void SetShowFloor( bool bInShowFloor );
+	void SetAutoAlignFloorToMesh(bool bInAutoAlignFloorToMesh);
 	void SetShowSky( bool bInShowSky );
 	void SetMuteAudio( bool bInMuteAudio );
 	void SetGridSize( int32 InGridSize );
diff --git a/Engine/Source/Editor/UnrealEd/Private/Commandlets/ContentCommandlets.cpp b/Engine/Source/Editor/UnrealEd/Private/Commandlets/ContentCommandlets.cpp
--- a/Engine/Source/Editor/UnrealEd/Private/Commandlets/ContentCommandlets.cpp
+++ b/Engine/Source/Editor/UnrealEd/Private/Commandlets/ContentCommandlets.cpp
@@ -1273,14 +1273,14 @@ int32 UWrangleContentCommandlet::Main( const FString& Params )
 			{
 				if (It.Key() == TEXT("Package"))
 				{
-					PackagesToFullyLoad.Add(*It.Key().ToString(), *It.Value());
-					if ( FindPackage(NULL, *It.Value()) )
+					PackagesToFullyLoad.Add(*It.Key().ToString(), *It.Value().GetValue());
+					if ( FindPackage(NULL, *It.Value().GetValue()) )
 					{
-						UE_LOG(LogContentCommandlet, Warning, TEXT("Startup package '%s' was loaded"), *It.Value());
+						UE_LOG(LogContentCommandlet, Warning, TEXT("Startup package '%s' was loaded"), *It.Value().GetValue());
 					}
 					else
 					{
-						UE_LOG(LogContentCommandlet, Warning, TEXT("Startup package '%s' was not loaded during FStartupPackages::LoadAll..."), *It.Value());
+						UE_LOG(LogContentCommandlet, Warning, TEXT("Startup package '%s' was not loaded during FStartupPackages::LoadAll..."), *It.Value().GetValue());
 					}
 				}
 			}
@@ -1310,7 +1310,7 @@ int32 UWrangleContentCommandlet::Main( const FString& Params )
 		for (FConfigSectionMap::TIterator PackageIt(PackagesToFullyLoad); PackageIt; ++PackageIt)
 		{
 			// add dependencies for the per-map packages for this map (if any)
-			TArray<FString>* Packages = PerMapCookPackages.Find(PackageIt.Value());
+			TArray<FString>* Packages = PerMapCookPackages.Find(PackageIt.Value().GetValue());
 			if (Packages != NULL)
 			{
 				for (int32 PackageIndex = 0; PackageIndex < Packages->Num(); PackageIndex++)
@@ -1344,7 +1344,7 @@ int32 UWrangleContentCommandlet::Main( const FString& Params )
 			// there may be multiple sublevels to load if this package is a persistent level with sublevels
 			TArray<FString> PackagesToLoad;
 			// start off just loading this package (more may be added in the loop)
-			PackagesToLoad.Add(*PackageIt.Value());
+			PackagesToLoad.Add(*PackageIt.Value().GetValue());
 
 			for (int32 PackageIndex = 0; PackageIndex < PackagesToLoad.Num(); PackageIndex++)
 			{
diff --git a/Engine/Source/Editor/UnrealEd/Private/Commandlets/InternationalizationConditioningCommandlet.cpp b/Engine/Source/Editor/UnrealEd/Private/Commandlets/InternationalizationConditioningCommandlet.cpp
--- a/Engine/Source/Editor/UnrealEd/Private/Commandlets/InternationalizationConditioningCommandlet.cpp
+++ b/Engine/Source/Editor/UnrealEd/Private/Commandlets/InternationalizationConditioningCommandlet.cpp
@@ -69,22 +69,22 @@ void UInternationalizationConditioningCommandlet::FLocalizationFile::CompareToCo
 			for ( FConfigSection::TIterator It(MySection); It; ++It )
 			{
 				const FName Propname = It.Key();
-				const FString& PropValue = It.Value();
+				const FString& PropValue = It.Value().GetValue();
 
 				FString EscapedPropValue = PropValue.ReplaceQuotesWithEscapedQuotes();
 
 				// Find this key in the counterpart loc file
-				FString* OtherValue = OtherSection->Find(Propname);
+				FConfigValue* OtherValue = OtherSection->Find(Propname);
 				if ( OtherValue != NULL )
 				{
-					FString EscapedOtherValue = *OtherValue->ReplaceQuotesWithEscapedQuotes();
+					FString EscapedOtherValue = OtherValue->GetValue().ReplaceQuotesWithEscapedQuotes();
 
 					// If the counterpart has the same value as we do or is empty, the value is untranslated
-					if( OtherValue->IsEmpty() )
+					if( OtherValue->GetValue().IsEmpty() )
 					{
 						// If the entry is empty we do nothing for the time being.
 					}
-					else if ( PropValue == *OtherValue )
+					else if ( PropValue == OtherValue->GetValue() )
 					{
 						new(IdenticalProperties) FLocalizationFileEntry( Other->GetFilename(), LocSectionName, Propname.ToString(), EscapedPropValue, EscapedPropValue );
 					}
diff --git a/Engine/Source/Editor/UnrealEd/Private/CookOnTheFlyServer.cpp b/Engine/Source/Editor/UnrealEd/Private/CookOnTheFlyServer.cpp
--- a/Engine/Source/Editor/UnrealEd/Private/CookOnTheFlyServer.cpp
+++ b/Engine/Source/Editor/UnrealEd/Private/CookOnTheFlyServer.cpp
@@ -2800,7 +2800,7 @@ bool UCookOnTheFlyServer::GetCurrentIniVersionStrings( const ITargetPlatform* Ta
 					continue;
 				}
 
-				FString ProcessedIniValue = IniValue.Value;
+				FString ProcessedIniValue = IniValue.Value.GetValue();
 				// get rid of shit we don't care about
 				ProcessedIniValue.ReplaceInline(TEXT("{"), TEXT(""));
 				ProcessedIniValue.ReplaceInline(TEXT("}"), TEXT(""));
diff --git a/Engine/Source/Editor/UnrealEd/Private/EditorEngine.cpp b/Engine/Source/Editor/UnrealEd/Private/EditorEngine.cpp
--- a/Engine/Source/Editor/UnrealEd/Private/EditorEngine.cpp
+++ b/Engine/Source/Editor/UnrealEd/Private/EditorEngine.cpp
@@ -935,19 +935,6 @@ void UEditorEngine::Tick( float DeltaSeconds, bool bIdleMode )
 	FWorldContext& EditorContext = GetEditorWorldContext();
 	check( CurrentGWorld == EditorContext.World() );
 
-	// was there a reregister requested last frame?
-	if (bHasPendingGlobalReregister)
-	{
-		// make sure outstanding deletion has completed before the reregister
-		CollectGarbage(GARBAGE_COLLECTION_KEEPFLAGS);
-
-		//Only reregister actors whose replacement primitive is a child of the global reregister list
-		FGlobalComponentReregisterContext Reregister(ActorsForGlobalReregister);
-		ActorsForGlobalReregister.Empty();
-
-		bHasPendingGlobalReregister = false;
-	}
-
 	// early in the Tick() to get the callbacks for cvar changes called
 	IConsoleManager::Get().CallAllConsoleVariableSinks();
 
@@ -2526,7 +2513,7 @@ void UEditorEngine::LoadMapListFromIni(const FString& InSectionName, TArray<FStr
 		for (FConfigSectionMap::TConstIterator It(*MapListList) ; It ; ++It)
 		{
 			FName EntryType = It.Key();
-			const FString& EntryValue = It.Value();
+			const FString& EntryValue = It.Value().GetValue();
 
 			if (EntryType == NAME_Map)
 			{
diff --git a/Engine/Source/Editor/UnrealEd/Private/EditorSelectUtils.cpp b/Engine/Source/Editor/UnrealEd/Private/EditorSelectUtils.cpp
--- a/Engine/Source/Editor/UnrealEd/Private/EditorSelectUtils.cpp
+++ b/Engine/Source/Editor/UnrealEd/Private/EditorSelectUtils.cpp
@@ -250,9 +250,9 @@ void UUnrealEdEngine::SetActorSelectionFlags (AActor* InActor)
 		{
 			// If we have a 'child actor' component, want to update its visible selection state
 			UChildActorComponent* ChildActorComponent = Cast<UChildActorComponent>(Component);
-			if(ChildActorComponent != NULL && ChildActorComponent->ChildActor != NULL)
+			if(ChildActorComponent != NULL && ChildActorComponent->GetChildActor() != NULL)
 			{
-				SetActorSelectionFlags(ChildActorComponent->ChildActor);
+				SetActorSelectionFlags(ChildActorComponent->GetChildActor());
 			}
 
 			UPrimitiveComponent* PrimComponent = Cast<UPrimitiveComponent>(Component);
diff --git a/Engine/Source/Editor/UnrealEd/Private/Factories/ApexDestructibleAssetImport.cpp b/Engine/Source/Editor/UnrealEd/Private/Factories/ApexDestructibleAssetImport.cpp
--- a/Engine/Source/Editor/UnrealEd/Private/Factories/ApexDestructibleAssetImport.cpp
+++ b/Engine/Source/Editor/UnrealEd/Private/Factories/ApexDestructibleAssetImport.cpp
@@ -835,7 +835,7 @@ bool SetApexDestructibleAsset(UDestructibleMesh& DestructibleMesh, NxDestructibl
 
 	// Create initial bounding box based on expanded version of reference pose for meshes without physics assets. Can be overridden by artist.
 	FBox BoundingBox(SkelMeshImportDataPtr->Points.GetData(), SkelMeshImportDataPtr->Points.Num());
-	DestructibleMesh.Bounds= FBoxSphereBounds(BoundingBox);
+	DestructibleMesh.SetImportedBounds(FBoxSphereBounds(BoundingBox));
 
 	// Store whether or not this mesh has vertex colors
 	DestructibleMesh.bHasVertexColors = SkelMeshImportDataPtr->bHasVertexColors;
@@ -999,7 +999,20 @@ UDestructibleMesh* ImportDestructibleMeshFromApexDestructibleAsset(UObject* InPa
 	FlushRenderingCommands();
 
 	UDestructibleMesh* DestructibleMesh = FindObject<UDestructibleMesh>(InParent, *Name.ToString());
-	if (DestructibleMesh == NULL)
+	if(DestructibleMesh)
+	{
+		// we found an existing mesh (reimport), clean up LODModels for the import process
+		if(FSkeletalMeshResource* ImportedResource = DestructibleMesh->GetImportedResource())
+		{
+			for(FStaticLODModel& LODModel : ImportedResource->LODModels)
+			{
+				LODModel.ReleaseResources();
+			}
+
+			ImportedResource->LODModels.Empty(1);
+		}
+	}
+	else
 	{
 		// Create the new UDestructibleMesh object if the one with the same name does not exist
 		DestructibleMesh = NewObject<UDestructibleMesh>(InParent, Name, Flags);
diff --git a/Engine/Source/Editor/UnrealEd/Private/Factories/PackFactory.cpp b/Engine/Source/Editor/UnrealEd/Private/Factories/PackFactory.cpp
--- a/Engine/Source/Editor/UnrealEd/Private/Factories/PackFactory.cpp
+++ b/Engine/Source/Editor/UnrealEd/Private/Factories/PackFactory.cpp
@@ -190,7 +190,7 @@ namespace PackFactoryHelper
 				if (SettingPair.Key.ToString().Contains("ActionMappings"))
 				{
 					FInputActionKeyMapping ActionKeyMapping;
-					ActionMappingsProp->Inner->ImportText(*SettingPair.Value, &ActionKeyMapping, PPF_None, nullptr);
+					ActionMappingsProp->Inner->ImportText(*SettingPair.Value.GetValue(), &ActionKeyMapping, PPF_None, nullptr);
 
 					if (InputSettingsCDO->ActionMappings.FindByPredicate(FMatchMappingByName(ActionKeyMapping.ActionName)) == nullptr)
 					{
@@ -200,7 +200,7 @@ namespace PackFactoryHelper
 				else if (SettingPair.Key.ToString().Contains("AxisMappings"))
 				{
 					FInputAxisKeyMapping AxisKeyMapping;
-					AxisMappingsProp->Inner->ImportText(*SettingPair.Value, &AxisKeyMapping, PPF_None, nullptr);
+					AxisMappingsProp->Inner->ImportText(*SettingPair.Value.GetValue(), &AxisKeyMapping, PPF_None, nullptr);
 
 					if (InputSettingsCDO->AxisMappings.FindByPredicate(FMatchMappingByName(AxisKeyMapping.AxisName)) == nullptr)
 					{
@@ -239,9 +239,9 @@ namespace PackFactoryHelper
 		FConfigSection* RedirectsSection = PackConfig.Find("Redirects");
 		if (RedirectsSection)
 		{	
-			if (FString* GameName = RedirectsSection->Find("GameName"))
+			if (FConfigValue* GameName = RedirectsSection->Find("GameName"))
 			{
-				ConfigParameters.GameName = *GameName;
+				ConfigParameters.GameName = GameName->GetValue();
 			}
 		}
 
@@ -252,8 +252,8 @@ namespace PackFactoryHelper
 			{
 				if (FilePair.Key.ToString().Contains("Files"))
 				{
-					FString Filename = FPaths::GetCleanFilename(FilePair.Value);
-					FString Directory = FPaths::RootDir() / FPaths::GetPath(FilePair.Value);
+					FString Filename = FPaths::GetCleanFilename(FilePair.Value.GetValue());
+					FString Directory = FPaths::RootDir() / FPaths::GetPath(FilePair.Value.GetValue());
 					FPaths::MakeStandardFilename(Directory);
 					FPakFile::MakeDirectoryFromPath(Directory);
 
diff --git a/Engine/Source/Editor/UnrealEd/Private/Fbx/FbxMainExport.cpp b/Engine/Source/Editor/UnrealEd/Private/Fbx/FbxMainExport.cpp
--- a/Engine/Source/Editor/UnrealEd/Private/Fbx/FbxMainExport.cpp
+++ b/Engine/Source/Editor/UnrealEd/Private/Fbx/FbxMainExport.cpp
@@ -1286,7 +1286,7 @@ FbxNode* FFbxExporter::ExportActor(AActor* Actor, AMatineeActor* InMatineeActor,
 				{
 					ComponentsToExport.Add(Component);
 				}
-				else if (ChildActorComp && ChildActorComp->ChildActor)
+				else if (ChildActorComp && ChildActorComp->GetChildActor())
 				{
 					ComponentsToExport.Add(ChildActorComp);
 				}
@@ -1370,10 +1370,10 @@ FbxNode* FFbxExporter::ExportActor(AActor* Actor, AMatineeActor* InMatineeActor,
 					FillFbxLightAttribute(Light, ActorNode, Cast<ULightComponent>(Component));
 					ExportNode->SetNodeAttribute(Light);
 				}
-				else if (ChildActorComp && ChildActorComp->ChildActor)
+				else if (ChildActorComp && ChildActorComp->GetChildActor())
 				{
-					FbxNode* ChildActorNode = ExportActor(ChildActorComp->ChildActor, InMatineeActor, true);
-					FbxActors.Add(ChildActorComp->ChildActor, ChildActorNode);
+					FbxNode* ChildActorNode = ExportActor(ChildActorComp->GetChildActor(), InMatineeActor, true);
+					FbxActors.Add(ChildActorComp->GetChildActor(), ChildActorNode);
 				}
 			}
 		}
diff --git a/Engine/Source/Editor/UnrealEd/Private/Fbx/FbxSkeletalMeshImport.cpp b/Engine/Source/Editor/UnrealEd/Private/Fbx/FbxSkeletalMeshImport.cpp
--- a/Engine/Source/Editor/UnrealEd/Private/Fbx/FbxSkeletalMeshImport.cpp
+++ b/Engine/Source/Editor/UnrealEd/Private/Fbx/FbxSkeletalMeshImport.cpp
@@ -1303,7 +1303,7 @@ USkeletalMesh* UnFbx::FFbxImporter::ImportSkeletalMesh(UObject* InParent, TArray
 	//BEN const int32 CoordToTuck = bAssumeMayaCoordinates ? 1 : 2;
 	//BEN BoundingBox.Min[CoordToTuck]	= Temp.Min[CoordToTuck] + 0.1f*(Temp.Min[CoordToTuck] - MidMesh[CoordToTuck]);
 	BoundingBox.Min[2]	= Temp.Min[2] + 0.1f*(Temp.Min[2] - MidMesh[2]);
-	SkeletalMesh->Bounds= FBoxSphereBounds(BoundingBox);
+	SkeletalMesh->SetImportedBounds(FBoxSphereBounds(BoundingBox));
 
 	// Store whether or not this mesh has vertex colors
 	SkeletalMesh->bHasVertexColors = SkelMeshImportDataPtr->bHasVertexColors;
diff --git a/Engine/Source/Editor/UnrealEd/Private/PackageTools.cpp b/Engine/Source/Editor/UnrealEd/Private/PackageTools.cpp
--- a/Engine/Source/Editor/UnrealEd/Private/PackageTools.cpp
+++ b/Engine/Source/Editor/UnrealEd/Private/PackageTools.cpp
@@ -28,7 +28,7 @@ extern CORE_API EditorPostReachabilityAnalysisCallbackType EditorPostReachabilit
 namespace PackageTools
 {
 	/** State passed to RestoreStandaloneOnReachableObjects. */
-	static UPackage* PackageBeingUnloaded = NULL;
+	static UPackage* PackageBeingUnloaded = nullptr;
 	static TMap<UObject*,UObject*> ObjectsThatHadFlagsCleared;
 
 	/**
@@ -37,38 +37,30 @@ namespace PackageTools
 	 */
 	void RestoreStandaloneOnReachableObjects()
 	{
-		TArray<UObject*> ObjectsInPackage;
-		GetObjectsWithOuter(PackageBeingUnloaded, ObjectsInPackage);
-		for ( UObject* Object : ObjectsInPackage )
-		{
-			if ( !Object->IsUnreachable() )
+		ForEachObjectWithOuter(PackageBeingUnloaded, [](UObject* Object)
 			{
 				if ( ObjectsThatHadFlagsCleared.Find(Object) )
 				{
 					Object->SetFlags(RF_Standalone);
 				}
-			}
-		}
+			},
+			true, RF_NoFlags, EInternalObjectFlags::Unreachable);
 	}
 
-
-
 	/**
 	 * Filters the global set of packages.
 	 *
 	 * @param	OutGroupPackages			The map that receives the filtered list of group packages.
 	 * @param	OutPackageList				The array that will contain the list of filtered packages.
 	 */
-	void GetFilteredPackageList( TSet<const UPackage*>& OutFilteredPackageMap,
-								 TSet<UPackage*>* OutGroupPackages,
-								 TArray<UPackage*>& OutPackageList )
+	void GetFilteredPackageList(TSet<UPackage*>& OutFilteredPackageMap)
 	{
 		// The UObject list is iterated rather than the UPackage list because we need to be sure we are only adding
 		// group packages that contain things the generic browser cares about.  The packages are derived by walking
 		// the outer chain of each object.
 
 		// Assemble a list of packages.  Only show packages that match the current resource type filter.
-		for (auto* Obj : TObjectRange<UObject>())
+		for (UObject* Obj : TObjectRange<UObject>())
 		{
 			// Make sure that we support displaying this object type
 			bool bIsSupported = ObjectTools::IsObjectBrowsable( Obj );
@@ -79,26 +71,8 @@ namespace PackageTools
 				{
 					OutFilteredPackageMap.Add( ObjectPackage );
 				}
-
-				if( OutGroupPackages != NULL )
-				{
-					for ( UObject* NextOuter = Obj->GetOuter(); NextOuter && NextOuter->GetOuter(); NextOuter = NextOuter->GetOuter() )
-					{
-						UPackage* NextGroup = Cast<UPackage>(NextOuter);
-						if ( NextGroup != NULL && !OutGroupPackages->Contains(NextGroup) )
-						{
-							OutGroupPackages->Add(NextGroup);
-						}
-					}
-				}
 			}
 		}
-
-		// Make a TArray copy of PackageMap.
-		for ( TSet<const UPackage*>::TConstIterator It( OutFilteredPackageMap ) ; It ; ++It )
-		{
-			OutPackageList.Add( const_cast<UPackage*>( *It ) );
-		}
 	}
 
 	/**
@@ -109,43 +83,32 @@ namespace PackageTools
 	 * @param	OutObjects			[out] Receives the list of objects
 	 * @param	bMustBeBrowsable	If specified, does a check to see if object is browsable. Defaults to true.
 	 */
-	void GetObjectsInPackages( const TArray<UPackage*>* InPackages,
-							   TArray<UObject*>& OutObjects )
+	void GetObjectsInPackages( const TArray<UPackage*>* InPackages, TArray<UObject*>& OutObjects )
 	{
-		// Iterate over all objects.
-		for( TObjectIterator<UObject> It ; It ; ++It )
+		if (InPackages)
 		{
-			UObject* Obj = *It;
-
-			// Filter out invalid objects early.
-			if( !ObjectTools::IsObjectBrowsable( Obj ) )
+			for (UPackage* Package : *InPackages)
 			{
-				continue;
+				ForEachObjectWithOuter(Package,[&OutObjects](UObject* Obj)
+					{
+						if (ObjectTools::IsObjectBrowsable(Obj))
+						{
+							OutObjects.Add(Obj);
+						}
+					});
 			}
-
-			// Should we filter based on a list of allowed packages?
-			if( InPackages != NULL )
+		}
+		else
+		{
+			for (TObjectIterator<UObject> It; It; ++It)
 			{
-				// Make sure this object resides in one of the specified packages.
-				bool bIsInPackage = false;
-				for ( int32 PackageIndex = 0 ; PackageIndex < InPackages->Num() ; ++PackageIndex )
-				{
-					const UPackage* Package = ( *InPackages )[PackageIndex];
-					if ( Obj->IsIn(Package) )
-					{
-						bIsInPackage = true;
-						break;
-					}
-				}
+				UObject* Obj = *It;
 
-				if( !bIsInPackage )
+				if (ObjectTools::IsObjectBrowsable(Obj))
 				{
-					continue;
+					OutObjects.Add(Obj);
 				}
 			}
-
-			// Add to the list.
-			OutObjects.Add( Obj );
 		}
 	}
 
@@ -567,14 +530,13 @@ namespace PackageTools
 		bool bHasExternalPackageRefs = false;
 
 		// Find all external packages
-		TSet< const UPackage* > FilteredPackageMap;
-		TArray< UPackage* > LoadedPackageList;
-		GetFilteredPackageList( FilteredPackageMap, NULL, LoadedPackageList );
+		TSet<UPackage*> FilteredPackageMap;
+		GetFilteredPackageList(FilteredPackageMap);
 
 		TArray< UPackage* > ExternalPackages;
-		for(int32 pkgIdx = 0; pkgIdx < LoadedPackageList.Num(); ++pkgIdx)
+		ExternalPackages.Reserve(FilteredPackageMap.Num());
+		for (UPackage* Pkg : FilteredPackageMap)
 		{
-			UPackage* Pkg = LoadedPackageList[ pkgIdx ];
 			FString OutFilename;
 			const FString PackageName = Pkg->GetName();
 			const FGuid PackageGuid = Pkg->GetGuid();
@@ -598,38 +560,27 @@ namespace PackageTools
 		}
 		else
 		{
-			for( TObjectIterator<UObject> It ; It ; ++It )
-			{
-				// Gather all the objects in our level
-				UObject* Obj = *It;
-				if ( Obj->IsIn(LevelToCheck) )
-				{
-					ObjectsInPackageToCheck.Add(Obj);
-				}
+			GetObjectsWithOuter(LevelToCheck, ObjectsInPackageToCheck);
 
-				// Gather all objects in any loaded external packages			
-				for ( int32 PackageIndex = 0 ; PackageIndex < ExternalPackages.Num() ; ++PackageIndex )
-				{
-					const UPackage* Package = ExternalPackages[PackageIndex];
-					if ( Obj->IsIn(Package) && ObjectTools::IsObjectBrowsable( Obj ) )
-					{
-						ObjectsInExternalPackages.Add( Obj );
-						break;
-					}
-				}
+			// Gather all objects in any loaded external packages			
+			for (const UPackage* Package : ExternalPackages)
+			{
+				ForEachObjectWithOuter(Package,
+									   [&ObjectsInExternalPackages](UObject* Obj)
+									   {
+										   if (ObjectTools::IsObjectBrowsable(Obj))
+										   {
+											   ObjectsInExternalPackages.Add(Obj);
+										   }
+									   });
 			}
 		}
 
-		// compare
-		for(int32 ExtObjIdx = 0; ExtObjIdx < ObjectsInExternalPackages.Num(); ++ExtObjIdx)
+		// only check objects which are in packages to be saved.  This should greatly reduce the overhead by not searching through objects we don't intend to save
+		for (UObject* CheckObject : ObjectsInPackageToCheck)
 		{
-			UObject* ExternalObject = ObjectsInExternalPackages[ ExtObjIdx ];
-
-			// only check objects which are in packages to be saved.  This should greatly reduce the overhead by not searching through objects we don't intend to save
-			for(int32 CheckObjIdx = 0; CheckObjIdx < ObjectsInPackageToCheck.Num(); ++CheckObjIdx)
+			for (UObject* ExternalObject : ObjectsInExternalPackages)
 			{
-				UObject* CheckObject = ObjectsInPackageToCheck[ CheckObjIdx ];
-
 				FArchiveFindCulprit ArFind( ExternalObject, CheckObject, false );
 				if( ArFind.GetCount() > 0 )
 				{
@@ -642,6 +593,7 @@ namespace PackageTools
 						OutObjectsWithExternalRefs->Add( CheckObject );
 					}
 					bHasExternalPackageRefs = true;
+					break;
 				}
 			}
 		}
diff --git a/Engine/Source/Editor/UnrealEd/Private/PreferenceStubs.cpp b/Engine/Source/Editor/UnrealEd/Private/PreferenceStubs.cpp
--- a/Engine/Source/Editor/UnrealEd/Private/PreferenceStubs.cpp
+++ b/Engine/Source/Editor/UnrealEd/Private/PreferenceStubs.cpp
@@ -46,6 +46,8 @@ UPersonaOptions::UPersonaOptions(const FObjectInitializer& ObjectInitializer)
 	SectionTimingNodeColor = FLinearColor(0.0f, 1.0f, 0.0f);
 	NotifyTimingNodeColor = FLinearColor(1.0f, 0.0f, 0.0f);
 	BranchingPointTimingNodeColor = FLinearColor(0.5f, 1.0f, 1.0f);
+
+	bAutoAlignFloorToMesh = true;
 }
 
 void UPersonaOptions::SetViewportBackgroundColor( const FLinearColor& InViewportBackgroundColor)
@@ -83,6 +85,13 @@ void UPersonaOptions::SetShowFloor( bool bInShowFloor )
 	bShowFloor = bInShowFloor;
 	SaveConfig();
 }
+
+void UPersonaOptions::SetAutoAlignFloorToMesh(bool bInAutoAlignFloorToMesh)
+{
+	bAutoAlignFloorToMesh = bInAutoAlignFloorToMesh;
+	SaveConfig();
+}
+
 void UPersonaOptions::SetShowSky( bool bInShowSky )
 {
 	bShowSky = bInShowSky;
diff --git a/Engine/Source/Editor/UnrealEd/Private/SSkeletonWidget.cpp b/Engine/Source/Editor/UnrealEd/Private/SSkeletonWidget.cpp
--- a/Engine/Source/Editor/UnrealEd/Private/SSkeletonWidget.cpp
+++ b/Engine/Source/Editor/UnrealEd/Private/SSkeletonWidget.cpp
@@ -1391,7 +1391,7 @@ void SBasePoseViewport::SetSkeleton(USkeleton* Skeleton)
 				NewPosition.Normalize();
 				if(PreviewSkeletalMesh)
 				{
-					NewPosition *= (PreviewSkeletalMesh->Bounds.SphereRadius*1.5f);
+					NewPosition *= (PreviewSkeletalMesh->GetImportedBounds().SphereRadius*1.5f);
 				}
 				Client->SetViewLocation(NewPosition);
 			}
diff --git a/Engine/Source/Editor/UnrealEd/Public/PackageTools.h b/Engine/Source/Editor/UnrealEd/Public/PackageTools.h
--- a/Engine/Source/Editor/UnrealEd/Public/PackageTools.h
+++ b/Engine/Source/Editor/UnrealEd/Public/PackageTools.h
@@ -17,11 +17,8 @@ namespace PackageTools
 	 * @param	OutGroupPackages			The map that receives the filtered list of group packages.
 	 * @param	OutPackageList				The array that will contain the list of filtered packages.
 	 */
-	UNREALED_API void GetFilteredPackageList( TSet<const UPackage*> &OutFilteredPackageMap,
-								 TSet<UPackage*>* OutGroupPackages,
-								 TArray<UPackage*> &OutPackageList );
-
-
+	UNREALED_API void GetFilteredPackageList(TSet<UPackage*>& OutFilteredPackageMap);
+	
 	/**
 	 * Fills the OutObjects list with all valid objects that are supported by the current
 	 * browser settings and that reside withing the set of specified packages.
@@ -29,8 +26,7 @@ namespace PackageTools
 	 * @param	InPackages			Filters objects based on package.
 	 * @param	OutObjects			[out] Receives the list of objects
 	 */
-	UNREALED_API void GetObjectsInPackages( const TArray<UPackage*>* InPackages,
-							   TArray<UObject*>& OutObjects );
+	UNREALED_API void GetObjectsInPackages( const TArray<UPackage*>* InPackages, TArray<UObject*>& OutObjects );
 
 	/**
 	 * Handles fully loading passed in packages.
diff --git a/Engine/Source/Programs/UnrealHeaderTool/Private/HeaderParser.cpp b/Engine/Source/Programs/UnrealHeaderTool/Private/HeaderParser.cpp
--- a/Engine/Source/Programs/UnrealHeaderTool/Private/HeaderParser.cpp
+++ b/Engine/Source/Programs/UnrealHeaderTool/Private/HeaderParser.cpp
@@ -7026,6 +7026,7 @@ FHeaderParser::FHeaderParser(FFeedbackContext* InWarn, FUHTMakefile& InUHTMakefi
 	DelegateParameterCountStrings.Add(TEXT("_SixParams"));
 	DelegateParameterCountStrings.Add(TEXT("_SevenParams"));
 	DelegateParameterCountStrings.Add(TEXT("_EightParams"));
+	DelegateParameterCountStrings.Add(TEXT("_NineParams"));
 
 	FString Version;
 	if (GConfig->GetString(TEXT("GeneratedCodeVersion"), TEXT("UnrealHeaderTool"), Version, GEngineIni))
diff --git a/Engine/Source/Runtime/AIModule/Classes/AIController.h b/Engine/Source/Runtime/AIModule/Classes/AIController.h
--- a/Engine/Source/Runtime/AIModule/Classes/AIController.h
+++ b/Engine/Source/Runtime/AIModule/Classes/AIController.h
@@ -84,28 +84,28 @@ class AIMODULE_API AAIController : public AController, public IAIPerceptionListe
 
 protected:
 	FFocusKnowledge	FocusInformation;
-	
+
 	/** By default AI's logic gets stopped when controlled Pawn is unpossesed. Setting this flag to false
 	 *	will make AI logic persist past loosing controll over a pawn */
 	UPROPERTY(EditAnywhere, BlueprintReadWrite, Category = AI)
-	uint32 bStopAILogicOnUnposses : 1;
+		uint32 bStopAILogicOnUnposses : 1;
 
 public:
 	/** used for alternating LineOfSight traces */
 	UPROPERTY()
-	mutable uint32 bLOSflag:1;
+		mutable uint32 bLOSflag : 1;
 
 	/** Skip extra line of sight traces to extremities of target being checked. */
 	UPROPERTY()
-	uint32 bSkipExtraLOSChecks:1;
+		uint32 bSkipExtraLOSChecks : 1;
 
 	/** Is strafing allowed during movement? */
 	UPROPERTY()
-	uint32 bAllowStrafe:1;	
+		uint32 bAllowStrafe : 1;
 
 	/** Specifies if this AI wants its own PlayerState. */
 	UPROPERTY()
-	uint32 bWantsPlayerState:1;
+		uint32 bWantsPlayerState : 1;
 
 	/** Copy Pawn rotation to ControlRotation, if there is no focus point. */
 	UPROPERTY(EditAnywhere, BlueprintReadWrite, Category = AI)
@@ -115,31 +115,31 @@ private_subobject:
 
 	/** Component used for moving along a path. */
 	DEPRECATED_FORGAME(4.6, "PathFollowingComponent should not be accessed directly, please use GetPathFollowingComponent() function instead. PathFollowingComponent will soon be private and your code will not compile.")
-	UPROPERTY(VisibleDefaultsOnly, Category = AI)
-	UPathFollowingComponent* PathFollowingComponent;
+		UPROPERTY(VisibleDefaultsOnly, Category = AI)
+		UPathFollowingComponent* PathFollowingComponent;
 
 public:
 
 	/** Component responsible for behaviors. */
 	UPROPERTY(BlueprintReadWrite, Category = AI)
-	UBrainComponent* BrainComponent;
+		UBrainComponent* BrainComponent;
 
 	DEPRECATED_FORGAME(4.8, "PerceptionComponent should not be accessed directly, please use GetAIPerceptionComponent() function instead. PerceptionComponent will soon be private and your code will not compile.")
-	UPROPERTY(VisibleDefaultsOnly, Category = AI)
-	UAIPerceptionComponent* PerceptionComponent;
+		UPROPERTY(VisibleDefaultsOnly, Category = AI)
+		UAIPerceptionComponent* PerceptionComponent;
 
 private_subobject:
 	DEPRECATED_FORGAME(4.6, "ActionsComp should not be accessed directly, please use GetActionsComp() function instead. ActionsComp will soon be private and your code will not compile.")
-	UPROPERTY(BlueprintReadOnly, Category = AI, meta = (AllowPrivateAccess = "true"))
-	UPawnActionsComponent* ActionsComp;
+		UPROPERTY(BlueprintReadOnly, Category = AI, meta = (AllowPrivateAccess = "true"))
+		UPawnActionsComponent* ActionsComp;
 
 protected:
 	/** blackboard */
 	UPROPERTY(BlueprintReadOnly, Category = AI, meta = (AllowPrivateAccess = "true"))
-	UBlackboardComponent* Blackboard;
+		UBlackboardComponent* Blackboard;
 
 	UPROPERTY()
-	UGameplayTasksComponent* CachedGameplayTasksComponent;
+		UGameplayTasksComponent* CachedGameplayTasksComponent;
 
 public:
 
@@ -147,11 +147,11 @@ public:
 
 	/** Event called when PossessedPawn is possessed by this controller. */
 	UFUNCTION(BlueprintImplementableEvent)
-	void OnPossess(APawn* PossessedPawn);
+		void OnPossess(APawn* PossessedPawn);
 
 	/** Gets triggered after given pawn has been unpossesed */
 	UFUNCTION(BlueprintImplementableEvent)
-	void OnUnpossess(APawn* UnpossessedPawn);
+		void OnUnpossess(APawn* UnpossessedPawn);
 
 	virtual void SetPawn(APawn* InPawn) override;
 
@@ -165,9 +165,9 @@ public:
 	 *	@note AcceptanceRadius has default value or -1 due to Header Parser not being able to recognize UPathFollowingComponent::DefaultAcceptanceRadius
 	 */
 	UFUNCTION(BlueprintCallable, Category = "AI|Navigation", Meta = (AdvancedDisplay = "bStopOnOverlap,bCanStrafe,bAllowPartialPath"))
-	EPathFollowingRequestResult::Type MoveToActor(AActor* Goal, float AcceptanceRadius = -1, bool bStopOnOverlap = true,
-		bool bUsePathfinding = true, bool bCanStrafe = true,
-		TSubclassOf<UNavigationQueryFilter> FilterClass = NULL, bool bAllowPartialPath = true);
+		EPathFollowingRequestResult::Type MoveToActor(AActor* Goal, float AcceptanceRadius = -1, bool bStopOnOverlap = true,
+			bool bUsePathfinding = true, bool bCanStrafe = true,
+			TSubclassOf<UNavigationQueryFilter> FilterClass = NULL, bool bAllowPartialPath = true);
 
 	/** Makes AI go toward specified Dest location
 	 *  @param AcceptanceRadius - finish move if pawn gets close enough
@@ -180,26 +180,26 @@ public:
 	 *	@note AcceptanceRadius has default value or -1 due to Header Parser not being able to recognize UPathFollowingComponent::DefaultAcceptanceRadius
 	 */
 	UFUNCTION(BlueprintCallable, Category = "AI|Navigation", Meta = (AdvancedDisplay = "bStopOnOverlap,bCanStrafe,bAllowPartialPath"))
-	EPathFollowingRequestResult::Type MoveToLocation(const FVector& Dest, float AcceptanceRadius = -1, bool bStopOnOverlap = true,
-		bool bUsePathfinding = true, bool bProjectDestinationToNavigation = false, bool bCanStrafe = true,
-		TSubclassOf<UNavigationQueryFilter> FilterClass = NULL, bool bAllowPartialPath = true);
+		EPathFollowingRequestResult::Type MoveToLocation(const FVector& Dest, float AcceptanceRadius = -1, bool bStopOnOverlap = true,
+			bool bUsePathfinding = true, bool bProjectDestinationToNavigation = false, bool bCanStrafe = true,
+			TSubclassOf<UNavigationQueryFilter> FilterClass = NULL, bool bAllowPartialPath = true);
 
 	/** Makes AI go toward specified destination */
 	EPathFollowingRequestResult::Type MoveTo(const FAIMoveRequest& MoveRequest);
 
 	/** Fills pathfinding query for navigation system from given move request */
 	virtual bool PreparePathfinding(const FAIMoveRequest& MoveRequest, FPathFindingQuery& Query);
-	
+
 	DEPRECATED(4.8, "This function is deprecated, please use version with MoveRequest parameter instead.")
-	virtual bool PreparePathfinding(FPathFindingQuery& Query, const FVector& Dest, AActor* Goal, bool bUsePathfinding = true, TSubclassOf<class UNavigationQueryFilter> FilterClass = NULL);
+		virtual bool PreparePathfinding(FPathFindingQuery& Query, const FVector& Dest, AActor* Goal, bool bUsePathfinding = true, TSubclassOf<class UNavigationQueryFilter> FilterClass = NULL);
 
 	/** Executes pathfinding query and starts move request
 	 *  @return RequestID, or 0 when failed
 	 */
 	virtual FAIRequestID RequestPathAndMove(const FAIMoveRequest& MoveRequest, FPathFindingQuery& Query);
-	
+
 	DEPRECATED(4.8, "This function is deprecated, please use version with MoveRequest parameter instead.")
-	virtual FAIRequestID RequestPathAndMove(FPathFindingQuery& Query, AActor* Goal, float AcceptanceRadius, bool bStopOnOverlap, FCustomMoveSharedPtr CustomData);
+		virtual FAIRequestID RequestPathAndMove(FPathFindingQuery& Query, AActor* Goal, float AcceptanceRadius, bool bStopOnOverlap, FCustomMoveSharedPtr CustomData);
 
 	/** Handle move requests
 	 *  @return RequestID, or 0 when failed
@@ -207,7 +207,7 @@ public:
 	virtual FAIRequestID RequestMove(const FAIMoveRequest& MoveRequest, FNavPathSharedPtr Path);
 
 	DEPRECATED(4.8, "This function is deprecated, please use version with MoveRequest parameter instead.")
-	virtual FAIRequestID RequestMove(FNavPathSharedPtr Path, AActor* Goal = NULL, float AcceptanceRadius = UPathFollowingComponent::DefaultAcceptanceRadius, bool bStopOnOverlap = true, FCustomMoveSharedPtr CustomData = NULL);
+		virtual FAIRequestID RequestMove(FNavPathSharedPtr Path, AActor* Goal = NULL, float AcceptanceRadius = UPathFollowingComponent::DefaultAcceptanceRadius, bool bStopOnOverlap = true, FCustomMoveSharedPtr CustomData = NULL);
 
 	/** if AI is currently moving due to request given by RequestToPause, then the move will be paused */
 	bool PauseMove(FAIRequestID RequestToPause);
@@ -225,85 +225,85 @@ public:
 	FORCEINLINE FAIRequestID GetCurrentMoveRequestID() const { return GetPathFollowingComponent() ? GetPathFollowingComponent()->GetCurrentRequestId() : FAIRequestID::InvalidRequest; }
 
 	/** Blueprint notification that we've completed the current movement request */
-	UPROPERTY(BlueprintAssignable, meta=(DisplayName="MoveCompleted"))
-	FAIMoveCompletedSignature ReceiveMoveCompleted;
+	UPROPERTY(BlueprintAssignable, meta = (DisplayName = "MoveCompleted"))
+		FAIMoveCompletedSignature ReceiveMoveCompleted;
 
 	/** Returns status of path following */
-	UFUNCTION(BlueprintCallable, Category="AI|Navigation")
-	EPathFollowingStatus::Type GetMoveStatus() const;
+	UFUNCTION(BlueprintCallable, Category = "AI|Navigation")
+		EPathFollowingStatus::Type GetMoveStatus() const;
 
 	/** Returns true if the current PathFollowingComponent's path is partial (does not reach desired destination). */
-	UFUNCTION(BlueprintCallable, Category="AI|Navigation")
-	bool HasPartialPath() const;
-	
+	UFUNCTION(BlueprintCallable, Category = "AI|Navigation")
+		bool HasPartialPath() const;
+
 	/** Returns position of current path segment's end. */
-	UFUNCTION(BlueprintCallable, Category="AI|Navigation")
-	FVector GetImmediateMoveDestination() const;
+	UFUNCTION(BlueprintCallable, Category = "AI|Navigation")
+		FVector GetImmediateMoveDestination() const;
 
 	/** Updates state of movement block detection. */
-	UFUNCTION(BlueprintCallable, Category="AI|Navigation")
-	void SetMoveBlockDetection(bool bEnable);
+	UFUNCTION(BlueprintCallable, Category = "AI|Navigation")
+		void SetMoveBlockDetection(bool bEnable);
 
 	/** Prepares path finding and path following components. */
 	virtual void InitNavigationControl(UPathFollowingComponent*& PathFollowingComp) override;
 
 	/** Starts executing behavior tree. */
-	UFUNCTION(BlueprintCallable, Category="AI")
-	virtual bool RunBehaviorTree(UBehaviorTree* BTAsset);
+	UFUNCTION(BlueprintCallable, Category = "AI")
+		virtual bool RunBehaviorTree(UBehaviorTree* BTAsset);
 
-	/** 
+	/**
 	 * Makes AI use the specified Blackboard asset & creates a Blackboard Component if one does not already exist.
 	 * @param	BlackboardAsset			The Blackboard asset to use.
 	 * @param	BlackboardComponent		The Blackboard component that was used or created to work with the passed-in Blackboard Asset.
 	 * @return true if we successfully linked the blackboard asset to the blackboard component.
 	 */
 	UFUNCTION(BlueprintCallable, Category = "AI")
-	bool UseBlackboard(UBlackboardData* BlackboardAsset, UBlackboardComponent*& BlackboardComponent);
+		bool UseBlackboard(UBlackboardData* BlackboardAsset, UBlackboardComponent*& BlackboardComponent);
 
 	/** does this AIController allow given UBlackboardComponent sync data with it */
 	virtual bool ShouldSyncBlackboardWith(const UBlackboardComponent& OtherBlackboardComponent) const { return true; }
-		
+
 
 protected:
 	UFUNCTION(BlueprintImplementableEvent)
-	void OnUsingBlackBoard(UBlackboardComponent* BlackboardComp, UBlackboardData* BlackboardAsset);
+		void OnUsingBlackBoard(UBlackboardComponent* BlackboardComp, UBlackboardData* BlackboardAsset);
 
 	virtual bool InitializeBlackboard(UBlackboardComponent& BlackboardComp, UBlackboardData& BlackboardAsset);
 
 public:
 	/** Retrieve the final position that controller should be looking at. */
-	UFUNCTION(BlueprintCallable, Category="AI")
-	FVector GetFocalPoint() const;
+	UFUNCTION(BlueprintCallable, Category = "AI")
+		FVector GetFocalPoint() const;
 
 	FVector GetFocalPointForPriority(EAIFocusPriority::Type InPriority) const;
 
 	/** Retrieve the focal point this controller should focus to on given actor. */
-	UFUNCTION(BlueprintCallable, Category="AI")
-	virtual FVector GetFocalPointOnActor(const AActor *Actor) const;
+	UFUNCTION(BlueprintCallable, Category = "AI")
+		virtual FVector GetFocalPointOnActor(const AActor *Actor) const;
 
 	/** Set the position that controller should be looking at. */
-	UFUNCTION(BlueprintCallable, Category="AI", meta=(DisplayName="SetFocalPoint", Keywords="focus"))
-	void K2_SetFocalPoint(FVector FP);
+	UFUNCTION(BlueprintCallable, Category = "AI", meta = (DisplayName = "SetFocalPoint", Keywords = "focus"))
+		void K2_SetFocalPoint(FVector FP);
 
 	/** Set Focus for actor, will set FocalPoint as a result. */
-	UFUNCTION(BlueprintCallable, Category="AI", meta=(DisplayName="SetFocus"))
-	void K2_SetFocus(AActor* NewFocus);
+	UFUNCTION(BlueprintCallable, Category = "AI", meta = (DisplayName = "SetFocus"))
+		void K2_SetFocus(AActor* NewFocus);
 
 	/** Get the focused actor. */
-	UFUNCTION(BlueprintCallable, Category="AI")
-	AActor* GetFocusActor() const;
+	UFUNCTION(BlueprintCallable, Category = "AI")
+		AActor* GetFocusActor() const;
 
-	FORCEINLINE AActor* GetFocusActorForPriority(EAIFocusPriority::Type InPriority) const {  return FocusInformation.Priorities.IsValidIndex(InPriority) ? FocusInformation.Priorities[InPriority].Actor.Get() : nullptr; }
+	FORCEINLINE AActor* GetFocusActorForPriority(EAIFocusPriority::Type InPriority) const { return FocusInformation.Priorities.IsValidIndex(InPriority) ? FocusInformation.Priorities[InPriority].Actor.Get() : nullptr; }
 
 	/** Clears Focus, will also clear FocalPoint as a result */
-	UFUNCTION(BlueprintCallable, Category="AI", meta=(DisplayName="ClearFocus"))
-	void K2_ClearFocus();
+	UFUNCTION(BlueprintCallable, Category = "AI", meta = (DisplayName = "ClearFocus"))
+		void K2_ClearFocus();
 
 
-	/** 
+	/**
 	 * Computes a launch velocity vector to toss a projectile and hit the given destination.
 	 * Performance note: Potentially expensive. Nonzero CollisionRadius and bOnlyTraceUp=false are the more expensive options.
-	 * 
+	 *
 	 * @param OutTossVelocity - out param stuffed with the computed velocity to use
 	 * @param Start - desired start point of arc
 	 * @param End - desired end point of arc
@@ -312,7 +312,7 @@ public:
 	 * @param bOnlyTraceUp  (optional) - when true collision checks verifying the arc will only be done along the upward portion of the arc
 	 * @return - true if a valid arc was computed, false if no valid solution could be found
 	 */
-	bool SuggestTossVelocity(FVector& OutTossVelocity, FVector Start, FVector End, float TossSpeed, bool bPreferHighArc, float CollisionRadius=0, bool bOnlyTraceUp=false);
+	bool SuggestTossVelocity(FVector& OutTossVelocity, FVector Start, FVector End, float TossSpeed, bool bPreferHighArc, float CollisionRadius = 0, bool bOnlyTraceUp = false);
 
 	//~ Begin AActor Interface
 	virtual void Tick(float DeltaTime) override;
@@ -348,7 +348,7 @@ public:
 	virtual void UpdateControlRotation(float DeltaTime, bool bUpdatePawn = true);
 
 	/** Set FocalPoint for given priority as absolute position or offset from base. */
-	virtual void SetFocalPoint(FVector NewFocus, EAIFocusPriority::Type InPriority=EAIFocusPriority::Gameplay);
+	virtual void SetFocalPoint(FVector NewFocus, EAIFocusPriority::Type InPriority = EAIFocusPriority::Gameplay);
 
 	/* Set Focus actor for given priority, will set FocalPoint as a result. */
 	virtual void SetFocus(AActor* NewFocus, EAIFocusPriority::Type InPriority = EAIFocusPriority::Gameplay);
@@ -380,16 +380,13 @@ public:
 	//----------------------------------------------------------------------//
 	// IGameplayTaskOwnerInterface
 	//----------------------------------------------------------------------//
-	UGameplayTasksComponent* GetGameplayTasksComponent() const { return CachedGameplayTasksComponent; }
-	virtual UGameplayTasksComponent* GetGameplayTasksComponent(const UGameplayTask& Task) const override { return CachedGameplayTasksComponent; }
-	virtual void OnTaskActivated(UGameplayTask& Task) override {}
-	virtual void OnTaskDeactivated(UGameplayTask& Task) override {}
-	virtual void OnTaskInitialized(UGameplayTask& Task) override {}
-	virtual AActor* GetOwnerActor(const UGameplayTask* Task) const override { return const_cast<AAIController*>(this); }
-	virtual AActor* GetAvatarActor(const UGameplayTask* Task) const override { return GetPawn(); }
-	virtual uint8 GetDefaultPriority() const override { return FGameplayTasks::DefaultPriority - 1; }
-
-	// other GT tasks related
+	virtual UGameplayTasksComponent* GetGameplayTasksComponent(const UGameplayTask& Task) const override { return GetGameplayTasksComponent(); }
+	virtual AActor* GetGameplayTaskOwner(const UGameplayTask* Task) const override { return const_cast<AAIController*>(this); }
+	virtual AActor* GetGameplayTaskAvatar(const UGameplayTask* Task) const override { return GetPawn(); }
+	virtual uint8 GetGameplayTaskDefaultPriority() const { return FGameplayTasks::DefaultPriority - 1; }
+
+	FORCEINLINE UGameplayTasksComponent* GetGameplayTasksComponent() const { return CachedGameplayTasksComponent; }
+
 	UFUNCTION()
 	virtual void OnGameplayTaskResourcesClaimed(FGameplayResourceSet NewlyClaimed, FGameplayResourceSet FreshlyReleased);
 
diff --git a/Engine/Source/Runtime/AIModule/Classes/AITypes.h b/Engine/Source/Runtime/AIModule/Classes/AITypes.h
--- a/Engine/Source/Runtime/AIModule/Classes/AITypes.h
+++ b/Engine/Source/Runtime/AIModule/Classes/AITypes.h
@@ -474,7 +474,7 @@ struct AIMODULE_API FAIMoveRequest
 	FAIMoveRequest& SetUserData(FCustomMoveSharedPtr Data) { UserData = Data; return *this; }
 
 	/** the request should be either set up to move to a location, of go to a valid actor */
-	bool IsValid() const { return bMoveToActor == false || GoalActor != nullptr; }
+	bool IsValid() const { return bInitialized && (!bMoveToActor || GoalActor); }
 
 	bool IsMoveToActorRequest() const { return bMoveToActor; }
 	AActor* GetGoalActor() const { return bMoveToActor ? GoalActor : nullptr; }
diff --git a/Engine/Source/Runtime/AIModule/Classes/BehaviorTree/BTAuxiliaryNode.h b/Engine/Source/Runtime/AIModule/Classes/BehaviorTree/BTAuxiliaryNode.h
--- a/Engine/Source/Runtime/AIModule/Classes/BehaviorTree/BTAuxiliaryNode.h
+++ b/Engine/Source/Runtime/AIModule/Classes/BehaviorTree/BTAuxiliaryNode.h
@@ -44,6 +44,15 @@ class AIMODULE_API UBTAuxiliaryNode : public UBTNode
 	virtual void DescribeRuntimeValues(const UBehaviorTreeComponent& OwnerComp, uint8* NodeMemory, EBTDescriptionVerbosity::Type Verbosity, TArray<FString>& Values) const override;
 	virtual uint16 GetSpecialMemorySize() const override;
 
+	/** fill in data about tree structure */
+	void InitializeParentLink(uint8 InChildIndex);
+
+	/** @return parent task node */
+	const UBTNode* GetMyNode() const;
+
+	/** @return index of child in parent's array or MAX_uint8 */
+	uint8 GetChildIndex() const;
+
 protected:
 
 	/** if set, OnBecomeRelevant will be used */
@@ -58,6 +67,9 @@ protected:
 	/** if set, conditional tick will use remaining time form node's memory */
 	uint8 bTickIntervals : 1;
 
+	/** child index in parent node */
+	uint8 ChildIndex;
+
 	/** called when auxiliary node becomes active
 	 * this function should be considered as const (don't modify state of object) if node is not instanced! */
 	virtual void OnBecomeRelevant(UBehaviorTreeComponent& OwnerComp, uint8* NodeMemory);
@@ -93,3 +105,8 @@ public:
 	DEPRECATED(4.7, "This version is deprecated. Please use the one taking reference to UBehaviorTreeComponent rather than a pointer.")
 	virtual void TickNode(UBehaviorTreeComponent* OwnerComp, uint8* NodeMemory, float DeltaSeconds);
 };
+
+FORCEINLINE uint8 UBTAuxiliaryNode::GetChildIndex() const
+{
+	return ChildIndex;
+}
diff --git a/Engine/Source/Runtime/AIModule/Classes/BehaviorTree/BTDecorator.h b/Engine/Source/Runtime/AIModule/Classes/BehaviorTree/BTDecorator.h
--- a/Engine/Source/Runtime/AIModule/Classes/BehaviorTree/BTDecorator.h
+++ b/Engine/Source/Runtime/AIModule/Classes/BehaviorTree/BTDecorator.h
@@ -32,9 +32,6 @@ class AIMODULE_API UBTDecorator : public UBTAuxiliaryNode
 {
 	GENERATED_UCLASS_BODY()
 
-	/** fill in data about tree structure */
-	void InitializeDecorator(uint8 InChildIndex);
-
 	/** wrapper for node instancing: CalculateRawConditionValue */
 	bool WrappedCanExecute(UBehaviorTreeComponent& OwnerComp, uint8* NodeMemory) const;
 
@@ -47,12 +44,6 @@ class AIMODULE_API UBTDecorator : public UBTAuxiliaryNode
 	/** wrapper for node instancing: OnNodeProcessed */
 	void WrappedOnNodeProcessed(FBehaviorTreeSearchData& SearchData, EBTNodeResult::Type& NodeResult) const;
 
-	/** @return decorated child */
-	const UBTNode* GetMyNode() const;
-
-	/** @return index of child in parent's array */
-	uint8 GetChildIndex() const;
-
 	/** @return flow controller's abort mode */
 	EBTFlowAbortMode::Type GetFlowAbortMode() const;
 
@@ -100,9 +91,6 @@ protected:
 	UPROPERTY(Category=FlowControl, EditAnywhere)
 	TEnumAsByte<EBTFlowAbortMode::Type> FlowAbortMode;
 
-	/** child index in parent node */
-	uint8 ChildIndex;
-
 	void SetIsInversed(bool bShouldBeInversed);
 
 	/** called when underlying node is activated
@@ -127,6 +115,8 @@ protected:
 	//----------------------------------------------------------------------//
 	DEPRECATED(4.7, "This version is deprecated. Please use the one taking reference to UBehaviorTreeComponent rather than a pointer.")
 	virtual bool CalculateRawConditionValue(UBehaviorTreeComponent* OwnerComp, uint8* NodeMemory) const;
+	DEPRECATED(4.12, "This function is deprecated, please use InitializeParentLink instead.")
+	void InitializeDecorator(uint8 InChildIndex);
 };
 
 
@@ -142,8 +132,3 @@ FORCEINLINE bool UBTDecorator::IsInversed() const
 {
 	return bInverseCondition;
 }
-
-FORCEINLINE uint8 UBTDecorator::GetChildIndex() const
-{
-	return ChildIndex;
-}
diff --git a/Engine/Source/Runtime/AIModule/Classes/BehaviorTree/BTNode.h b/Engine/Source/Runtime/AIModule/Classes/BehaviorTree/BTNode.h
--- a/Engine/Source/Runtime/AIModule/Classes/BehaviorTree/BTNode.h
+++ b/Engine/Source/Runtime/AIModule/Classes/BehaviorTree/BTNode.h
@@ -5,6 +5,7 @@
 #include "BehaviorTreeComponent.h"
 #include "BehaviorTreeTypes.h"
 #include "GameplayTaskOwnerInterface.h"
+#include "Tasks/AITask.h"
 #include "BTNode.generated.h"
 
 AIMODULE_API DECLARE_LOG_CATEGORY_EXTERN(LogBehaviorTree, Display, All);
@@ -143,29 +144,22 @@ class AIMODULE_API UBTNode : public UObject, public IGameplayTaskOwnerInterface
 	/** Gets called only for instanced nodes(bCreateNodeInstance == true). In practive overridden by BP-implemented BT nodes */
 	virtual void SetOwner(AActor* ActorOwner) {}
 
-	//----------------------------------------------------------------------//
-	// IGameplayTaskOwnerInterface & other gameplay task stuff
-	//----------------------------------------------------------------------//
+	// BEGIN IGameplayTaskOwnerInterface
 	virtual UGameplayTasksComponent* GetGameplayTasksComponent(const UGameplayTask& Task) const override;
-	virtual void OnTaskActivated(UGameplayTask& Task) override;
-	virtual void OnTaskDeactivated(UGameplayTask& Task) override;
-	virtual void OnTaskInitialized(UGameplayTask& Task) override;
-	virtual AActor* GetOwnerActor(const UGameplayTask* Task) const;
-	virtual AActor* GetAvatarActor(const UGameplayTask* Task) const;
-	virtual uint8 GetDefaultPriority() const;
+	virtual AActor* GetGameplayTaskOwner(const UGameplayTask* Task) const override;
+	virtual AActor* GetGameplayTaskAvatar(const UGameplayTask* Task) const override;
+	virtual uint8 GetGameplayTaskDefaultPriority() const override;
+	virtual void OnGameplayTaskInitialized(UGameplayTask& Task) override;
+	// END IGameplayTaskOwnerInterface
 
 	UBehaviorTreeComponent* GetBTComponentForTask(UGameplayTask& Task) const;
 	
 	template <class T>
 	T* NewBTAITask(UBehaviorTreeComponent& BTComponent)
 	{
-		T* NewAITask = NewObject<T>();
-		AAIController* AIController = BTComponent.GetAIOwner();
-		check(AIController && "Can\'t spawn an AI task without AI controller!");
-		NewAITask->InitAITask(*AIController, *this, GetDefaultPriority());
+		check(BTComponent.GetAIOwner());
 		bOwnsGameplayTasks = true;
-
-		return NewAITask;
+		return UAITask::NewAITask<T>(*BTComponent.GetAIOwner(), *this, TEXT("Behavior"));
 	}
 
 	/** node name */
diff --git a/Engine/Source/Runtime/AIModule/Classes/BehaviorTree/BTTaskNode.h b/Engine/Source/Runtime/AIModule/Classes/BehaviorTree/BTTaskNode.h
--- a/Engine/Source/Runtime/AIModule/Classes/BehaviorTree/BTTaskNode.h
+++ b/Engine/Source/Runtime/AIModule/Classes/BehaviorTree/BTTaskNode.h
@@ -6,6 +6,7 @@
 #include "BTTaskNode.generated.h"
 
 class UBehaviorTreeComponent;
+class UBTService;
 struct FAIMessage;
 
 /** 
@@ -43,6 +44,7 @@ public:
 #if WITH_EDITOR
 	virtual FName GetNodeIconName() const override;
 #endif // WITH_EDITOR
+	virtual void OnGameplayTaskDeactivated(UGameplayTask& Task) override;
 
 	/** message observer's hook */
 	void ReceivedMessage(UBrainComponent* BrainComp, const FAIMessage& Message);
@@ -65,14 +67,13 @@ public:
 	/** helper function: finishes latent aborting */
 	void FinishLatentAbort(UBehaviorTreeComponent& OwnerComp) const;
 
-	//----------------------------------------------------------------------//
-	// UBTTaskNode IGameplayTaskOwnerInterface
-	//----------------------------------------------------------------------//
-	virtual void OnTaskDeactivated(UGameplayTask& Task) override;
-
 	/** @return true if task search should be discarded when this task is selected to execute but is already running */
 	bool ShouldIgnoreRestartSelf() const;
 
+	/** service nodes */
+	UPROPERTY()
+	TArray<UBTService*> Services;
+
 protected:
 
 	/** if set, task search will be discarded when this task is selected to execute but is already running */
diff --git a/Engine/Source/Runtime/AIModule/Classes/BehaviorTree/BehaviorTreeComponent.h b/Engine/Source/Runtime/AIModule/Classes/BehaviorTree/BehaviorTreeComponent.h
--- a/Engine/Source/Runtime/AIModule/Classes/BehaviorTree/BehaviorTreeComponent.h
+++ b/Engine/Source/Runtime/AIModule/Classes/BehaviorTree/BehaviorTreeComponent.h
@@ -60,11 +60,17 @@ struct FBTPendingExecutionInfo
 	/** if set, request can't be executed */
 	uint32 bLocked : 1;
 
-	FBTPendingExecutionInfo() : NextTask(NULL), bOutOfNodes(false), bLocked(false) {}
+	/** if set, active task is in the middle of Abort call */
+	uint32 bAborting : 1;
+
+	FBTPendingExecutionInfo() : NextTask(NULL), bOutOfNodes(false), bLocked(false), bAborting(false) {}
 	bool IsSet() const { return (NextTask || bOutOfNodes) && !bLocked; }
 
 	void Lock() { bLocked = true; }
 	void Unlock() { bLocked = false; }
+
+	void OnAbortStart() { bAborting = true; }
+	void OnAbortProcessed() { bAborting = false; }
 };
 
 struct FBTTreeStartInfo
diff --git a/Engine/Source/Runtime/AIModule/Classes/BehaviorTree/BehaviorTreeTypes.h b/Engine/Source/Runtime/AIModule/Classes/BehaviorTree/BehaviorTreeTypes.h
--- a/Engine/Source/Runtime/AIModule/Classes/BehaviorTree/BehaviorTreeTypes.h
+++ b/Engine/Source/Runtime/AIModule/Classes/BehaviorTree/BehaviorTreeTypes.h
@@ -58,6 +58,8 @@ namespace BTSpecialChild
 {
 	const int32 NotInitialized = -1;	// special value for child indices: needs to be initialized
 	const int32 ReturnToParent = -2;	// special value for child indices: return to parent node
+	
+	const uint8 OwnedByComposite = MAX_uint8;	// special value for aux node's child index: owned by composite node instead of a task
 }
 
 UENUM(BlueprintType)
diff --git a/Engine/Source/Runtime/AIModule/Classes/BehaviorTree/Tasks/BTTask_MoveDirectlyToward.h b/Engine/Source/Runtime/AIModule/Classes/BehaviorTree/Tasks/BTTask_MoveDirectlyToward.h
--- a/Engine/Source/Runtime/AIModule/Classes/BehaviorTree/Tasks/BTTask_MoveDirectlyToward.h
+++ b/Engine/Source/Runtime/AIModule/Classes/BehaviorTree/Tasks/BTTask_MoveDirectlyToward.h
@@ -1,51 +1,34 @@
 // Copyright 1998-2016 Epic Games, Inc. All Rights Reserved.
 
 #pragma once
-#include "BehaviorTree/Tasks/BTTask_BlackboardBase.h"
+#include "BehaviorTree/Tasks/BTTask_MoveTo.h"
 #include "BTTask_MoveDirectlyToward.generated.h"
 
-struct FBTMoveDirectlyTowardMemory
-{
-	/** Move request ID */
-	FAIRequestID MoveRequestID;
-};
-
 /**
  * Move Directly Toward task node.
  * Moves the AI pawn toward the specified Actor or Location (Vector) blackboard entry in a straight line, without regard to any navigation system. If you need the AI to navigate, use the "Move To" node instead.
  */
 UCLASS(config=Game)
-class AIMODULE_API UBTTask_MoveDirectlyToward : public UBTTask_BlackboardBase
+class AIMODULE_API UBTTask_MoveDirectlyToward : public UBTTask_MoveTo
 {
 	GENERATED_UCLASS_BODY()
 
-	UPROPERTY(config, Category=Node, EditAnywhere, meta=(ClampMin = "0.0"))
-	float AcceptableRadius;
+	virtual void PostLoad() override;
+	virtual void InitializeFromAsset(UBehaviorTree& Asset) override;
+#if WITH_EDITOR
+	virtual FName GetNodeIconName() const override;
+#endif // WITH_EDITOR
 
-	/** set to true will result in not updating move destination in 
-	 *	case where goal is an Actor that can change 
-	 *	his location while task is being performed */
-	UPROPERTY(Category=Node, EditAnywhere)
+	DEPRECATED_FORGAME(4.12, "This property is now deprecated, use UBTTask_MoveTo::bTrackMovingGoal instead.")
+	UPROPERTY()
 	uint32 bDisablePathUpdateOnGoalLocationChange : 1;
 
-	UPROPERTY(Category=Node, EditAnywhere)
-	uint32 bProjectVectorGoalToNavigation:1;
+	DEPRECATED_FORGAME(4.12, "This property is now deprecated, use UBTTask_MoveTo::bProjectGoalLocation instead.")
+	UPROPERTY()
+	uint32 bProjectVectorGoalToNavigation : 1;
 
-	UPROPERTY(Category=Node, EditAnywhere)
-	uint32 bAllowStrafe : 1;
+private:
 
-	/** if set to true agent's radius will be added to AcceptableRadius for purposes of checking
-	 *	if path's end point has been reached. Will result in AI stopping on contact with destination location*/
-	UPROPERTY(Category = Node, EditAnywhere)
-	uint32 bStopOnOverlap : 1;
-
-	virtual EBTNodeResult::Type ExecuteTask(UBehaviorTreeComponent& OwnerComp, uint8* NodeMemory) override;
-	virtual EBTNodeResult::Type AbortTask(UBehaviorTreeComponent& OwnerComp, uint8* NodeMemory) override;
-	virtual uint16 GetInstanceMemorySize() const override;
-	virtual void DescribeRuntimeValues(const UBehaviorTreeComponent& OwnerComp, uint8* NodeMemory, EBTDescriptionVerbosity::Type Verbosity, TArray<FString>& Values) const override;
-	virtual FString GetStaticDescription() const override;
-
-#if WITH_EDITOR
-	virtual FName GetNodeIconName() const override;
-#endif // WITH_EDITOR
+	UPROPERTY()
+	uint32 bUpdatedDeprecatedProperties : 1;
 };
diff --git a/Engine/Source/Runtime/AIModule/Classes/BehaviorTree/Tasks/BTTask_MoveTo.h b/Engine/Source/Runtime/AIModule/Classes/BehaviorTree/Tasks/BTTask_MoveTo.h
--- a/Engine/Source/Runtime/AIModule/Classes/BehaviorTree/Tasks/BTTask_MoveTo.h
+++ b/Engine/Source/Runtime/AIModule/Classes/BehaviorTree/Tasks/BTTask_MoveTo.h
@@ -2,6 +2,7 @@
 
 #pragma once
 #include "BehaviorTree/Tasks/BTTask_BlackboardBase.h"
+#include "AITypes.h"
 #include "BTTask_MoveTo.generated.h"
 
 class UAITask_MoveTo;
@@ -18,6 +19,7 @@ struct FBTMoveToTaskMemory
 	TWeakObjectPtr<UAITask_MoveTo> Task;
 
 	uint8 bWaitingForPath : 1;
+	uint8 bObserverCanFinishTask : 1;
 };
 
 /**
@@ -29,11 +31,11 @@ class AIMODULE_API UBTTask_MoveTo : public UBTTask_BlackboardBase
 {
 	GENERATED_UCLASS_BODY()
 
-	UPROPERTY(config, Category=Node, EditAnywhere, meta=(ClampMin = "0.0", UIMin="0.0"))
+	UPROPERTY(config, Category = Node, EditAnywhere, meta=(ClampMin = "0.0", UIMin="0.0"))
 	float AcceptableRadius;
 
 	/** "None" will result in default filter being used */
-	UPROPERTY(Category=Node, EditAnywhere)
+	UPROPERTY(Category = Node, EditAnywhere)
 	TSubclassOf<UNavigationQueryFilter> FilterClass;
 
 	/** if task is expected to react to changes to location represented by BB key 
@@ -42,27 +44,43 @@ class AIMODULE_API UBTTask_MoveTo : public UBTTask_BlackboardBase
 	UPROPERTY(Category=Blackboard, EditAnywhere, AdvancedDisplay, meta = (ClampMin = "1", UIMin = "1", EditCondition="bObserveBlackboardValue"))
 	float ObservedBlackboardValueTolerance;
 
-	UPROPERTY(Category=Node, EditAnywhere)
+	/** if move goal in BB changes the move will be redirected to new location */
+	UPROPERTY()
+	uint32 bObserveBlackboardValue : 1;
+
+	UPROPERTY(Category = Node, EditAnywhere)
 	uint32 bAllowStrafe : 1;
+
 	/** if set, use incomplete path when goal can't be reached */
 	UPROPERTY(Category = Node, EditAnywhere, AdvancedDisplay)
 	uint32 bAllowPartialPath : 1;
 
+	/** if set, path to goal actor will update itself when actor moves */
+	UPROPERTY(Category = Node, EditAnywhere, AdvancedDisplay)
+	uint32 bTrackMovingGoal : 1;
+
+	/** if set, goal location will be projected on navigation data (navmesh) before using */
+	UPROPERTY(Category = Node, EditAnywhere, AdvancedDisplay)
+	uint32 bProjectGoalLocation : 1;
+
 	/** if set to true agent's radius will be added to AcceptableRadius for purposes of checking 
 	 *	if path's end point has been reached. Will result in AI stopping on contact with destination location*/
-	UPROPERTY(Category=Node, EditAnywhere)
+	UPROPERTY(Category = Node, EditAnywhere)
 	uint32 bStopOnOverlap : 1;
 
-	/** if move goal in BB changes the move will be redirected to new location */
-	UPROPERTY(Category=Blackboard, EditAnywhere)
-	uint32 bObserveBlackboardValue : 1;
-	
+	/** if set, move will use pathfinding. Not exposed on purpose, please use BTTask_MoveDirectlyToward */
+	uint32 bUsePathfinding : 1;
+
+	/** set automatically if move should use GameplayTasks */
+	uint32 bUseGameplayTasks : 1;
+
 	virtual EBTNodeResult::Type ExecuteTask(UBehaviorTreeComponent& OwnerComp, uint8* NodeMemory) override;
 	virtual EBTNodeResult::Type AbortTask(UBehaviorTreeComponent& OwnerComp, uint8* NodeMemory) override;
-	virtual  void OnTaskFinished(UBehaviorTreeComponent& OwnerComp, uint8* NodeMemory, EBTNodeResult::Type TaskResult) override;
+	virtual void OnTaskFinished(UBehaviorTreeComponent& OwnerComp, uint8* NodeMemory, EBTNodeResult::Type TaskResult) override;
 	virtual void TickTask(UBehaviorTreeComponent& OwnerComp, uint8* NodeMemory, float DeltaSeconds) override;
 	virtual uint16 GetInstanceMemorySize() const override;
 
+	virtual void OnGameplayTaskDeactivated(UGameplayTask& Task) override;
 	virtual void OnMessage(UBehaviorTreeComponent& OwnerComp, uint8* NodeMemory, FName Message, int32 RequestID, bool bSuccess) override;
 	EBlackboardNotificationResult OnBlackboardValueChange(const UBlackboardComponent& Blackboard, FBlackboard::FKey ChangedKeyID);
 
@@ -76,4 +94,7 @@ class AIMODULE_API UBTTask_MoveTo : public UBTTask_BlackboardBase
 protected:
 
 	EBTNodeResult::Type PerformMoveTask(UBehaviorTreeComponent& OwnerComp, uint8* NodeMemory);
+	
+	/** prepares move task for activation */
+	virtual UAITask_MoveTo* PrepareMoveTask(UBehaviorTreeComponent& OwnerComp, UAITask_MoveTo* ExistingTask, const FAIMoveRequest& MoveRequest);
 };
diff --git a/Engine/Source/Runtime/AIModule/Classes/Navigation/PathFollowingComponent.h b/Engine/Source/Runtime/AIModule/Classes/Navigation/PathFollowingComponent.h
--- a/Engine/Source/Runtime/AIModule/Classes/Navigation/PathFollowingComponent.h
+++ b/Engine/Source/Runtime/AIModule/Classes/Navigation/PathFollowingComponent.h
@@ -148,10 +148,6 @@ class AIMODULE_API UPathFollowingComponent : public UActorComponent, public IAIR
 		return RequestMove(InPath, UnboundRequestDelegate, InDestinationActor, InAcceptanceRadius, InStopOnOverlap, InGameData);
 	}
 
-	/** update path for specified request
-	 *  @param RequestID - request to update */
-	virtual bool UpdateMove(FNavPathSharedPtr Path, FAIRequestID RequestID = FAIRequestID::CurrentRequest);
-
 	/** aborts following path
 	 *  @param RequestID - request to abort, 0 = current
 	 *  @param bResetVelocity - try to stop movement component
@@ -172,7 +168,7 @@ class AIMODULE_API UPathFollowingComponent : public UActorComponent, public IAIR
 	/** notify about finishing move along current path segment */
 	virtual void OnSegmentFinished();
 
-	/** notify about changing current path */
+	/** notify about changing current path: new pointer or update from path event */
 	virtual void OnPathUpdated();
 
 	/** set associated movement component */
@@ -297,12 +293,15 @@ class AIMODULE_API UPathFollowingComponent : public UActorComponent, public IAIR
 	virtual bool IsResourceLocked() const override;
 	// IAIResourceInterface end
 
-	void OnPathEvent(FNavigationPath* InvalidatedPath, ENavPathEvent::Type Event);
+	void OnPathEvent(FNavigationPath* InPath, ENavPathEvent::Type Event);
 
 	/** helper function for sending a path for visual log */
 	static void LogPathHelper(const AActor* LogOwner, FNavPathSharedPtr InLogPath, const AActor* LogGoalActor);
 	static void LogPathHelper(const AActor* LogOwner, FNavigationPath* InLogPath, const AActor* LogGoalActor);
 
+	DEPRECATED_FORGAME(4.12, "This function is now deprecated and replaced with HandlePathUpdateEvent. Receiving new path pointer for the same move request is no longer supported, please either update data within current path and call FNavigationPath::DoneUpdating or start new move request.")
+	virtual bool UpdateMove(FNavPathSharedPtr Path, FAIRequestID RequestID = FAIRequestID::CurrentRequest);
+
 protected:
 
 	/** associated movement component */
@@ -490,6 +489,10 @@ protected:
 	 *	@param bForce results in looking for owner's movement component even if pointer to one is already cached */
 	virtual bool UpdateMovementComponent(bool bForce = false);
 
+	/** called after receiving update event from current path
+	 *  @return false if path was not accepted and move request needs to be aborted */
+	virtual bool HandlePathUpdateEvent();
+
 	/** called from timer if component spends too much time in Waiting state */
 	virtual void OnWaitingPathTimeout();
 
diff --git a/Engine/Source/Runtime/AIModule/Classes/Perception/AIPerceptionComponent.h b/Engine/Source/Runtime/AIModule/Classes/Perception/AIPerceptionComponent.h
--- a/Engine/Source/Runtime/AIModule/Classes/Perception/AIPerceptionComponent.h
+++ b/Engine/Source/Runtime/AIModule/Classes/Perception/AIPerceptionComponent.h
@@ -188,7 +188,7 @@ public:
 	virtual void OnUnregister() override;
 
 	UFUNCTION()
-	void OnOwnerEndPlay(EEndPlayReason::Type EndPlayReason);
+	void OnOwnerEndPlay(AActor* Actor, EEndPlayReason::Type EndPlayReason);
 	
 	void GetLocationAndDirection(FVector& Location, FVector& Direction) const;
 	const AActor* GetBodyActor() const;
diff --git a/Engine/Source/Runtime/AIModule/Classes/Perception/AIPerceptionSystem.h b/Engine/Source/Runtime/AIModule/Classes/Perception/AIPerceptionSystem.h
--- a/Engine/Source/Runtime/AIModule/Classes/Perception/AIPerceptionSystem.h
+++ b/Engine/Source/Runtime/AIModule/Classes/Perception/AIPerceptionSystem.h
@@ -154,7 +154,7 @@ public:
 protected:
 	
 	UFUNCTION()
-	void OnPerceptionStimuliSourceEndPlay(EEndPlayReason::Type EndPlayReason);
+	void OnPerceptionStimuliSourceEndPlay(AActor* Actor, EEndPlayReason::Type EndPlayReason);
 	
 	/** requests registration of a given actor as a perception data source for specified sense */
 	void RegisterSource(FAISenseID SenseID, AActor& SourceActor);
diff --git a/Engine/Source/Runtime/AIModule/Classes/Tasks/AITask.h b/Engine/Source/Runtime/AIModule/Classes/Tasks/AITask.h
--- a/Engine/Source/Runtime/AIModule/Classes/Tasks/AITask.h
+++ b/Engine/Source/Runtime/AIModule/Classes/Tasks/AITask.h
@@ -35,7 +35,44 @@ public:
 	AAIController* GetAIController() const { return OwnerController; };
 
 	void InitAITask(AAIController& AIOwner, IGameplayTaskOwnerInterface& InTaskOwner, uint8 InPriority);
+	void InitAITask(AAIController& AIOwner, IGameplayTaskOwnerInterface& InTaskOwner);
 
 	/** effectively adds UAIResource_Logic to the set of Claimed resources */
 	void RequestAILogicLocking();
-};
\ No newline at end of file
+
+	template <class T>
+	static T* NewAITask(AAIController& AIOwner, IGameplayTaskOwnerInterface& InTaskOwner, FName InstanceName = FName())
+	{
+		T* TaskInstance = NewObject<T>();
+		TaskInstance->InstanceName = InstanceName;
+		TaskInstance->InitAITask(AIOwner, InTaskOwner);
+		return TaskInstance;
+	}
+
+	template <class T>
+	static T* NewAITask(AAIController& AIOwner, IGameplayTaskOwnerInterface& InTaskOwner, EAITaskPriority InPriority, FName InstanceName = FName())
+	{
+		T* TaskInstance = NewObject<T>();
+		TaskInstance->InstanceName = InstanceName;
+		TaskInstance->InitAITask(AIOwner, InTaskOwner, (uint8)InPriority);
+		return TaskInstance;
+	}
+
+	template <class T>
+	static T* NewAITask(AAIController& AIOwner, FName InstanceName = FName())
+	{
+		T* TaskInstance = NewObject<T>();
+		TaskInstance->InstanceName = InstanceName;
+		TaskInstance->InitAITask(AIOwner, AIOwner);
+		return TaskInstance;
+	}
+
+	template <class T>
+	static T* NewAITask(AAIController& AIOwner, EAITaskPriority InPriority, FName InstanceName = FName())
+	{
+		T* TaskInstance = NewObject<T>();
+		TaskInstance->InstanceName = InstanceName;
+		TaskInstance->InitAITask(AIOwner, AIOwner, (uint8)InPriority);
+		return TaskInstance;
+	}
+};
diff --git a/Engine/Source/Runtime/AIModule/Classes/Tasks/AITask_LockLogic.h b/Engine/Source/Runtime/AIModule/Classes/Tasks/AITask_LockLogic.h
new file mode 100644
--- /dev/null
+++ b/Engine/Source/Runtime/AIModule/Classes/Tasks/AITask_LockLogic.h
@@ -0,0 +1,12 @@
+// Copyright 1998-2016 Epic Games, Inc. All Rights Reserved.
+#pragma once
+
+#include "AITask.h"
+#include "AITask_LockLogic.generated.h"
+
+/** Locks AI logic until removed by external trigger */
+UCLASS()
+class AIMODULE_API UAITask_LockLogic : public UAITask
+{
+	GENERATED_UCLASS_BODY()
+};
diff --git a/Engine/Source/Runtime/AIModule/Classes/Tasks/AITask_MoveTo.h b/Engine/Source/Runtime/AIModule/Classes/Tasks/AITask_MoveTo.h
--- a/Engine/Source/Runtime/AIModule/Classes/Tasks/AITask_MoveTo.h
+++ b/Engine/Source/Runtime/AIModule/Classes/Tasks/AITask_MoveTo.h
@@ -13,6 +13,25 @@ UCLASS()
 class AIMODULE_API UAITask_MoveTo : public UAITask
 {
 	GENERATED_BODY()
+
+public:
+	UAITask_MoveTo(const FObjectInitializer& ObjectInitializer);
+
+	/** tries to start move request and handles retry timer */
+	void ConditionalPerformMove();
+
+	/** prepare move task for activation */
+	void SetUp(AAIController* Controller, const FAIMoveRequest& InMoveRequest);
+
+	EPathFollowingResult::Type GetMoveResult() const { return MoveResult; }
+	bool WasMoveSuccessful() const { return MoveResult == EPathFollowingResult::Success; }
+
+	UFUNCTION(BlueprintCallable, Category = "AI|Tasks", meta = (AdvancedDisplay = "AcceptanceRadius,StopOnOverlap,AcceptPartialPath,bUsePathfinding", DefaultToSelf = "Controller", BlueprintInternalUseOnly = "TRUE", DisplayName = "Move To Location or Actor"))
+	static UAITask_MoveTo* AIMoveTo(AAIController* Controller, FVector GoalLocation, AActor* GoalActor = nullptr, float AcceptanceRadius = -1.f, EAIOptionFlag::Type StopOnOverlap = EAIOptionFlag::Default, EAIOptionFlag::Type AcceptPartialPath = EAIOptionFlag::Default, bool bUsePathfinding = true, bool bLockAILogic = true);
+
+	DEPRECATED(4.12, "This function is now depreacted, please use version with FAIMoveRequest parameter")
+	void SetUp(AAIController* Controller, FVector GoalLocation, AActor* GoalActor = nullptr, float AcceptanceRadius = -1.f, bool bUsePathfinding = true, EAIOptionFlag::Type StopOnOverlap = EAIOptionFlag::Default, EAIOptionFlag::Type AcceptPartialPath = EAIOptionFlag::Default);
+
 protected:
 	UPROPERTY(BlueprintAssignable)
 	FGenericGameplayTaskDelegate OnRequestFailed;
@@ -20,49 +39,57 @@ protected:
 	UPROPERTY(BlueprintAssignable)
 	FMoveTaskCompletedSignature OnMoveFinished;
 
-	UPROPERTY(BlueprintReadOnly, meta = (ExposeOnSpawn = true, DisplayName="Goal Location"), Category = "AITask")
-	FVector MoveGoalLocation;
+	/** parameters of move request */
+	UPROPERTY()
+	FAIMoveRequest MoveRequest;
 
-	/** gets set to GoalLocation if GoalActor != nullptr, otherwise is FAISystem::InvalidLocation */
-	FVector RealGoalLocation;
-	
-	UPROPERTY(BlueprintReadOnly, meta = (ExposeOnSpawn = true, DisplayName="Goal Actor"), Category = "AITask")
-	AActor* MoveGoalActor;
+	/** handle of path following's OnMoveFinished delegate */
+	FDelegateHandle PathFinishDelegateHandle;
 
-	UPROPERTY(BlueprintReadOnly, meta = (ExposeOnSpawn = true, DisplayName = "Acceptance Radius"), Category = "AITask", AdvancedDisplay)
-	float MoveAcceptanceRadius; 
-	
-	UPROPERTY(BlueprintReadOnly, meta = (ExposeOnSpawn = true), Category = "AITask", AdvancedDisplay)
-	bool bShouldStopOnOverlap;
+	/** handle of path's update event delegate */
+	FDelegateHandle PathUpdateDelegateHandle;
 
-	UPROPERTY(BlueprintReadOnly, meta = (ExposeOnSpawn = true), Category = "AITask", AdvancedDisplay)
-	bool bShouldAcceptPartialPath;
+	/** handle of active ConditionalPerformMove timer  */
+	FTimerHandle MoveRetryTimerHandle;
 
-	UPROPERTY(BlueprintReadOnly, meta = (ExposeOnSpawn = true), Category = "AITask", AdvancedDisplay)
-	bool bShouldUsePathfinding;
+	/** handle of active ConditionalUpdatePath timer */
+	FTimerHandle PathRetryTimerHandle;
 
-	FDelegateHandle PathFollowingDelegateHandle;
+	/** request ID of path following's request */
 	FAIRequestID MoveRequestID;
 
-	virtual void HandleMoveFinished(FAIRequestID RequestID, EPathFollowingResult::Type Result);
+	/** currently followed path */
+	FNavPathSharedPtr Path;
+
+	TEnumAsByte<EPathFollowingResult::Type> MoveResult;
+
 	virtual void Activate() override;
 	virtual void OnDestroy(bool bOwnerFinished) override;
 
 	virtual void Pause() override;
 	virtual void Resume() override;
 
-	virtual void PostInitProperties() override;
-	
-public:
-	UAITask_MoveTo(const FObjectInitializer& ObjectInitializer);
+	/** finish task */
+	void FinishMoveTask(EPathFollowingResult::Type InResult);
 
-	/**
-	 *	@param 
-	 */
-	UFUNCTION(BlueprintCallable, Category = "AI|Tasks", meta = (AdvancedDisplay = "AcceptanceRadius,StopOnOverlap,AcceptPartialPath,bUsePathfinding", DefaultToSelf = "Controller", BlueprintInternalUseOnly = "TRUE", DisplayName = "Move To Location or Actor"))
-	static UAITask_MoveTo* AIMoveTo(AAIController* Controller, FVector GoalLocation, AActor* GoalActor = nullptr, float AcceptanceRadius = -1.f, EAIOptionFlag::Type StopOnOverlap = EAIOptionFlag::Default, EAIOptionFlag::Type AcceptPartialPath = EAIOptionFlag::Default, bool bUsePathfinding = true, bool bLockAILogic = true);
+	/** stores path and starts observing its events */
+	void SetObservedPath(FNavPathSharedPtr InPath);
 
-	void SetUp(AAIController* Controller, FVector GoalLocation, AActor* GoalActor = nullptr, float AcceptanceRadius = -1.f, bool bUsePathfinding = true, EAIOptionFlag::Type StopOnOverlap = EAIOptionFlag::Default, EAIOptionFlag::Type AcceptPartialPath = EAIOptionFlag::Default);
+	/** remove all delegates */
+	virtual void ResetObservers();
+
+	/** remove all timers */
+	virtual void ResetTimers();
+
+	/** tries to update invalidated path and handles retry timer */
+	void ConditionalUpdatePath();
+
+	/** start move request */
+	virtual void PerformMove();
+
+	/** event from followed path */
+	virtual void OnPathEvent(FNavigationPath* InPath, ENavPathEvent::Type Event);
 
-	void PerformMove();
-};
\ No newline at end of file
+	/** event from path following */
+	virtual void OnRequestFinished(FAIRequestID RequestID, EPathFollowingResult::Type Result);
+};
diff --git a/Engine/Source/Runtime/AIModule/Private/BehaviorTree/BTAuxiliaryNode.cpp b/Engine/Source/Runtime/AIModule/Private/BehaviorTree/BTAuxiliaryNode.cpp
--- a/Engine/Source/Runtime/AIModule/Private/BehaviorTree/BTAuxiliaryNode.cpp
+++ b/Engine/Source/Runtime/AIModule/Private/BehaviorTree/BTAuxiliaryNode.cpp
@@ -1,6 +1,7 @@
 // Copyright 1998-2016 Epic Games, Inc. All Rights Reserved.
 
 #include "AIModulePrivate.h"
+#include "BehaviorTree/BTCompositeNode.h"
 #include "BehaviorTree/BTAuxiliaryNode.h"
 
 UBTAuxiliaryNode::UBTAuxiliaryNode(const FObjectInitializer& ObjectInitializer) : Super(ObjectInitializer)
@@ -9,6 +10,7 @@ UBTAuxiliaryNode::UBTAuxiliaryNode(const FObjectInitializer& ObjectInitializer)
 	bNotifyCeaseRelevant = false;
 	bNotifyTick = false;
 	bTickIntervals = false;
+	ChildIndex = BTSpecialChild::OwnedByComposite;
 }
 
 void UBTAuxiliaryNode::WrappedOnBecomeRelevant(UBehaviorTreeComponent& OwnerComp, uint8* NodeMemory) const
@@ -118,6 +120,16 @@ uint16 UBTAuxiliaryNode::GetSpecialMemorySize() const
 	return bTickIntervals ? sizeof(FBTAuxiliaryMemory) : Super::GetSpecialMemorySize();
 }
 
+void UBTAuxiliaryNode::InitializeParentLink(uint8 MyChildIndex)
+{
+	ChildIndex = MyChildIndex;
+}
+
+const UBTNode* UBTAuxiliaryNode::GetMyNode() const
+{
+	return (ChildIndex == BTSpecialChild::OwnedByComposite) ? GetParentNode() : (GetParentNode() ? GetParentNode()->GetChildNode(ChildIndex) : nullptr);
+}
+
 //----------------------------------------------------------------------//
 // DEPRECATED
 //----------------------------------------------------------------------//
diff --git a/Engine/Source/Runtime/AIModule/Private/BehaviorTree/BTCompositeNode.cpp b/Engine/Source/Runtime/AIModule/Private/BehaviorTree/BTCompositeNode.cpp
--- a/Engine/Source/Runtime/AIModule/Private/BehaviorTree/BTCompositeNode.cpp
+++ b/Engine/Source/Runtime/AIModule/Private/BehaviorTree/BTCompositeNode.cpp
@@ -85,6 +85,9 @@ void UBTCompositeNode::OnChildActivation(FBehaviorTreeSearchData& SearchData, in
 	// so they can access previously executed one if needed
 	NotifyDecoratorsOnActivation(SearchData, ChildIndex);
 
+	// don't activate task services here, it's applied BEFORE aborting (e.g. abort lower pri decorator)
+	// use UBehaviorTreeComponent::ExecuteTask instead
+
 	// pass to child composite
 	if (ChildInfo.ChildComposite)
 	{
@@ -107,8 +110,16 @@ void UBTCompositeNode::OnChildDeactivation(FBehaviorTreeSearchData& SearchData,
 	// pass to decorators
 	NotifyDecoratorsOnDeactivation(SearchData, ChildIndex, NodeResult);
 
+	// pass to task services
+	if (ChildInfo.ChildTask)
+	{
+		for (int32 ServiceIndex = 0; ServiceIndex < ChildInfo.ChildTask->Services.Num(); ServiceIndex++)
+		{
+			SearchData.AddUniqueUpdate(FBehaviorTreeSearchUpdate(ChildInfo.ChildTask->Services[ServiceIndex], SearchData.OwnerComp.GetActiveInstanceIdx(), EBTNodeUpdateMode::Remove));
+		}
+	}
 	// pass to child composite
-	if (ChildInfo.ChildComposite)
+	else if (ChildInfo.ChildComposite)
 	{
 		ChildInfo.ChildComposite->OnNodeDeactivation(SearchData, NodeResult);
 	}
diff --git a/Engine/Source/Runtime/AIModule/Private/BehaviorTree/BTDecorator.cpp b/Engine/Source/Runtime/AIModule/Private/BehaviorTree/BTDecorator.cpp
--- a/Engine/Source/Runtime/AIModule/Private/BehaviorTree/BTDecorator.cpp
+++ b/Engine/Source/Runtime/AIModule/Private/BehaviorTree/BTDecorator.cpp
@@ -18,11 +18,6 @@ UBTDecorator::UBTDecorator(const FObjectInitializer& ObjectInitializer) : Super(
 	bInverseCondition = false;
 }
 
-void UBTDecorator::InitializeDecorator(uint8 MyChildIndex)
-{
-	ChildIndex = MyChildIndex;
-}
-
 bool UBTDecorator::CalculateRawConditionValue(UBehaviorTreeComponent& OwnerComp, uint8* NodeMemory) const
 {
 	return true;
@@ -87,11 +82,6 @@ void UBTDecorator::WrappedOnNodeProcessed(FBehaviorTreeSearchData& SearchData, E
 	}
 }
 
-const UBTNode* UBTDecorator::GetMyNode() const
-{
-	return GetParentNode() ? GetParentNode()->GetChildNode(ChildIndex) : NULL;
-}
-
 FString UBTDecorator::GetStaticDescription() const
 {
 	FString FlowAbortDesc;
@@ -182,3 +172,8 @@ bool UBTDecorator::CalculateRawConditionValue(UBehaviorTreeComponent* OwnerComp,
 {
 	return OwnerComp ? CalculateRawConditionValue(*OwnerComp, NodeMemory) : false;
 }
+
+void UBTDecorator::InitializeDecorator(uint8 InChildIndex)
+{
+	InitializeParentLink(InChildIndex);
+}
diff --git a/Engine/Source/Runtime/AIModule/Private/BehaviorTree/BTNode.cpp b/Engine/Source/Runtime/AIModule/Private/BehaviorTree/BTNode.cpp
--- a/Engine/Source/Runtime/AIModule/Private/BehaviorTree/BTNode.cpp
+++ b/Engine/Source/Runtime/AIModule/Private/BehaviorTree/BTNode.cpp
@@ -193,6 +193,91 @@ bool UBTNode::UsesBlueprint() const
 
 #endif
 
+UGameplayTasksComponent* UBTNode::GetGameplayTasksComponent(const UGameplayTask& Task) const
+{
+	const UAITask* AITask = Cast<const UAITask>(&Task);
+	return (AITask && AITask->GetAIController()) ? AITask->GetAIController()->GetGameplayTasksComponent(Task) : Task.GetGameplayTasksComponent();
+}
+
+AActor* UBTNode::GetGameplayTaskOwner(const UGameplayTask* Task) const
+{
+	if (Task == nullptr)
+	{
+		if (IsInstanced())
+		{
+			const UBehaviorTreeComponent* BTComponent = Cast<const UBehaviorTreeComponent>(GetOuter());
+			//not having BT component for an instanced BT node is invalid!
+			check(BTComponent);
+			return BTComponent->GetAIOwner();
+		}
+		else
+		{
+			UE_LOG(LogBehaviorTree, Warning, TEXT("%s: Unable to determine default GameplayTaskOwner!"), *GetName());
+			return nullptr;
+		}
+	}
+
+	const UAITask* AITask = Cast<const UAITask>(Task);
+	if (AITask)
+	{
+		return AITask->GetAIController();
+	}
+
+	const UGameplayTasksComponent* TasksComponent = Task->GetGameplayTasksComponent();
+	return TasksComponent ? TasksComponent->GetGameplayTaskOwner(Task) : nullptr;
+}
+
+AActor* UBTNode::GetGameplayTaskAvatar(const UGameplayTask* Task) const
+{
+	if (Task == nullptr)
+	{
+		if (IsInstanced())
+		{
+			const UBehaviorTreeComponent* BTComponent = Cast<const UBehaviorTreeComponent>(GetOuter());
+			//not having BT component for an instanced BT node is invalid!
+			check(BTComponent);
+			return BTComponent->GetAIOwner();
+		}
+		else
+		{
+			UE_LOG(LogBehaviorTree, Warning, TEXT("%s: Unable to determine default GameplayTaskAvatar!"), *GetName());
+			return nullptr;
+		}
+	}
+
+	const UAITask* AITask = Cast<const UAITask>(Task);
+	if (AITask)
+	{
+		return AITask->GetAIController() ? AITask->GetAIController()->GetPawn() : nullptr;
+	}
+
+	const UGameplayTasksComponent* TasksComponent = Task->GetGameplayTasksComponent();
+	return TasksComponent ? TasksComponent->GetGameplayTaskAvatar(Task) : nullptr;
+}
+
+uint8 UBTNode::GetGameplayTaskDefaultPriority() const
+{
+	return static_cast<uint8>(EAITaskPriority::AutonomousAI);
+}
+
+void UBTNode::OnGameplayTaskInitialized(UGameplayTask& Task)
+{
+	const UAITask* AITask = Cast<const UAITask>(&Task);
+	if (AITask && (AITask->GetAIController() == nullptr))
+	{
+		// this means that the task has either been created without specifying 
+		// UAITAsk::OwnerController's value (like via BP's Construct Object node)
+		// or it has been created in C++ with inappropriate function
+		UE_LOG(LogBehaviorTree, Error, TEXT("Missing AIController in AITask %s"), *AITask->GetName());
+	}
+}
+
+UBehaviorTreeComponent* UBTNode::GetBTComponentForTask(UGameplayTask& Task) const
+{
+	UAITask* AITask = Cast<UAITask>(&Task);
+	return (AITask && AITask->GetAIController()) ? Cast<UBehaviorTreeComponent>(AITask->GetAIController()->BrainComponent) : nullptr;
+}
+
 //----------------------------------------------------------------------//
 // DEPRECATED
 //----------------------------------------------------------------------//
@@ -261,110 +346,3 @@ FString UBTNode::GetRuntimeDescription(const UBehaviorTreeComponent* OwnerComp,
 	}
 	return TEXT("");
 }
-
-
-//----------------------------------------------------------------------//
-// UBTNode IGameplayTaskOwnerInterface
-//----------------------------------------------------------------------//
-UGameplayTasksComponent* UBTNode::GetGameplayTasksComponent(const UGameplayTask& Task) const
-{
-	const UAITask* AsAITask = Cast<const UAITask>(&Task);
-	if (AsAITask)
-	{
-		return AsAITask->GetAIController() ? AsAITask->GetAIController()->GetGameplayTasksComponent(Task) : nullptr;
-	}
-
-	return Task.GetGameplayTasksComponent();
-}
-
-void UBTNode::OnTaskInitialized(UGameplayTask& Task)
-{
-	// validate the task
-	UAITask* AsAITask = Cast<UAITask>(&Task);
-	if (AsAITask != nullptr && AsAITask->GetAIController() == nullptr)
-	{
-		// this means that the task has either been created without specifying 
-		// UAITAsk::OwnerController's value (like via BP's Construct Object node)
-		// or it has been created in C++ with inappropriate function
-		UE_LOG(LogBehaviorTree, Error, TEXT("Missing AIController in AITask %s"), *AsAITask->GetName());
-	}
-}
-
-UBehaviorTreeComponent* UBTNode::GetBTComponentForTask(UGameplayTask& Task) const
-{
-	UAITask* AsAITask = Cast<UAITask>(&Task);
-	return AsAITask && AsAITask->GetAIController() ? Cast<UBehaviorTreeComponent>(AsAITask->GetAIController()->BrainComponent) : nullptr;
-}
-
-void UBTNode::OnTaskActivated(UGameplayTask& Task)
-{
-	ensure(Task.GetTaskOwner() == this);
-}
-
-void UBTNode::OnTaskDeactivated(UGameplayTask& Task)
-{
-	ensure(Task.GetTaskOwner() == this);
-}
-
-AActor* UBTNode::GetOwnerActor(const UGameplayTask* Task) const
-{
-	if (Task == nullptr)
-	{
-		if (IsInstanced())
-		{
-			const UBehaviorTreeComponent* BTComponent = Cast<const UBehaviorTreeComponent>(GetOuter());
-			//not having BT component for an instanced BT node is invalid!
-			check(BTComponent);
-			return BTComponent->GetAIOwner();
-		}
-		else
-		{
-			UE_LOG(LogBehaviorTree, Warning, TEXT("%s: Unable to determine Owner Actor for a null GameplayTask"), *GetName());
-			return nullptr;
-		}
-	}
-
-	const UAITask* AsAITask = Cast<const UAITask>(Task);
-	if (AsAITask)
-	{
-		return AsAITask->GetAIController();
-	}
-
-	const UGameplayTasksComponent* GTComponent = Task->GetGameplayTasksComponent();
-
-	return GTComponent ? GTComponent->GetOwnerActor(Task) : nullptr;
-}
-
-AActor* UBTNode::GetAvatarActor(const UGameplayTask* Task) const
-{
-	if (Task == nullptr)
-	{
-		if (IsInstanced())
-		{
-			const UBehaviorTreeComponent* BTComponent = Cast<const UBehaviorTreeComponent>(GetOuter());
-			//not having BT component for an instanced BT node is invalid!
-			check(BTComponent);
-			return BTComponent->GetAIOwner() ? BTComponent->GetAIOwner()->GetPawn() : nullptr;
-		}
-		else
-		{
-			UE_LOG(LogBehaviorTree, Warning, TEXT("%s: Unable to determine Avatar Actor for a null GameplayTask"), *GetName());
-			return nullptr;
-		}
-	}
-
-	const UAITask* AsAITask = Cast<const UAITask>(Task);
-	if (AsAITask)
-	{
-		return AsAITask->GetAIController() ? AsAITask->GetAIController()->GetPawn() : nullptr;
-	}
-
-	const UGameplayTasksComponent* GTComponent = Task->GetGameplayTasksComponent();
-
-	return GTComponent ? GTComponent->GetOwnerActor(Task) : nullptr;
-}
-
-uint8 UBTNode::GetDefaultPriority() const
-{
-	return uint8(EAITaskPriority::AutonomousAI);
-}
diff --git a/Engine/Source/Runtime/AIModule/Private/BehaviorTree/BTTaskNode.cpp b/Engine/Source/Runtime/AIModule/Private/BehaviorTree/BTTaskNode.cpp
--- a/Engine/Source/Runtime/AIModule/Private/BehaviorTree/BTTaskNode.cpp
+++ b/Engine/Source/Runtime/AIModule/Private/BehaviorTree/BTTaskNode.cpp
@@ -6,7 +6,6 @@
 
 UBTTaskNode::UBTTaskNode(const FObjectInitializer& ObjectInitializer) : Super(ObjectInitializer)
 {
-	NodeName = GetClass()->GetName();
 	bNotifyTick = false;
 	bNotifyTaskFinished = false;
 	bIgnoreRestartSelf = false;
@@ -154,12 +153,10 @@ FName UBTTaskNode::GetNodeIconName() const
 
 #endif	// WITH_EDITOR
 
-//----------------------------------------------------------------------//
-// UBTTaskNode IGameplayTaskOwnerInterface
-//----------------------------------------------------------------------//
-void UBTTaskNode::OnTaskDeactivated(UGameplayTask& Task)
+void UBTTaskNode::OnGameplayTaskDeactivated(UGameplayTask& Task)
 {
 	ensure(Task.GetTaskOwner() == this);
+
 	UBehaviorTreeComponent* BTComp = GetBTComponentForTask(Task);
 	if (BTComp)
 	{
