Copying //UE4/Dev-Framework to Dev-Main (//UE4/Dev-Main) @ 2926658

diff --git a/Engine/Source/Runtime/AIModule/Private/BehaviorTree/BehaviorTreeComponent.cpp b/Engine/Source/Runtime/AIModule/Private/BehaviorTree/BehaviorTreeComponent.cpp
--- a/Engine/Source/Runtime/AIModule/Private/BehaviorTree/BehaviorTreeComponent.cpp
+++ b/Engine/Source/Runtime/AIModule/Private/BehaviorTree/BehaviorTreeComponent.cpp
@@ -749,7 +749,8 @@ void UBehaviorTreeComponent::RequestExecution(UBTCompositeNode* RequestedOn, int
 	if (bSwitchToHigherPriority)
 	{
 		// check if decorators allow execution on requesting link
-		const bool bShouldCheckDecorators = (RequestedByChildIndex >= 0);
+		// unless it's branch restart (abort result within current branch), when it can't be skipped because branch can be no longer valid
+		const bool bShouldCheckDecorators = (RequestedByChildIndex >= 0) && !IsExecutingBranch(RequestedBy, RequestedByChildIndex);
 		const bool bCanExecute = !bShouldCheckDecorators || RequestedOn->DoDecoratorsAllowExecution(*this, InstanceIdx, RequestedByChildIndex);
 		if (!bCanExecute)
 		{
@@ -856,7 +857,7 @@ void UBehaviorTreeComponent::RequestExecution(UBTCompositeNode* RequestedOn, int
 
 	ExecutionRequest.SearchStart = ExecutionIdx;
 	ExecutionRequest.ContinueWithResult = ContinueWithResult;
-	ExecutionRequest.bTryNextChild = !bSwitchToHigherPriority;
+	ExecutionRequest.bTryNextChild = !bSwitchToHigherPriority && !PendingExecution.bAborting;
 	ExecutionRequest.bIsRestart = (RequestedBy != GetActiveNode());
 	PendingExecution.Lock();
 	
@@ -1280,7 +1281,11 @@ void UBehaviorTreeComponent::ProcessExecutionRequest()
 		// abort task if needed
 		if (InstanceStack.Last().ActiveNodeType == EBTActiveNode::ActiveTask)
 		{
+			PendingExecution.OnAbortStart();
+
 			AbortCurrentTask();
+
+			PendingExecution.OnAbortProcessed();
 		}
 
 		// set next task to execute
@@ -1440,6 +1445,19 @@ void UBehaviorTreeComponent::ExecuteTask(UBTTaskNode* TaskNode)
 	SCOPE_CYCLE_COUNTER(STAT_AI_BehaviorTree_ExecutionTime);
 
 	FBehaviorTreeInstance& ActiveInstance = InstanceStack[ActiveInstanceIdx];
+
+	// task service activation is not part of search update (although deactivation is, through DeactivateUpTo), start them before execution
+	for (int32 ServiceIndex = 0; ServiceIndex < TaskNode->Services.Num(); ServiceIndex++)
+	{
+		UBTService* ServiceNode = TaskNode->Services[ServiceIndex];
+		uint8* NodeMemory = (uint8*)ServiceNode->GetNodeMemory<uint8>(ActiveInstance);
+
+		ActiveInstance.ActiveAuxNodes.Add(ServiceNode);
+
+		UE_VLOG(GetOwner(), LogBehaviorTree, Log, TEXT("Activating task service: %s"), *UBehaviorTreeTypes::DescribeNodeHelper(ServiceNode));
+		ServiceNode->WrappedOnBecomeRelevant(*this, NodeMemory);
+	}
+
 	ActiveInstance.ActiveNode = TaskNode;
 	ActiveInstance.ActiveNodeType = EBTActiveNode::ActiveTask;
 
@@ -1788,9 +1806,20 @@ UBTNode* UBehaviorTreeComponent::FindTemplateNode(const UBTNode* Node) const
 	{
 		FBTCompositeChild& ChildInfo = ParentNode->Children[ChildIndex];
 
-		if (ChildInfo.ChildTask && ChildInfo.ChildTask->GetExecutionIndex() == Node->GetExecutionIndex())
+		if (ChildInfo.ChildTask)
 		{
-			return ChildInfo.ChildTask;
+			if (ChildInfo.ChildTask->GetExecutionIndex() == Node->GetExecutionIndex())
+			{
+				return ChildInfo.ChildTask;
+			}
+
+			for (int32 ServiceIndex = 0; ServiceIndex < ChildInfo.ChildTask->Services.Num(); ServiceIndex++)
+			{
+				if (ChildInfo.ChildTask->Services[ServiceIndex]->GetExecutionIndex() == Node->GetExecutionIndex())
+				{
+					return ChildInfo.ChildTask->Services[ServiceIndex];
+				}
+			}
 		}
 
 		for (int32 DecoratorIndex = 0; DecoratorIndex < ChildInfo.Decorators.Num(); DecoratorIndex++)
diff --git a/Engine/Source/Runtime/AIModule/Private/BehaviorTree/BehaviorTreeManager.cpp b/Engine/Source/Runtime/AIModule/Private/BehaviorTree/BehaviorTreeManager.cpp
--- a/Engine/Source/Runtime/AIModule/Private/BehaviorTree/BehaviorTreeManager.cpp
+++ b/Engine/Source/Runtime/AIModule/Private/BehaviorTree/BehaviorTreeManager.cpp
@@ -110,6 +110,7 @@ static void InitializeNodeHelper(UBTCompositeNode* ParentNode, UBTNode* NodeOb,
 				Service->GetInstanceMemorySize(), Service->GetSpecialMemorySize()));
 
 			Service->InitializeFromAsset(TreeAsset);
+			// don't initialize parent link for services on composite node
 			ExecutionIndex++;
 		}
 
@@ -135,7 +136,7 @@ static void InitializeNodeHelper(UBTCompositeNode* ParentNode, UBTNode* NodeOb,
 					Decorator->GetInstanceMemorySize(), Decorator->GetSpecialMemorySize()));
 
 				Decorator->InitializeFromAsset(TreeAsset);
-				Decorator->InitializeDecorator(ChildIndex);
+				Decorator->InitializeParentLink(ChildIndex);
 				ExecutionIndex++;
 			}
 
@@ -150,6 +151,29 @@ static void InitializeNodeHelper(UBTCompositeNode* ParentNode, UBTNode* NodeOb,
 			{
 				ChildInfo.ChildTask = Cast<UBTTaskNode>(StaticDuplicateObject(ChildInfo.ChildTask, NodeOuter));
 				ChildNode = ChildInfo.ChildTask;
+
+				for (int32 ServiceIndex = 0; ServiceIndex < ChildInfo.ChildTask->Services.Num(); ServiceIndex++)
+				{
+					if (ChildInfo.ChildTask->Services[ServiceIndex] == NULL)
+					{
+						UE_LOG(LogBehaviorTree, Warning, TEXT("%s has missing service node! (parent: %s)"),
+							*TreeAsset.GetName(), *UBehaviorTreeTypes::DescribeNodeHelper(ChildInfo.ChildTask));
+
+						ChildInfo.ChildTask->Services.RemoveAt(ServiceIndex, 1, false);
+						ServiceIndex--;
+						continue;
+					}
+
+					UBTService* Service = Cast<UBTService>(StaticDuplicateObject(ChildInfo.ChildTask->Services[ServiceIndex], NodeOuter));
+					ChildInfo.ChildTask->Services[ServiceIndex] = Service;
+
+					InitList.Add(FNodeInitializationData(Service, CompositeOb, ExecutionIndex, TreeDepth,
+						Service->GetInstanceMemorySize(), Service->GetSpecialMemorySize()));
+
+					Service->InitializeFromAsset(TreeAsset);
+					Service->InitializeParentLink(ChildIndex);
+					ExecutionIndex++;
+				}
 			}
 
 			if (ChildNode)
@@ -307,11 +331,10 @@ void StatNodeUsage(UBTNode* Node, FNodeClassCounter& NodeCounter)
 	{
 		for (int32 ServiceIndex = 0; ServiceIndex < CompositeOb->Services.Num(); ServiceIndex++)
 		{
-			if (CompositeOb->Services[ServiceIndex] == NULL)
+			if (CompositeOb->Services[ServiceIndex])
 			{
-				continue;
+				NodeCounter.CountNode(CompositeOb->Services[ServiceIndex]);
 			}
-			NodeCounter.CountNode(CompositeOb->Services[ServiceIndex]);
 		}
 
 		for (int32 ChildIndex = 0; ChildIndex < CompositeOb->Children.Num(); ChildIndex++)
@@ -319,12 +342,10 @@ void StatNodeUsage(UBTNode* Node, FNodeClassCounter& NodeCounter)
 			FBTCompositeChild& ChildInfo = CompositeOb->Children[ChildIndex];
 			for (int32 DecoratorIndex = 0; DecoratorIndex < ChildInfo.Decorators.Num(); DecoratorIndex++)
 			{
-				if (ChildInfo.Decorators[DecoratorIndex] == NULL)
+				if (ChildInfo.Decorators[DecoratorIndex])
 				{
-					continue;
+					NodeCounter.CountNode(ChildInfo.Decorators[DecoratorIndex]);
 				}
-
-				NodeCounter.CountNode(ChildInfo.Decorators[DecoratorIndex]);
 			}
 
 			UBTNode* ChildNode = NULL;
@@ -335,6 +356,14 @@ void StatNodeUsage(UBTNode* Node, FNodeClassCounter& NodeCounter)
 			}
 			else if (ChildInfo.ChildTask)
 			{
+				for (int32 ServiceIndex = 0; ServiceIndex < ChildInfo.ChildTask->Services.Num(); ServiceIndex++)
+				{
+					if (ChildInfo.ChildTask->Services[ServiceIndex])
+					{
+						NodeCounter.CountNode(ChildInfo.ChildTask->Services[ServiceIndex]);
+					}
+				}
+
 				NodeCounter.CountNode(ChildInfo.ChildTask);
 			}
 		}
diff --git a/Engine/Source/Runtime/AIModule/Private/BehaviorTree/BehaviorTreeTypes.cpp b/Engine/Source/Runtime/AIModule/Private/BehaviorTree/BehaviorTreeTypes.cpp
--- a/Engine/Source/Runtime/AIModule/Private/BehaviorTree/BehaviorTreeTypes.cpp
+++ b/Engine/Source/Runtime/AIModule/Private/BehaviorTree/BehaviorTreeTypes.cpp
@@ -39,7 +39,7 @@ void FBehaviorTreeInstance::Initialize(UBehaviorTreeComponent& OwnerComp, UBTCom
 			UBTDecorator* InstancedDecoratorOb = Cast<UBTDecorator>(DecoratorOb->GetNodeInstance(OwnerComp, DecoratorMemory));
 			if (InstancedDecoratorOb)
 			{
-				InstancedDecoratorOb->InitializeDecorator(DecoratorOb->GetChildIndex());
+				InstancedDecoratorOb->InitializeParentLink(DecoratorOb->GetChildIndex());
 			}
 		}
 
@@ -49,6 +49,19 @@ void FBehaviorTreeInstance::Initialize(UBehaviorTreeComponent& OwnerComp, UBTCom
 		}
 		else if (ChildInfo.ChildTask)
 		{
+			for (int32 ServiceIndex = 0; ServiceIndex < ChildInfo.ChildTask->Services.Num(); ServiceIndex++)
+			{
+				UBTService* ServiceOb = ChildInfo.ChildTask->Services[ServiceIndex];
+				uint8* ServiceMemory = ServiceOb->GetNodeMemory<uint8>(*this);
+				ServiceOb->InitializeInSubtree(OwnerComp, ServiceMemory, InstancedIndex, InitType);
+
+				UBTService* InstancedServiceOb = Cast<UBTService>(ServiceOb->GetNodeInstance(OwnerComp, ServiceMemory));
+				if (InstancedServiceOb)
+				{
+					InstancedServiceOb->InitializeParentLink(ServiceOb->GetChildIndex());
+				}
+			}
+
 			ChildInfo.ChildTask->InitializeInSubtree(OwnerComp, ChildInfo.ChildTask->GetNodeMemory<uint8>(*this), InstancedIndex, InitType);
 		}
 	}
@@ -119,6 +132,11 @@ void FBehaviorTreeInstance::CleanupNodes(UBehaviorTreeComponent& OwnerComp, UBTC
 		}
 		else if (ChildInfo.ChildTask)
 		{
+			for (int32 ServiceIndex = 0; ServiceIndex < ChildInfo.ChildTask->Services.Num(); ServiceIndex++)
+			{
+				ChildInfo.ChildTask->Services[ServiceIndex]->CleanupInSubtree(OwnerComp, ChildInfo.ChildTask->Services[ServiceIndex]->GetNodeMemory<uint8>(*this), CleanupType);
+			}
+
 			ChildInfo.ChildTask->CleanupInSubtree(OwnerComp, ChildInfo.ChildTask->GetNodeMemory<uint8>(*this), CleanupType);
 		}
 	}
diff --git a/Engine/Source/Runtime/AIModule/Private/BehaviorTree/Tasks/BTTask_MoveDirectlyToward.cpp b/Engine/Source/Runtime/AIModule/Private/BehaviorTree/Tasks/BTTask_MoveDirectlyToward.cpp
--- a/Engine/Source/Runtime/AIModule/Private/BehaviorTree/Tasks/BTTask_MoveDirectlyToward.cpp
+++ b/Engine/Source/Runtime/AIModule/Private/BehaviorTree/Tasks/BTTask_MoveDirectlyToward.cpp
@@ -1,154 +1,36 @@
 // Copyright 1998-2016 Epic Games, Inc. All Rights Reserved.
 
 #include "AIModulePrivate.h"
-#include "BehaviorTree/BlackboardComponent.h"
-#include "BehaviorTree/Blackboard/BlackboardKeyType_Object.h"
-#include "BehaviorTree/Blackboard/BlackboardKeyType_Vector.h"
 #include "BehaviorTree/Tasks/BTTask_MoveDirectlyToward.h"
-#include "Tasks/AITask_MoveTo.h"
 
-UBTTask_MoveDirectlyToward::UBTTask_MoveDirectlyToward(const FObjectInitializer& ObjectInitializer) 
-	: Super(ObjectInitializer)
-	, bProjectVectorGoalToNavigation(true)
+UBTTask_MoveDirectlyToward::UBTTask_MoveDirectlyToward(const FObjectInitializer& ObjectInitializer) : Super(ObjectInitializer)
 {
 	NodeName = "MoveDirectlyToward";
+	bUsePathfinding = false;
 
-	AcceptableRadius = GET_AI_CONFIG_VAR(AcceptanceRadius);
-	bStopOnOverlap = GET_AI_CONFIG_VAR(bFinishMoveOnGoalOverlap);
-	bAllowStrafe = GET_AI_CONFIG_VAR(bAllowStrafing);
-
-	// accept only actors and vectors
-	BlackboardKey.AddObjectFilter(this, GET_MEMBER_NAME_CHECKED(UBTTask_MoveDirectlyToward, BlackboardKey), AActor::StaticClass());
-	BlackboardKey.AddVectorFilter(this, GET_MEMBER_NAME_CHECKED(UBTTask_MoveDirectlyToward, BlackboardKey));
-}
-
-EBTNodeResult::Type UBTTask_MoveDirectlyToward::ExecuteTask(UBehaviorTreeComponent& OwnerComp, uint8* NodeMemory)
-{
-	const UBlackboardComponent* MyBlackboard = OwnerComp.GetBlackboardComponent();
-	FBTMoveDirectlyTowardMemory* MyMemory = reinterpret_cast<FBTMoveDirectlyTowardMemory*>(NodeMemory);
-	AAIController* MyController = OwnerComp.GetAIOwner();
-	EBTNodeResult::Type NodeResult = EBTNodeResult::Failed;
-
-	if (MyController && MyBlackboard)
-	{
-		if (GET_AI_CONFIG_VAR(bEnableBTAITasks))
-		{
-			UAITask_MoveTo* AIMoveTask = NewBTAITask<UAITask_MoveTo>(OwnerComp);
-
-			if (AIMoveTask != nullptr)
-			{
-				bool bSetUp = false;
-				if (BlackboardKey.SelectedKeyType == UBlackboardKeyType_Object::StaticClass())
-				{
-					UObject* KeyValue = MyBlackboard->GetValue<UBlackboardKeyType_Object>(BlackboardKey.GetSelectedKeyID());
-					AActor* TargetActor = Cast<AActor>(KeyValue);
-					if (TargetActor)
-					{
-						AIMoveTask->SetUp(MyController, FVector::ZeroVector, TargetActor, AcceptableRadius, /*bUsePathfinding=*/false, FAISystem::BoolToAIOption(bStopOnOverlap));
-						NodeResult = EBTNodeResult::InProgress;
-					}
-					else
-					{
-						UE_VLOG(MyController, LogBehaviorTree, Warning, TEXT("UBTTask_MoveDirectlyToward::ExecuteTask tried to go to actor while BB %s entry was empty"), *BlackboardKey.SelectedKeyName.ToString());
-					}
-				}
-				else if (BlackboardKey.SelectedKeyType == UBlackboardKeyType_Vector::StaticClass())
-				{
-					const FVector TargetLocation = MyBlackboard->GetValue<UBlackboardKeyType_Vector>(BlackboardKey.GetSelectedKeyID());
-					AIMoveTask->SetUp(MyController, TargetLocation, nullptr, AcceptableRadius, /*bUsePathfinding=*/false, FAISystem::BoolToAIOption(bStopOnOverlap));
-					NodeResult = EBTNodeResult::InProgress;
-				}
-
-				if (NodeResult == EBTNodeResult::InProgress)
-				{
-					AIMoveTask->ReadyForActivation();
-				}
-			}
-		}
-		else
-		{
-			EPathFollowingRequestResult::Type RequestResult = EPathFollowingRequestResult::Failed;
-
-			if (BlackboardKey.SelectedKeyType == UBlackboardKeyType_Object::StaticClass())
-			{
-				UObject* KeyValue = MyBlackboard->GetValue<UBlackboardKeyType_Object>(BlackboardKey.GetSelectedKeyID());
-				AActor* TargetActor = Cast<AActor>(KeyValue);
-				if (TargetActor)
-				{
-					RequestResult = bDisablePathUpdateOnGoalLocationChange ?
-						MyController->MoveToLocation(TargetActor->GetActorLocation(), AcceptableRadius, bStopOnOverlap, /*bUsePathfinding=*/false, /*bProjectDestinationToNavigation=*/bProjectVectorGoalToNavigation, bAllowStrafe) :
-						MyController->MoveToActor(TargetActor, AcceptableRadius, bStopOnOverlap, /*bUsePathfinding=*/false, bAllowStrafe);
-				}
-			}
-			else if (BlackboardKey.SelectedKeyType == UBlackboardKeyType_Vector::StaticClass())
-			{
-				const FVector TargetLocation = MyBlackboard->GetValue<UBlackboardKeyType_Vector>(BlackboardKey.GetSelectedKeyID());
-				RequestResult = MyController->MoveToLocation(TargetLocation, AcceptableRadius, bStopOnOverlap, /*bUsePathfinding=*/false, /*bProjectDestinationToNavigation=*/bProjectVectorGoalToNavigation, bAllowStrafe);
-			}
-
-			if (RequestResult == EPathFollowingRequestResult::RequestSuccessful)
-			{
-				const FAIRequestID RequestID = MyController->GetCurrentMoveRequestID();
-
-				MyMemory->MoveRequestID = RequestID;
-				WaitForMessage(OwnerComp, UBrainComponent::AIMessage_MoveFinished, RequestID);
-
-				NodeResult = EBTNodeResult::InProgress;
-			}
-			else if (RequestResult == EPathFollowingRequestResult::AlreadyAtGoal)
-			{
-				NodeResult = EBTNodeResult::Succeeded;
-			}
-		}
-	}
-
-	return NodeResult;
+	bProjectVectorGoalToNavigation = true;
+	bDisablePathUpdateOnGoalLocationChange = false;
+	bUpdatedDeprecatedProperties = false;
 }
 
-EBTNodeResult::Type UBTTask_MoveDirectlyToward::AbortTask(UBehaviorTreeComponent& OwnerComp, uint8* NodeMemory)
+void UBTTask_MoveDirectlyToward::InitializeFromAsset(UBehaviorTree& Asset)
 {
-	FBTMoveDirectlyTowardMemory* MyMemory = reinterpret_cast<FBTMoveDirectlyTowardMemory*>(NodeMemory);
-	AAIController* MyController = OwnerComp.GetAIOwner();
-
-	if (MyController && MyController->GetPathFollowingComponent())
-	{
-		MyController->GetPathFollowingComponent()->AbortMove(TEXT("BehaviorTree abort"), MyMemory->MoveRequestID);
-	}
-
-	return Super::AbortTask(OwnerComp, NodeMemory);
-}
-
-FString UBTTask_MoveDirectlyToward::GetStaticDescription() const
-{
-	FString KeyDesc("invalid");
-	if (BlackboardKey.SelectedKeyType == UBlackboardKeyType_Object::StaticClass() ||
-		BlackboardKey.SelectedKeyType == UBlackboardKeyType_Vector::StaticClass())
-	{
-		KeyDesc = BlackboardKey.SelectedKeyName.ToString();
-	}
-
-	return FString::Printf(TEXT("%s: %s"), *Super::GetStaticDescription(), *KeyDesc);
+	bUpdatedDeprecatedProperties = true;
+	Super::InitializeFromAsset(Asset);
 }
 
-void UBTTask_MoveDirectlyToward::DescribeRuntimeValues(const UBehaviorTreeComponent& OwnerComp, uint8* NodeMemory, EBTDescriptionVerbosity::Type Verbosity, TArray<FString>& Values) const
+void UBTTask_MoveDirectlyToward::PostLoad()
 {
-	Super::DescribeRuntimeValues(OwnerComp, NodeMemory, Verbosity, Values);
-
-	const UBlackboardComponent* BlackboardComp = OwnerComp.GetBlackboardComponent();
-	FBTMoveDirectlyTowardMemory* MyMemory = reinterpret_cast<FBTMoveDirectlyTowardMemory*>(NodeMemory);
+	Super::PostLoad();
 
-	if (MyMemory->MoveRequestID && BlackboardComp)
+	if (!bUpdatedDeprecatedProperties)
 	{
-		FString KeyValue = BlackboardComp->DescribeKeyValue(BlackboardKey.GetSelectedKeyID(), EBlackboardDescription::OnlyValue);
-		Values.Add(FString::Printf(TEXT("move target: %s"), *KeyValue));
+		bProjectGoalLocation = bProjectVectorGoalToNavigation;
+		bTrackMovingGoal = bDisablePathUpdateOnGoalLocationChange;
+		bUpdatedDeprecatedProperties = true;
 	}
 }
 
-uint16 UBTTask_MoveDirectlyToward::GetInstanceMemorySize() const
-{
-	return sizeof(FBTMoveDirectlyTowardMemory);
-}
-
 #if WITH_EDITOR
 
 FName UBTTask_MoveDirectlyToward::GetNodeIconName() const
diff --git a/Engine/Source/Runtime/AIModule/Private/BehaviorTree/Tasks/BTTask_MoveTo.cpp b/Engine/Source/Runtime/AIModule/Private/BehaviorTree/Tasks/BTTask_MoveTo.cpp
--- a/Engine/Source/Runtime/AIModule/Private/BehaviorTree/Tasks/BTTask_MoveTo.cpp
+++ b/Engine/Source/Runtime/AIModule/Private/BehaviorTree/Tasks/BTTask_MoveTo.cpp
@@ -8,17 +8,20 @@
 #include "BehaviorTree/Blackboard/BlackboardKeyType_Vector.h"
 #include "Tasks/AITask_MoveTo.h"
 
-UBTTask_MoveTo::UBTTask_MoveTo(const FObjectInitializer& ObjectInitializer) 
-	: Super(ObjectInitializer)
+UBTTask_MoveTo::UBTTask_MoveTo(const FObjectInitializer& ObjectInitializer) : Super(ObjectInitializer)
 {
 	NodeName = "Move To";
-	bNotifyTick = true;
+	bUseGameplayTasks = GET_AI_CONFIG_VAR(bEnableBTAITasks);
+	bNotifyTick = !bUseGameplayTasks;
 	bNotifyTaskFinished = true;
 
 	AcceptableRadius = GET_AI_CONFIG_VAR(AcceptanceRadius);
 	bStopOnOverlap = GET_AI_CONFIG_VAR(bFinishMoveOnGoalOverlap); 
 	bAllowStrafe = GET_AI_CONFIG_VAR(bAllowStrafing);
 	bAllowPartialPath = GET_AI_CONFIG_VAR(bAcceptPartialPaths);
+	bTrackMovingGoal = true;
+	bProjectGoalLocation = true;
+	bUsePathfinding = true;
 
 	ObservedBlackboardValueTolerance = AcceptableRadius * 0.95f;
 
@@ -33,9 +36,10 @@ EBTNodeResult::Type UBTTask_MoveTo::ExecuteTask(UBehaviorTreeComponent& OwnerCom
 
 	FBTMoveToTaskMemory* MyMemory = reinterpret_cast<FBTMoveToTaskMemory*>(NodeMemory);
 	MyMemory->PreviousGoalLocation = FAISystem::InvalidLocation;
-	AAIController* MyController = OwnerComp.GetAIOwner();
+	MyMemory->MoveRequestID = FAIRequestID::InvalidRequest;
 
-	MyMemory->bWaitingForPath = MyController->ShouldPostponePathUpdates();
+	AAIController* MyController = OwnerComp.GetAIOwner();
+	MyMemory->bWaitingForPath = bUseGameplayTasks ? false : MyController->ShouldPostponePathUpdates();
 	if (!MyMemory->bWaitingForPath)
 	{
 		NodeResult = PerformMoveTask(OwnerComp, NodeMemory);
@@ -57,8 +61,7 @@ EBTNodeResult::Type UBTTask_MoveTo::ExecuteTask(UBehaviorTreeComponent& OwnerCom
 			}
 			MyMemory->BBObserverDelegateHandle = BlackboardComp->RegisterObserver(BlackboardKey.GetSelectedKeyID(), this, FOnBlackboardChangeNotification::CreateUObject(this, &UBTTask_MoveTo::OnBlackboardValueChange));
 		}
-	}
-	
+	}	
 	
 	return NodeResult;
 }
@@ -70,125 +73,111 @@ EBTNodeResult::Type UBTTask_MoveTo::PerformMoveTask(UBehaviorTreeComponent& Owne
 	AAIController* MyController = OwnerComp.GetAIOwner();
 
 	EBTNodeResult::Type NodeResult = EBTNodeResult::Failed;
-
 	if (MyController && MyBlackboard)
 	{
-		if (GET_AI_CONFIG_VAR(bEnableBTAITasks))
+		FAIMoveRequest MoveReq;
+		MoveReq.SetNavigationFilter(FilterClass);
+		MoveReq.SetAllowPartialPath(bAllowPartialPath);
+		MoveReq.SetAcceptanceRadius(AcceptableRadius);
+		MoveReq.SetCanStrafe(bAllowStrafe);
+		MoveReq.SetStopOnOverlap(bStopOnOverlap);
+		MoveReq.SetProjectGoalLocation(bProjectGoalLocation);
+		MoveReq.SetUsePathfinding(bUsePathfinding);
+
+		if (BlackboardKey.SelectedKeyType == UBlackboardKeyType_Object::StaticClass())
 		{
-			bool bTaskReusing = false;
-			UAITask_MoveTo* AIMoveTask = nullptr;
-
-			if (MyMemory->Task.IsValid())
+			UObject* KeyValue = MyBlackboard->GetValue<UBlackboardKeyType_Object>(BlackboardKey.GetSelectedKeyID());
+			AActor* TargetActor = Cast<AActor>(KeyValue);
+			if (TargetActor)
 			{
-				AIMoveTask = MyMemory->Task.Get();
-				ensure(AIMoveTask->GetState() != EGameplayTaskState::Finished);
-				bTaskReusing = true;
-			}
-			else
-			{
-				AIMoveTask = NewBTAITask<UAITask_MoveTo>(OwnerComp);
-			}			
-			
-			if (AIMoveTask != nullptr)
-			{
-				bool bSetUp = false;
-				if (BlackboardKey.SelectedKeyType == UBlackboardKeyType_Object::StaticClass())
+				if (bTrackMovingGoal)
 				{
-					UObject* KeyValue = MyBlackboard->GetValue<UBlackboardKeyType_Object>(BlackboardKey.GetSelectedKeyID());
-					AActor* TargetActor = Cast<AActor>(KeyValue);
-					if (TargetActor)
-					{
-						AIMoveTask->SetUp(MyController, FVector::ZeroVector, TargetActor, AcceptableRadius, /*bUsePathfinding=*/true, FAISystem::BoolToAIOption(bStopOnOverlap), FAISystem::BoolToAIOption(bAllowPartialPath));
-						NodeResult = EBTNodeResult::InProgress;
-					}
-					else
-					{
-						UE_VLOG(MyController, LogBehaviorTree, Warning, TEXT("UBTTask_MoveTo::ExecuteTask tried to go to actor while BB %s entry was empty"), *BlackboardKey.SelectedKeyName.ToString());
-					}
+					MoveReq.SetGoalActor(TargetActor);
 				}
-				else if (BlackboardKey.SelectedKeyType == UBlackboardKeyType_Vector::StaticClass())
+				else
 				{
-					const FVector TargetLocation = MyBlackboard->GetValue<UBlackboardKeyType_Vector>(BlackboardKey.GetSelectedKeyID());
-					MyMemory->PreviousGoalLocation = TargetLocation;
-					AIMoveTask->SetUp(MyController, TargetLocation, nullptr, AcceptableRadius, /*bUsePathfinding=*/true, FAISystem::BoolToAIOption(bStopOnOverlap), FAISystem::BoolToAIOption(bAllowPartialPath));
-					NodeResult = EBTNodeResult::InProgress;
+					MoveReq.SetGoalLocation(TargetActor->GetActorLocation());
 				}
+			}
+			else
+			{
+				UE_VLOG(MyController, LogBehaviorTree, Warning, TEXT("UBTTask_MoveTo::ExecuteTask tried to go to actor while BB %s entry was empty"), *BlackboardKey.SelectedKeyName.ToString());
+			}
+		}
+		else if (BlackboardKey.SelectedKeyType == UBlackboardKeyType_Vector::StaticClass())
+		{
+			const FVector TargetLocation = MyBlackboard->GetValue<UBlackboardKeyType_Vector>(BlackboardKey.GetSelectedKeyID());
+			MoveReq.SetGoalLocation(TargetLocation);
+
+			MyMemory->PreviousGoalLocation = TargetLocation;
+		}
 
-				if (NodeResult == EBTNodeResult::InProgress)
+		if (MoveReq.IsValid())
+		{
+			if (GET_AI_CONFIG_VAR(bEnableBTAITasks))
+			{
+				UAITask_MoveTo* MoveTask = MyMemory->Task.Get();
+				const bool bReuseExistingTask = (MoveTask != nullptr);
+
+				MoveTask = PrepareMoveTask(OwnerComp, MoveTask, MoveReq);
+				if (MoveTask)
 				{
-					/*WaitForMessage(OwnerComp, UBrainComponent::AIMessage_MoveFinished, RequestID);
-					WaitForMessage(OwnerComp, UBrainComponent::AIMessage_RepathFailed);*/
+					MyMemory->bObserverCanFinishTask = false;
 
-					if (bTaskReusing == false)
+					if (bReuseExistingTask)
 					{
-						MyMemory->Task = AIMoveTask;
-						UE_VLOG(MyController, LogBehaviorTree, Verbose, TEXT("\'%s\' task implementing move with task %s"), *GetNodeName(), *AIMoveTask->GetName());
-						AIMoveTask->ReadyForActivation();
+						ensure(MoveTask->IsActive());
+						UE_VLOG(MyController, LogBehaviorTree, Verbose, TEXT("\'%s\' reusing AITask %s"), *GetNodeName(), *MoveTask->GetName());
+						MoveTask->ConditionalPerformMove();
 					}
 					else
 					{
-						ensure(AIMoveTask->IsActive());
-						UE_VLOG(MyController, LogBehaviorTree, Verbose, TEXT("\'%s\' reusing AITask %s"), *GetNodeName(), *AIMoveTask->GetName());
-						AIMoveTask->PerformMove();
+						MyMemory->Task = MoveTask;
+						UE_VLOG(MyController, LogBehaviorTree, Verbose, TEXT("\'%s\' task implementing move with task %s"), *GetNodeName(), *MoveTask->GetName());
+						MoveTask->ReadyForActivation();
 					}
+
+					MyMemory->bObserverCanFinishTask = true;
+					NodeResult = (MoveTask->GetState() != EGameplayTaskState::Finished) ? EBTNodeResult::InProgress :
+						MoveTask->WasMoveSuccessful() ? EBTNodeResult::Succeeded :
+						EBTNodeResult::Failed;
 				}
 			}
-		}
-		else
-		{
-			EPathFollowingRequestResult::Type RequestResult = EPathFollowingRequestResult::Failed;
-
-			FAIMoveRequest MoveReq;
-			MoveReq.SetNavigationFilter(FilterClass);
-			MoveReq.SetAllowPartialPath(bAllowPartialPath);
-			MoveReq.SetAcceptanceRadius(AcceptableRadius);
-			MoveReq.SetCanStrafe(bAllowStrafe);
-			MoveReq.SetStopOnOverlap(bStopOnOverlap);
-
-			if (BlackboardKey.SelectedKeyType == UBlackboardKeyType_Object::StaticClass())
+			else
 			{
-				UObject* KeyValue = MyBlackboard->GetValue<UBlackboardKeyType_Object>(BlackboardKey.GetSelectedKeyID());
-				AActor* TargetActor = Cast<AActor>(KeyValue);
-				if (TargetActor)
+				EPathFollowingRequestResult::Type RequestResult = MyController->MoveTo(MoveReq);
+				if (RequestResult == EPathFollowingRequestResult::RequestSuccessful)
 				{
-					MoveReq.SetGoalActor(TargetActor);
+					const FAIRequestID RequestID = MyController->GetCurrentMoveRequestID();
+
+					MyMemory->MoveRequestID = RequestID;
+					WaitForMessage(OwnerComp, UBrainComponent::AIMessage_MoveFinished, RequestID);
+					WaitForMessage(OwnerComp, UBrainComponent::AIMessage_RepathFailed);
 
-					RequestResult = MyController->MoveTo(MoveReq);
+					NodeResult = EBTNodeResult::InProgress;
 				}
-				else
+				else if (RequestResult == EPathFollowingRequestResult::AlreadyAtGoal)
 				{
-					UE_VLOG(MyController, LogBehaviorTree, Warning, TEXT("UBTTask_MoveTo::ExecuteTask tried to go to actor while BB %s entry was empty"), *BlackboardKey.SelectedKeyName.ToString());
+					NodeResult = EBTNodeResult::Succeeded;
 				}
 			}
-			else if (BlackboardKey.SelectedKeyType == UBlackboardKeyType_Vector::StaticClass())
-			{
-				const FVector TargetLocation = MyBlackboard->GetValue<UBlackboardKeyType_Vector>(BlackboardKey.GetSelectedKeyID());
-				MoveReq.SetGoalLocation(TargetLocation);
-				MyMemory->PreviousGoalLocation = TargetLocation;
-
-				RequestResult = MyController->MoveTo(MoveReq);
-			}
-
-			if (RequestResult == EPathFollowingRequestResult::RequestSuccessful)
-			{
-				const FAIRequestID RequestID = MyController->GetCurrentMoveRequestID();
-
-				MyMemory->MoveRequestID = RequestID;
-				WaitForMessage(OwnerComp, UBrainComponent::AIMessage_MoveFinished, RequestID);
-				WaitForMessage(OwnerComp, UBrainComponent::AIMessage_RepathFailed);
-
-				NodeResult = EBTNodeResult::InProgress;
-			}
-			else if (RequestResult == EPathFollowingRequestResult::AlreadyAtGoal)
-			{
-				NodeResult = EBTNodeResult::Succeeded;
-			}
 		}
 	}
 
 	return NodeResult;
 }
 
+UAITask_MoveTo* UBTTask_MoveTo::PrepareMoveTask(UBehaviorTreeComponent& OwnerComp, UAITask_MoveTo* ExistingTask, const FAIMoveRequest& MoveRequest)
+{
+	UAITask_MoveTo* MoveTask = ExistingTask ? ExistingTask : NewBTAITask<UAITask_MoveTo>(OwnerComp);
+	if (MoveTask)
+	{
+		MoveTask->SetUp(MoveTask->GetAIController(), MoveRequest);
+	}
+
+	return MoveTask;
+}
+
 EBlackboardNotificationResult UBTTask_MoveTo::OnBlackboardValueChange(const UBlackboardComponent& Blackboard, FBlackboard::FKey ChangedKeyID)
 {
 	UBehaviorTreeComponent* BehaviorComp = Cast<UBehaviorTreeComponent>(Blackboard.GetBrainComponent());
@@ -224,19 +213,19 @@ EBlackboardNotificationResult UBTTask_MoveTo::OnBlackboardValueChange(const UBla
 		{
 			const FVector TargetLocation = Blackboard.GetValue<UBlackboardKeyType_Vector>(BlackboardKey.GetSelectedKeyID());
 
-			bUpdateMove = (FVector::DistSquared(TargetLocation, MyMemory->PreviousGoalLocation) > ObservedBlackboardValueTolerance*ObservedBlackboardValueTolerance);
+			bUpdateMove = (FVector::DistSquared(TargetLocation, MyMemory->PreviousGoalLocation) > FMath::Square(ObservedBlackboardValueTolerance));
 		}
 
 		if (bUpdateMove)
 		{
 			// don't abort move if using AI tasks - it will mess things up
-			if (GET_AI_CONFIG_VAR(bEnableBTAITasks) == false)
+			if (MyMemory->MoveRequestID.IsValid())
 			{
 				StopWaitingForMessages(*BehaviorComp);
 				BehaviorComp->GetAIOwner()->GetPathFollowingComponent()->AbortMove(TEXT("Updating move due to BB value change"), MyMemory->MoveRequestID, /*bResetVelocity=*/false, /*bSilent=*/true);
 			}
 
-			if (BehaviorComp->GetAIOwner()->ShouldPostponePathUpdates())
+			if (!bUseGameplayTasks && BehaviorComp->GetAIOwner()->ShouldPostponePathUpdates())
 			{
 				// NodeTick will take care of requesting move
 				MyMemory->bWaitingForPath = true;
@@ -244,7 +233,6 @@ EBlackboardNotificationResult UBTTask_MoveTo::OnBlackboardValueChange(const UBla
 			else
 			{
 				const EBTNodeResult::Type NodeResult = PerformMoveTask(*BehaviorComp, RawMemory);
-
 				if (NodeResult != EBTNodeResult::InProgress)
 				{
 					FinishLatentTask(*BehaviorComp, NodeResult);
@@ -261,11 +249,26 @@ EBTNodeResult::Type UBTTask_MoveTo::AbortTask(UBehaviorTreeComponent& OwnerComp,
 	FBTMoveToTaskMemory* MyMemory = reinterpret_cast<FBTMoveToTaskMemory*>(NodeMemory);
 	if (!MyMemory->bWaitingForPath)
 	{
-		AAIController* MyController = OwnerComp.GetAIOwner();
-
-		if (MyController && MyController->GetPathFollowingComponent())
+		if (MyMemory->MoveRequestID.IsValid())
 		{
-			MyController->GetPathFollowingComponent()->AbortMove(TEXT("BehaviorTree abort"), MyMemory->MoveRequestID);
+			AAIController* MyController = OwnerComp.GetAIOwner();
+			if (MyController && MyController->GetPathFollowingComponent())
+			{
+				MyController->GetPathFollowingComponent()->AbortMove(TEXT("BehaviorTree abort"), MyMemory->MoveRequestID);
+			}
+		}
+		else
+		{
+			MyMemory->bObserverCanFinishTask = false;
+			UAITask_MoveTo* MoveTask = MyMemory->Task.Get();
+			if (MoveTask)
+			{
+				MoveTask->ExternalCancel();
+			}
+			else
+			{
+				UE_VLOG(&OwnerComp, LogBehaviorTree, Error, TEXT("Can't abort path following! bWaitingForPath:false, MoveRequestID:invalid, MoveTask:none!"));
+			}
 		}
 	}
 
@@ -275,6 +278,7 @@ EBTNodeResult::Type UBTTask_MoveTo::AbortTask(UBehaviorTreeComponent& OwnerComp,
 void UBTTask_MoveTo::OnTaskFinished(UBehaviorTreeComponent& OwnerComp, uint8* NodeMemory, EBTNodeResult::Type TaskResult)
 {
 	FBTMoveToTaskMemory* MyMemory = reinterpret_cast<FBTMoveToTaskMemory*>(NodeMemory);
+	MyMemory->Task.Reset();
 
 	if (bObserveBlackboardValue)
 	{
@@ -283,12 +287,8 @@ void UBTTask_MoveTo::OnTaskFinished(UBehaviorTreeComponent& OwnerComp, uint8* No
 		{
 			BlackboardComp->UnregisterObserver(BlackboardKey.GetSelectedKeyID(), MyMemory->BBObserverDelegateHandle);
 		}
-		MyMemory->BBObserverDelegateHandle.Reset();
-	}
 
-	if (GET_AI_CONFIG_VAR(bEnableBTAITasks))
-	{
-		MyMemory->Task = nullptr;
+		MyMemory->BBObserverDelegateHandle.Reset();
 	}
 
 	Super::OnTaskFinished(OwnerComp, NodeMemory, TaskResult);
@@ -321,6 +321,27 @@ void UBTTask_MoveTo::OnMessage(UBehaviorTreeComponent& OwnerComp, uint8* NodeMem
 	Super::OnMessage(OwnerComp, NodeMemory, Message, SenderID, bSuccess);
 }
 
+void UBTTask_MoveTo::OnGameplayTaskDeactivated(UGameplayTask& Task)
+{
+	// AI move task finished
+	UAITask_MoveTo* MoveTask = Cast<UAITask_MoveTo>(&Task);
+	if (MoveTask && MoveTask->GetAIController() && MoveTask->GetState() != EGameplayTaskState::Paused)
+	{
+		UBehaviorTreeComponent* BehaviorComp = GetBTComponentForTask(Task);
+		if (BehaviorComp)
+		{
+			uint8* RawMemory = BehaviorComp->GetNodeMemory(this, BehaviorComp->FindInstanceContainingNode(this));
+			FBTMoveToTaskMemory* MyMemory = reinterpret_cast<FBTMoveToTaskMemory*>(RawMemory);
+
+			if (MyMemory->bObserverCanFinishTask && (MoveTask == MyMemory->Task))
+			{
+				const bool bSuccess = MoveTask->WasMoveSuccessful();
+				FinishLatentTask(*BehaviorComp, bSuccess ? EBTNodeResult::Succeeded : EBTNodeResult::Failed);
+			}
+		}
+	}
+}
+
 FString UBTTask_MoveTo::GetStaticDescription() const
 {
 	FString KeyDesc("invalid");
@@ -338,12 +359,20 @@ void UBTTask_MoveTo::DescribeRuntimeValues(const UBehaviorTreeComponent& OwnerCo
 	Super::DescribeRuntimeValues(OwnerComp, NodeMemory, Verbosity, Values);
 
 	const UBlackboardComponent* BlackboardComp = OwnerComp.GetBlackboardComponent();
-	FBTMoveToTaskMemory* MyMemory = (FBTMoveToTaskMemory*)NodeMemory;
 
-	if (MyMemory->MoveRequestID && BlackboardComp)
+	if (BlackboardComp)
 	{
-		FString KeyValue = BlackboardComp->DescribeKeyValue(BlackboardKey.GetSelectedKeyID(), EBlackboardDescription::OnlyValue);
-		Values.Add(FString::Printf(TEXT("move target: %s"), *KeyValue));
+		const FString KeyValue = BlackboardComp->DescribeKeyValue(BlackboardKey.GetSelectedKeyID(), EBlackboardDescription::OnlyValue);
+
+		FBTMoveToTaskMemory* MyMemory = (FBTMoveToTaskMemory*)NodeMemory;
+		const bool bIsUsingTask = MyMemory->Task.IsValid();
+		
+		const FString ModeDesc =
+			MyMemory->bWaitingForPath ? TEXT("(WAITING)") :
+			bIsUsingTask ? TEXT("(task)") :
+			TEXT("");
+
+		Values.Add(FString::Printf(TEXT("move target: %s%s"), *KeyValue, *ModeDesc));
 	}
 }
 
diff --git a/Engine/Source/Runtime/AIModule/Private/BehaviorTree/Tasks/BTTask_RunBehavior.cpp b/Engine/Source/Runtime/AIModule/Private/BehaviorTree/Tasks/BTTask_RunBehavior.cpp
--- a/Engine/Source/Runtime/AIModule/Private/BehaviorTree/Tasks/BTTask_RunBehavior.cpp
+++ b/Engine/Source/Runtime/AIModule/Private/BehaviorTree/Tasks/BTTask_RunBehavior.cpp
@@ -173,7 +173,7 @@ void UBTTask_RunBehavior::InjectNodes(UBehaviorTreeComponent& OwnerComp, uint8*
 			{
 				UBTDecorator* InstancedOb = Cast<UBTDecorator>(OwnerComp.NodeInstances[FirstNodeIdx + Idx]);
 				InstancedOb->InitializeFromAsset(*BehaviorAsset);
-				InstancedOb->InitializeDecorator(ChildIdx);
+				InstancedOb->InitializeParentLink(ChildIdx);
 
 				if (!bAlreadyInjected)
 				{
diff --git a/Engine/Source/Runtime/AIModule/Private/Navigation/PathFollowingComponent.cpp b/Engine/Source/Runtime/AIModule/Private/Navigation/PathFollowingComponent.cpp
--- a/Engine/Source/Runtime/AIModule/Private/Navigation/PathFollowingComponent.cpp
+++ b/Engine/Source/Runtime/AIModule/Private/Navigation/PathFollowingComponent.cpp
@@ -133,27 +133,66 @@ FString GetPathDescHelper(FNavPathSharedPtr Path)
 		FString::Printf(TEXT("%s:%d"), Path->IsPartial() ? TEXT("partial") : TEXT("complete"), Path->GetPathPoints().Num());
 }
 
-void UPathFollowingComponent::OnPathEvent(FNavigationPath* InvalidatedPath, ENavPathEvent::Type Event)
+void UPathFollowingComponent::OnPathEvent(FNavigationPath* InPath, ENavPathEvent::Type Event)
 {
 	const static UEnum* NavPathEventEnum = FindObject<UEnum>(ANY_PACKAGE, TEXT("ENavPathEvent"));
 	UE_VLOG(GetOwner(), LogPathFollowing, Log, TEXT("OnPathEvent: %s"), *NavPathEventEnum->GetEnumName(Event));
 
-	if (InvalidatedPath == nullptr || Path.Get() != InvalidatedPath)
-	{
-		return;
-	}
-
-	switch (Event)
+	if (InPath && Path.Get() == InPath)
 	{
+		switch (Event)
+		{
 		case ENavPathEvent::UpdatedDueToGoalMoved:
 		case ENavPathEvent::UpdatedDueToNavigationChanged:
-		{
-			UpdateMove(Path, GetCurrentRequestId());
+			{
+				const bool bIsPathAccepted = HandlePathUpdateEvent();
+				if (!bIsPathAccepted)
+				{
+					AbortMove(TEXT("Failed path update"));
+				}
+			}
+			break;
+
+		default:
+			break;
 		}
-		break;
 	}
 }
 
+bool UPathFollowingComponent::HandlePathUpdateEvent()
+{
+	UE_VLOG(GetOwner(), LogPathFollowing, Log, TEXT("HandlePathUpdateEvent Path(%s) Status(%s)"), *GetPathDescHelper(Path), *GetStatusDesc());
+	LogPathHelper(GetOwner(), Path, DestinationActor.Get());
+
+	if (Status == EPathFollowingStatus::Waiting && Path.IsValid() && !Path->IsValid())
+	{
+		UE_VLOG(GetOwner(), LogPathFollowing, Error, TEXT("Received unusable path in Waiting state! (ready:%d upToDate:%d pathPoints:%d)"),
+			Path->IsReady() ? 1 : 0,
+			Path->IsUpToDate() ? 1 : 0,
+			Path->GetPathPoints().Num());
+	}
+
+	if (Status == EPathFollowingStatus::Idle || !Path.IsValid() || !Path->IsValid())
+	{
+		UE_VLOG(GetOwner(), LogPathFollowing, Log, TEXT(">> invalid request"));
+		return false;
+	}
+
+	OnPathUpdated();
+	GetWorld()->GetTimerManager().ClearTimer(WaitingForPathTimer);
+	UpdateDecelerationData();
+
+	if (Status == EPathFollowingStatus::Waiting || Status == EPathFollowingStatus::Moving)
+	{
+		Status = EPathFollowingStatus::Moving;
+
+		const int32 CurrentSegment = DetermineStartingPathPoint(Path.Get());
+		SetMoveSegment(CurrentSegment);
+	}
+
+	return true;
+}
+
 FAIRequestID UPathFollowingComponent::RequestMove(FNavPathSharedPtr InPath, FRequestCompletedSignature OnComplete,
 	const AActor* InDestinationActor, float InAcceptanceRadius, bool bInStopOnOverlap, FCustomMoveSharedPtr InGameData)
 {
@@ -236,7 +275,7 @@ FAIRequestID UPathFollowingComponent::RequestMove(FNavPathSharedPtr InPath, FReq
 			CurrentRequestId.GetID(), ToDest.Size2D(), FMath::Abs(ToDest.Z));
 	#endif // ENABLE_VISUAL_LOG
 
-		// with async pathfinding paths can be incomplete, movement will start after receiving UpdateMove 
+		// with async pathfinding paths can be incomplete, movement will start after receiving path event 
 		if (Path->IsValid())
 		{
 			Status = EPathFollowingStatus::Moving;
@@ -257,41 +296,16 @@ FAIRequestID UPathFollowingComponent::RequestMove(FNavPathSharedPtr InPath, FReq
 
 bool UPathFollowingComponent::UpdateMove(FNavPathSharedPtr InPath, FAIRequestID RequestID)
 {
-	UE_VLOG(GetOwner(), LogPathFollowing, Log, TEXT("UpdateMove: Path(%s) Status(%s) RequestID(%u)"),
-		*GetPathDescHelper(InPath),
-		*GetStatusDesc(), RequestID);
-
-	LogPathHelper(GetOwner(), InPath, DestinationActor.Get());
-
-	if (Status == EPathFollowingStatus::Waiting && InPath.IsValid() && !InPath->IsValid())
+	if (InPath == Path && RequestID.IsEquivalent(CurrentRequestId))
 	{
-		UE_VLOG(GetOwner(), LogPathFollowing, Error, TEXT("Received unusable path in Waiting state! (ready:%d upToDate:%d pathPoints:%d)"),
-			InPath->IsReady() ? 1 : 0,
-			InPath->IsUpToDate() ? 1 : 0,
-			InPath->GetPathPoints().Num());
+		return HandlePathUpdateEvent();
 	}
 
-	if (!InPath.IsValid() || !InPath->IsValid() || Status == EPathFollowingStatus::Idle 
-		|| RequestID.IsEquivalent(GetCurrentRequestId()) == false)
-	{
-		UE_VLOG(GetOwner(), LogPathFollowing, Log, TEXT("UpdateMove: invalid request"));
-		return false;
-	}
-
-	Path = InPath;
-	OnPathUpdated();
-	GetWorld()->GetTimerManager().ClearTimer(WaitingForPathTimer);
-	UpdateDecelerationData();
-
-	if (Status == EPathFollowingStatus::Waiting || Status == EPathFollowingStatus::Moving)
-	{
-		Status = EPathFollowingStatus::Moving;
+	UE_VLOG(GetOwner(), LogPathFollowing, Error, TEXT("UpdateMove is ignored! path:%s requestId:%s"),
+		(InPath == Path) ? TEXT("ok") : TEXT("DIFFERENT"),
+		RequestID.IsEquivalent(CurrentRequestId) ? TEXT("ok") : TEXT("NOT CURRENT"));
 
-		const int32 CurrentSegment = DetermineStartingPathPoint(InPath.Get());
-		SetMoveSegment(CurrentSegment);
-	}
-
-	return true;
+	return false;
 }
 
 void UPathFollowingComponent::AbortMove(const FString& Reason, FAIRequestID RequestID, bool bResetVelocity, bool bSilent, uint8 MessageFlags)
@@ -1166,9 +1180,7 @@ bool UPathFollowingComponent::UpdateMovementComponent(bool bForce)
 
 		if (MyPawn)
 		{
-			FScriptDelegate Delegate;
-			Delegate.BindUFunction(this, "OnActorBump");
-			MyPawn->OnActorHit.AddUnique(Delegate);
+			MyPawn->OnActorHit.AddUniqueDynamic(this, &UPathFollowingComponent::OnActorBump);
 
 			SetMovementComponent(MyPawn->FindComponentByClass<UNavMovementComponent>());
 		}
diff --git a/Engine/Source/Runtime/AIModule/Private/Perception/AIPerceptionComponent.cpp b/Engine/Source/Runtime/AIModule/Private/Perception/AIPerceptionComponent.cpp
--- a/Engine/Source/Runtime/AIModule/Private/Perception/AIPerceptionComponent.cpp
+++ b/Engine/Source/Runtime/AIModule/Private/Perception/AIPerceptionComponent.cpp
@@ -202,7 +202,7 @@ void UAIPerceptionComponent::OnUnregister()
 	Super::OnUnregister();
 }
 
-void UAIPerceptionComponent::OnOwnerEndPlay(EEndPlayReason::Type EndPlayReason)
+void UAIPerceptionComponent::OnOwnerEndPlay(AActor* Actor, EEndPlayReason::Type EndPlayReason)
 {
 	if (EndPlayReason != EEndPlayReason::EndPlayInEditor && EndPlayReason != EEndPlayReason::Quit)
 	{
diff --git a/Engine/Source/Runtime/AIModule/Private/Perception/AIPerceptionSystem.cpp b/Engine/Source/Runtime/AIModule/Private/Perception/AIPerceptionSystem.cpp
--- a/Engine/Source/Runtime/AIModule/Private/Perception/AIPerceptionSystem.cpp
+++ b/Engine/Source/Runtime/AIModule/Private/Perception/AIPerceptionSystem.cpp
@@ -592,7 +592,7 @@ void UAIPerceptionSystem::RegisterSourceForSenseClass(TSubclassOf<UAISense> Sens
 	RegisterSource(SenseID, Target);
 }
 
-void UAIPerceptionSystem::OnPerceptionStimuliSourceEndPlay(EEndPlayReason::Type EndPlayReason)
+void UAIPerceptionSystem::OnPerceptionStimuliSourceEndPlay(AActor* Actor, EEndPlayReason::Type EndPlayReason)
 {
 	// this tells us just that _a_ source has been removed. We need to parse through all sources and find which one was it
 	// this is a fall-back behavior, if source gets unregistered manually this function won't get called 
diff --git a/Engine/Source/Runtime/AIModule/Private/Tasks/AITask.cpp b/Engine/Source/Runtime/AIModule/Private/Tasks/AITask.cpp
--- a/Engine/Source/Runtime/AIModule/Private/Tasks/AITask.cpp
+++ b/Engine/Source/Runtime/AIModule/Private/Tasks/AITask.cpp
@@ -55,6 +55,11 @@ void UAITask::InitAITask(AAIController& AIOwner, IGameplayTaskOwnerInterface& In
 	}
 }
 
+void UAITask::InitAITask(AAIController& AIOwner, IGameplayTaskOwnerInterface& InTaskOwner)
+{
+	InitAITask(AIOwner, InTaskOwner, InTaskOwner.GetGameplayTaskDefaultPriority());
+}
+
 void UAITask::RequestAILogicLocking()
 {
 	AddClaimedResource<UAIResource_Logic>();
diff --git a/Engine/Source/Runtime/AIModule/Private/Tasks/AITask_LockLogic.cpp b/Engine/Source/Runtime/AIModule/Private/Tasks/AITask_LockLogic.cpp
new file mode 100644
--- /dev/null
+++ b/Engine/Source/Runtime/AIModule/Private/Tasks/AITask_LockLogic.cpp
@@ -0,0 +1,8 @@
+// Copyright 1998-2016 Epic Games, Inc. All Rights Reserved.
+
+#include "AIModulePrivate.h"
+#include "Tasks/AITask_LockLogic.h"
+
+UAITask_LockLogic::UAITask_LockLogic(const FObjectInitializer& ObjectInitializer) : Super(ObjectInitializer)
+{
+}
diff --git a/Engine/Source/Runtime/AIModule/Private/Tasks/AITask_MoveTo.cpp b/Engine/Source/Runtime/AIModule/Private/Tasks/AITask_MoveTo.cpp
--- a/Engine/Source/Runtime/AIModule/Private/Tasks/AITask_MoveTo.cpp
+++ b/Engine/Source/Runtime/AIModule/Private/Tasks/AITask_MoveTo.cpp
@@ -8,29 +8,39 @@ UAITask_MoveTo::UAITask_MoveTo(const FObjectInitializer& ObjectInitializer)
 	: Super(ObjectInitializer)
 {
 	bIsPausable = true;
+	MoveRequestID = FAIRequestID::InvalidRequest;
 
-	RealGoalLocation = FAISystem::InvalidLocation;
-	MoveGoalActor = nullptr;
-	MoveAcceptanceRadius = GET_AI_CONFIG_VAR(AcceptanceRadius);
-	bShouldStopOnOverlap = GET_AI_CONFIG_VAR(bFinishMoveOnGoalOverlap);
-	bShouldAcceptPartialPath = GET_AI_CONFIG_VAR(bAcceptPartialPaths);
-	bShouldUsePathfinding = true;
+	MoveRequest.SetAcceptanceRadius(GET_AI_CONFIG_VAR(AcceptanceRadius));
+	MoveRequest.SetStopOnOverlap(GET_AI_CONFIG_VAR(bFinishMoveOnGoalOverlap));
+	MoveRequest.SetAllowPartialPath(GET_AI_CONFIG_VAR(bAcceptPartialPaths));
+	MoveRequest.SetUsePathfinding(true);
 
 	AddRequiredResource(UAIResource_Movement::StaticClass());
+	AddClaimedResource(UAIResource_Movement::StaticClass());
 }
 
 UAITask_MoveTo* UAITask_MoveTo::AIMoveTo(AAIController* Controller, FVector InGoalLocation, AActor* InGoalActor, float AcceptanceRadius, EAIOptionFlag::Type StopOnOverlap, EAIOptionFlag::Type AcceptPartialPath, bool bUsePathfinding, bool bLockAILogic)
 {
-	if (Controller == nullptr)
-	{
-		return nullptr; 
-	}
-
-	UAITask_MoveTo* MyTask = NewTask<UAITask_MoveTo>(*static_cast<IGameplayTaskOwnerInterface*>(Controller));
+	UAITask_MoveTo* MyTask = Controller ? UAITask::NewAITask<UAITask_MoveTo>(*Controller, EAITaskPriority::High) : nullptr;
 	if (MyTask)
 	{
-		MyTask->SetUp(Controller, InGoalLocation, InGoalActor, AcceptanceRadius, bUsePathfinding, StopOnOverlap, AcceptPartialPath);
-		MyTask->Priority = uint8(EAITaskPriority::High);
+		FAIMoveRequest MoveReq;
+		if (InGoalActor)
+		{
+			MoveReq.SetGoalActor(InGoalActor);
+		}
+		else
+		{
+			MoveReq.SetGoalLocation(InGoalLocation);
+		}
+
+		MoveReq.SetAcceptanceRadius(AcceptanceRadius);
+		MoveReq.SetStopOnOverlap(FAISystem::PickAIOption(StopOnOverlap, MoveReq.CanStopOnOverlap()));
+		MoveReq.SetAllowPartialPath(FAISystem::PickAIOption(AcceptPartialPath, MoveReq.IsUsingPartialPaths()));
+		MoveReq.SetUsePathfinding(bUsePathfinding);
+
+		MyTask->SetUp(Controller, MoveReq);
+
 		if (bLockAILogic)
 		{
 			MyTask->RequestAILogicLocking();
@@ -40,110 +50,94 @@ UAITask_MoveTo* UAITask_MoveTo::AIMoveTo(AAIController* Controller, FVector InGo
 	return MyTask;
 }
 
-void UAITask_MoveTo::SetUp(AAIController* Controller, FVector InGoalLocation, AActor* InGoalActor, float AcceptanceRadius, bool bUsePathfinding, EAIOptionFlag::Type StopOnOverlap, EAIOptionFlag::Type AcceptPartialPath)
+void UAITask_MoveTo::SetUp(AAIController* Controller, const FAIMoveRequest& InMoveRequest)
 {
 	OwnerController = Controller;
+	MoveRequest = InMoveRequest;
+}
 
-	if (InGoalActor != nullptr)
-	{
-		MoveGoalActor = InGoalActor;
-	}
-	else
+void UAITask_MoveTo::FinishMoveTask(EPathFollowingResult::Type InResult)
+{
+	if (MoveRequestID.IsValid())
 	{
-		RealGoalLocation = InGoalLocation;
+		UPathFollowingComponent* PFComp = OwnerController ? OwnerController->GetPathFollowingComponent() : nullptr;
+		if (PFComp && PFComp->GetStatus() != EPathFollowingStatus::Idle)
+		{
+			ResetObservers();
+			PFComp->AbortMove(TEXT("AITask_MoveTo finished"), MoveRequestID);
+		}
 	}
 
-	if (AcceptanceRadius >= 0)
+	MoveResult = InResult;
+	EndTask();
+
+	if (InResult == EPathFollowingResult::Invalid)
 	{
-		MoveAcceptanceRadius = AcceptanceRadius;
+		OnRequestFailed.Broadcast();
 	}
-	
-	bShouldStopOnOverlap = FAISystem::PickAIOption(StopOnOverlap, bShouldStopOnOverlap);
-	bShouldUsePathfinding = bUsePathfinding;
-
-	if (AcceptPartialPath != EAIOptionFlag::Default)
+	else
 	{
-		bShouldAcceptPartialPath = (AcceptPartialPath == EAIOptionFlag::Enable);
+		OnMoveFinished.Broadcast(InResult);
 	}
 }
 
-void UAITask_MoveTo::PostInitProperties()
+void UAITask_MoveTo::Activate()
 {
-	Super::PostInitProperties();
+	Super::Activate();
 
-	if (MoveGoalActor == nullptr)
-	{
-		RealGoalLocation = MoveGoalLocation;
-	}
+	MoveRequestID = FAIRequestID::InvalidRequest;
+	ConditionalPerformMove();
 }
 
-void UAITask_MoveTo::HandleMoveFinished(FAIRequestID RequestID, EPathFollowingResult::Type Result)
+void UAITask_MoveTo::ConditionalPerformMove()
 {
-	if (RequestID == MoveRequestID)
+	if (MoveRequest.IsUsingPathfinding() && OwnerController && OwnerController->ShouldPostponePathUpdates())
 	{
-		EndTask();
-		OnMoveFinished.Broadcast(Result);
+		UE_VLOG(GetGameplayTasksComponent(), LogGameplayTasks, Log, TEXT("%s> can't path right now, waiting..."), *GetName());
+		OwnerController->GetWorldTimerManager().SetTimer(MoveRetryTimerHandle, this, &UAITask_MoveTo::ConditionalPerformMove, 0.2f, false);
 	}
 	else
 	{
-		// @todo report issue to the owner component
-		UE_VLOG(GetGameplayTasksComponent(), LogGameplayTasks, Warning, TEXT("%s got movement-finished-notification but RequestID doesn't match. Possible task leak"), *GetName());
+		MoveRetryTimerHandle.Invalidate();
+		PerformMove();
 	}
 }
 
-void UAITask_MoveTo::Activate()
-{
-	Super::Activate();
-
-	PerformMove();
-}
-
 void UAITask_MoveTo::PerformMove()
 {
-	FAIMoveRequest MoveRequest = (MoveGoalActor != nullptr) ? FAIMoveRequest(MoveGoalActor) : FAIMoveRequest(RealGoalLocation);
-
-	MoveRequest.SetAllowPartialPath(bShouldAcceptPartialPath)
-		.SetAcceptanceRadius(MoveAcceptanceRadius)
-		.SetStopOnOverlap(bShouldStopOnOverlap)
-		.SetUsePathfinding(bShouldUsePathfinding);
-
-	if (PathFollowingDelegateHandle.IsValid())
+	UPathFollowingComponent* PFComp = OwnerController ? OwnerController->GetPathFollowingComponent() : nullptr;
+	if (PFComp == nullptr)
 	{
-		OwnerController->GetPathFollowingComponent()->OnMoveFinished.Remove(PathFollowingDelegateHandle);
+		FinishMoveTask(EPathFollowingResult::Invalid);
+		return;
 	}
 
-	const EPathFollowingRequestResult::Type RequestResult = OwnerController->MoveTo(MoveRequest);
+	ResetObservers();
+	ResetTimers();
 
+	// start new move request
+	const EPathFollowingRequestResult::Type RequestResult = OwnerController->MoveTo(MoveRequest);
 	switch (RequestResult)
 	{
 	case EPathFollowingRequestResult::Failed:
-		{
-			EndTask();
-			OnRequestFailed.Broadcast();
-		}
+		FinishMoveTask(EPathFollowingResult::Invalid);
 		break;
+
 	case EPathFollowingRequestResult::AlreadyAtGoal:
-		{
-			EndTask();
-			OnMoveFinished.Broadcast(EPathFollowingResult::Success);
-		}
+		FinishMoveTask(EPathFollowingResult::Success);
 		break;
+
 	case EPathFollowingRequestResult::RequestSuccessful:
-		if (OwnerController->GetPathFollowingComponent())
+		MoveRequestID = PFComp->GetCurrentRequestId();
+		PathFinishDelegateHandle = PFComp->OnMoveFinished.AddUObject(this, &UAITask_MoveTo::OnRequestFinished);
+		SetObservedPath(PFComp->GetPath());
+
+		if (TaskState == EGameplayTaskState::Finished)
 		{
-			MoveRequestID = OwnerController->GetPathFollowingComponent()->GetCurrentRequestId();
-
-			if (TaskState == EGameplayTaskState::Finished)
-			{
-				UE_VLOG(GetGameplayTasksComponent(), LogGameplayTasks, Error, TEXT("%s re-Activating Finished task!")
-					, *GetName());
-			}
-			else if (PathFollowingDelegateHandle.IsValid() == false)
-			{
-				PathFollowingDelegateHandle = OwnerController->GetPathFollowingComponent()->OnMoveFinished.AddUObject(this, &UAITask_MoveTo::HandleMoveFinished);
-			}
+			UE_VLOG(GetGameplayTasksComponent(), LogGameplayTasks, Error, TEXT("%s> re-Activating Finished task!"), *GetName());
 		}
 		break;
+
 	default:
 		checkNoEntry();
 		break;
@@ -152,27 +146,182 @@ void UAITask_MoveTo::PerformMove()
 
 void UAITask_MoveTo::Pause()
 {
-	if (OwnerController)
+	if (OwnerController && MoveRequestID.IsValid())
 	{
-		OwnerController->GetPathFollowingComponent()->OnMoveFinished.Remove(PathFollowingDelegateHandle);
-		PathFollowingDelegateHandle.Reset();
 		OwnerController->PauseMove(MoveRequestID);
 	}
+
+	ResetTimers();
 	Super::Pause();
 }
 
 void UAITask_MoveTo::Resume()
 {
-	PerformMove();
 	Super::Resume();
+
+	if (!MoveRequestID.IsValid() || !OwnerController->ResumeMove(MoveRequestID))
+	{
+		UE_CVLOG(MoveRequestID.IsValid(), GetGameplayTasksComponent(), LogGameplayTasks, Log, TEXT("%s> Resume move failed, starting new one."), *GetName());
+		ConditionalPerformMove();
+	}
+}
+
+void UAITask_MoveTo::SetObservedPath(FNavPathSharedPtr InPath)
+{
+	if (PathUpdateDelegateHandle.IsValid() && Path.IsValid())
+	{
+		Path->RemoveObserver(PathUpdateDelegateHandle);
+	}
+
+	PathUpdateDelegateHandle.Reset();
+	
+	Path = InPath;
+	if (Path.IsValid())
+	{
+		// disable auto repaths, it will be handled by move task to include ShouldPostponePathUpdates condition
+		Path->EnableRecalculationOnInvalidation(false);
+		PathUpdateDelegateHandle = Path->AddObserver(FNavigationPath::FPathObserverDelegate::FDelegate::CreateUObject(this, &UAITask_MoveTo::OnPathEvent));
+	}
+}
+
+void UAITask_MoveTo::ResetObservers()
+{
+	if (PathFinishDelegateHandle.IsValid())
+	{
+		UPathFollowingComponent* PFComp = OwnerController ? OwnerController->GetPathFollowingComponent() : nullptr;
+		if (PFComp)
+		{
+			PFComp->OnMoveFinished.Remove(PathFinishDelegateHandle);
+		}
+
+		PathFinishDelegateHandle.Reset();
+	}
+
+	if (PathUpdateDelegateHandle.IsValid())
+	{
+		if (Path.IsValid())
+		{
+			Path->RemoveObserver(PathUpdateDelegateHandle);
+		}
+
+		PathUpdateDelegateHandle.Reset();
+	}
+}
+
+void UAITask_MoveTo::ResetTimers()
+{
+	if (MoveRetryTimerHandle.IsValid())
+	{
+		if (OwnerController)
+		{
+			OwnerController->GetWorldTimerManager().ClearTimer(MoveRetryTimerHandle);
+		}
+
+		MoveRetryTimerHandle.Invalidate();
+	}
+
+	if (PathRetryTimerHandle.IsValid())
+	{
+		if (OwnerController)
+		{
+			OwnerController->GetWorldTimerManager().ClearTimer(PathRetryTimerHandle);
+		}
+
+		PathRetryTimerHandle.Invalidate();
+	}
 }
 
 void UAITask_MoveTo::OnDestroy(bool bOwnerFinished)
 {
 	Super::OnDestroy(bOwnerFinished);
-	if (OwnerController && OwnerController->GetPathFollowingComponent())
+	
+	ResetObservers();
+	ResetTimers();
+
+	if (MoveRequestID.IsValid())
+	{
+		UPathFollowingComponent* PFComp = OwnerController ? OwnerController->GetPathFollowingComponent() : nullptr;
+		if (PFComp && PFComp->GetStatus() != EPathFollowingStatus::Idle)
+		{
+			PFComp->AbortMove(TEXT("AITask_MoveTo finished"), MoveRequestID);
+		}
+	}
+}
+
+void UAITask_MoveTo::OnRequestFinished(FAIRequestID RequestID, EPathFollowingResult::Type Result)
+{
+	if (RequestID == MoveRequestID)
+	{
+		// reset request Id, FinishMoveTask doesn't need to update path following's state
+		MoveRequestID = FAIRequestID::InvalidRequest;
+
+		FinishMoveTask(Result);
+	}
+	else if (IsActive())
+	{
+		UE_VLOG(GetGameplayTasksComponent(), LogGameplayTasks, Warning, TEXT("%s> received HandleMoveFinished with not matching RequestID!"), *GetName());
+	}
+}
+
+void UAITask_MoveTo::OnPathEvent(FNavigationPath* InPath, ENavPathEvent::Type Event)
+{
+	const static UEnum* NavPathEventEnum = FindObject<UEnum>(ANY_PACKAGE, TEXT("ENavPathEvent"));
+	UE_VLOG(GetGameplayTasksComponent(), LogGameplayTasks, Log, TEXT("%s> Path event: %s"), *GetName(), *NavPathEventEnum->GetEnumName(Event));
+
+	switch (Event)
+	{
+	case ENavPathEvent::NewPath:
+	case ENavPathEvent::UpdatedDueToGoalMoved:
+	case ENavPathEvent::UpdatedDueToNavigationChanged:
+		if (InPath && InPath->IsPartial() && !MoveRequest.IsUsingPartialPaths())
+		{
+			UE_VLOG(GetGameplayTasksComponent(), LogGameplayTasks, Log, TEXT(">> partial path is not allowed, aborting"));
+			UPathFollowingComponent::LogPathHelper(OwnerController, InPath, MoveRequest.GetGoalActor());
+			FinishMoveTask(EPathFollowingResult::Aborted);
+		}
+#if ENABLE_VISUAL_LOG
+		else if (!IsActive())
+		{
+			UPathFollowingComponent::LogPathHelper(OwnerController, InPath, MoveRequest.GetGoalActor());
+		}
+#endif // ENABLE_VISUAL_LOG
+		break;
+
+	case ENavPathEvent::Invalidated:
+		ConditionalUpdatePath();
+		break;
+
+	case ENavPathEvent::Cleared:
+	case ENavPathEvent::RePathFailed:
+		UE_VLOG(GetGameplayTasksComponent(), LogGameplayTasks, Log, TEXT(">> no path, aborting!"));
+		FinishMoveTask(EPathFollowingResult::Aborted);
+		break;
+
+	default:
+		break;
+	}
+}
+
+void UAITask_MoveTo::ConditionalUpdatePath()
+{
+	if (MoveRequest.IsUsingPathfinding() && OwnerController && OwnerController->ShouldPostponePathUpdates())
+	{
+		UE_VLOG(GetGameplayTasksComponent(), LogGameplayTasks, Log, TEXT("%s> can't path right now, waiting..."), *GetName());
+		OwnerController->GetWorldTimerManager().SetTimer(PathRetryTimerHandle, this, &UAITask_MoveTo::ConditionalUpdatePath, 0.2f, false);
+	}
+	else
 	{
-		OwnerController->GetPathFollowingComponent()->OnMoveFinished.Remove(PathFollowingDelegateHandle);
-		OwnerController->GetPathFollowingComponent()->AbortMove(TEXT("UAITask_MoveTo finishing"), MoveRequestID);
+		PathRetryTimerHandle.Invalidate();
+		
+		ANavigationData* NavData = Path.IsValid() ? Path->GetNavigationDataUsed() : nullptr;
+		if (NavData)
+		{
+			NavData->RequestRePath(Path, ENavPathUpdateType::NavigationChanged);
+		}
+		else
+		{
+			UE_VLOG(GetGameplayTasksComponent(), LogGameplayTasks, Log, TEXT("%s> unable to repath, aborting!"), *GetName());
+			FinishMoveTask(EPathFollowingResult::Aborted);
+		}
 	}
 }
diff --git a/Engine/Source/Runtime/AnimGraphRuntime/Private/BoneControllers/AnimNode_AnimDynamics.cpp b/Engine/Source/Runtime/AnimGraphRuntime/Private/BoneControllers/AnimNode_AnimDynamics.cpp
--- a/Engine/Source/Runtime/AnimGraphRuntime/Private/BoneControllers/AnimNode_AnimDynamics.cpp
+++ b/Engine/Source/Runtime/AnimGraphRuntime/Private/BoneControllers/AnimNode_AnimDynamics.cpp
@@ -18,7 +18,8 @@ TAutoConsoleVariable<int32> CVarEnableWind(TEXT("p.AnimDynamicsWind"), 1, TEXT("
 const float FAnimNode_AnimDynamics::MaxTimeDebt = (1.0f / 60.0f) * 5.0f; // 5 frames max debt
 
 FAnimNode_AnimDynamics::FAnimNode_AnimDynamics()
-: BoxExtents(0.0f)
+: SimulationSpace(AnimPhysSimSpaceType::Component)
+, BoxExtents(0.0f)
 , LocalJointOffset(0.0f)
 , GravityScale(1.0f)
 , bLinearSpring(false)
@@ -58,6 +59,11 @@ void FAnimNode_AnimDynamics::Initialize(const FAnimationInitializeContext& Conte
 		PlanarLimit.DrivingBone.Initialize(RequiredBones);
 	}
 
+	if(SimulationSpace == AnimPhysSimSpaceType::BoneRelative)
+	{
+		RelativeSpaceBone.Initialize(RequiredBones);
+	}
+
 	if(BoundBone.IsValid(RequiredBones))
 	{
 		RequestInitialise();
@@ -134,6 +140,12 @@ void FAnimNode_AnimDynamics::EvaluateBoneTransforms(USkeletalMeshComponent* Skel
 				}
 			}
 
+			FVector OrientedExternalForce = ExternalForce;
+			if(!OrientedExternalForce.IsNearlyZero())
+			{
+				OrientedExternalForce = TransformWorldVectorToSimSpace(SkelComp, MeshBases, OrientedExternalForce);
+			}
+
 			if (CVarEnableAdaptiveSubstep.GetValueOnAnyThread() == 1)
 			{
 				float FixedTimeStep = MaxSubstepDeltaTime * CurrentTimeDilation;
@@ -158,7 +170,7 @@ void FAnimNode_AnimDynamics::EvaluateBoneTransforms(USkeletalMeshComponent* Skel
 				for (int32 Iter = 0; Iter < NumIters; ++Iter)
 				{
 					UpdateLimits(SkelComp, MeshBases);
-					FAnimPhys::PhysicsUpdate(FixedTimeStep, BaseBodyPtrs, LinearLimits, AngularLimits, Springs, NumSolverIterationsPreUpdate, NumSolverIterationsPostUpdate);
+					FAnimPhys::PhysicsUpdate(FixedTimeStep, BaseBodyPtrs, LinearLimits, AngularLimits, Springs, SimSpaceGravityDirection, OrientedExternalForce, NumSolverIterationsPreUpdate, NumSolverIterationsPostUpdate);
 				}
 			}
 			else
@@ -169,7 +181,7 @@ void FAnimNode_AnimDynamics::EvaluateBoneTransforms(USkeletalMeshComponent* Skel
 				NextTimeStep = FMath::Min(NextTimeStep, MaxDeltaTime);
 
 				UpdateLimits(SkelComp, MeshBases);
-				FAnimPhys::PhysicsUpdate(NextTimeStep, BaseBodyPtrs, LinearLimits, AngularLimits, Springs, NumSolverIterationsPreUpdate, NumSolverIterationsPostUpdate);
+				FAnimPhys::PhysicsUpdate(NextTimeStep, BaseBodyPtrs, LinearLimits, AngularLimits, Springs, SimSpaceGravityDirection, OrientedExternalForce, NumSolverIterationsPreUpdate, NumSolverIterationsPostUpdate);
 			}
 		}
 
@@ -193,6 +205,9 @@ void FAnimNode_AnimDynamics::EvaluateBoneTransforms(USkeletalMeshComponent* Skel
 				FCompactPoseBoneIndex BoneIndex = CurrentChainBone.GetCompactPoseIndex(BoneContainer);
 
 				FTransform NewBoneTransform(CurrentBody.Pose.Orientation, CurrentBody.Pose.Position + CurrentBody.Pose.Orientation.RotateVector(JointOffsets[Idx]));
+
+				NewBoneTransform = GetComponentSpaceTransformFromSimSpace(SkelComp, MeshBases, BoneIndex, NewBoneTransform);
+
 				OutBoneTransforms.Add(FBoneTransform(BoneIndex, NewBoneTransform));
 			}
 		}
@@ -338,8 +353,8 @@ void FAnimNode_AnimDynamics::InitPhysics(USkeletalMeshComponent* Component, FCSP
 		// Calculate joint offsets by looking at the length of the bones and extending the provided offset
 		if (BoundBoneReferences.Num() > 0)
 		{
-			FTransform CurrentBoneTransform = MeshBases.GetComponentSpaceTransform(LinkBoneRef.GetCompactPoseIndex(BoneContainer));
-			FTransform PreviousBoneTransform = MeshBases.GetComponentSpaceTransform(BoundBoneReferences.Last().GetCompactPoseIndex(BoneContainer));
+			FTransform CurrentBoneTransform = GetBoneTransformInSimSpace(Component, MeshBases, LinkBoneRef.GetCompactPoseIndex(BoneContainer));
+			FTransform PreviousBoneTransform = GetBoneTransformInSimSpace(Component, MeshBases, BoundBoneReferences.Last().GetCompactPoseIndex(BoneContainer));
 
 			FVector PreviousAnchor = PreviousBoneTransform.TransformPosition(-LocalJointOffset);
 			float DistanceToAnchor = (PreviousBoneTransform.GetTranslation() - CurrentBoneTransform.GetTranslation()).Size() * 0.5f;
@@ -365,7 +380,8 @@ void FAnimNode_AnimDynamics::InitPhysics(USkeletalMeshComponent* Component, FCSP
 
 		BoundBoneReferences.Add(LinkBoneRef);
 
-		FTransform BodyTransform = MeshBases.GetComponentSpaceTransform(LinkBoneRef.GetCompactPoseIndex(BoneContainer));
+		FTransform BodyTransform = GetBoneTransformInSimSpace(Component, MeshBases, LinkBoneRef.GetCompactPoseIndex(BoneContainer));
+
 		BodyTransform.SetTranslation(BodyTransform.GetTranslation() + BodyTransform.GetRotation().RotateVector(-LocalJointOffset));
 
 		FAnimPhysLinkedBody NewChainBody(BodyShapes, BodyTransform.GetTranslation(), LinkBoneRef);
@@ -440,6 +456,8 @@ void FAnimNode_AnimDynamics::InitPhysics(USkeletalMeshComponent* Component, FCSP
 		MaxSubsteps = 4;
 	}
 
+	SimSpaceGravityDirection = TransformWorldVectorToSimSpace(Component, MeshBases, FVector(0.0f, 0.0f, -1.0f));
+
 	bRequiresInit = false;
 }
 
@@ -487,7 +505,8 @@ void FAnimNode_AnimDynamics::UpdateLimits(USkeletalMeshComponent* SkelComp, FCSP
 
 		// Get joint transform
 		FCompactPoseBoneIndex BoneIndex = CurrentBoneRef.GetCompactPoseIndex(BoneContainer);
-		FTransform BoundBoneTransform = MeshBases.GetComponentSpaceTransform(BoneIndex);
+		FTransform BoundBoneTransform = GetBoneTransformInSimSpace(SkelComp, MeshBases, BoneIndex);
+
 		FTransform ShapeTransform = BoundBoneTransform;
 		
 		// Local offset to joint for Body1
@@ -552,7 +571,10 @@ void FAnimNode_AnimDynamics::UpdateLimits(USkeletalMeshComponent* SkelComp, FCSP
 				if(PlanarLimit.DrivingBone.IsValid(BoneContainer))
 				{
 					FCompactPoseBoneIndex DrivingBoneIndex = PlanarLimit.DrivingBone.GetCompactPoseIndex(BoneContainer);
-					LimitPlaneTransform *= MeshBases.GetComponentSpaceTransform(DrivingBoneIndex);// * LimitPlaneTransform;
+
+					FTransform DrivingBoneTransform = GetBoneTransformInSimSpace(SkelComp, MeshBases, DrivingBoneIndex);
+
+					LimitPlaneTransform *= DrivingBoneTransform;
 				}
 				
 				FAnimPhys::ConstrainPlanar(NextTimeStep, LinearLimits, &RigidBody, LimitPlaneTransform);
@@ -581,4 +603,197 @@ void FAnimNode_AnimDynamics::PreUpdate(const UAnimInstance* InAnimInstance)
 	const UWorld* World = SkelComp->GetWorld();
 	check(World->GetWorldSettings());
 	CurrentTimeDilation = World->GetWorldSettings()->GetEffectiveTimeDilation();
-}
\ No newline at end of file
+}
+
+FTransform FAnimNode_AnimDynamics::GetBoneTransformInSimSpace(USkeletalMeshComponent* SkelComp, FCSPose<FCompactPose>& MeshBases, const FCompactPoseBoneIndex& BoneIndex)
+{
+	FTransform Transform = MeshBases.GetComponentSpaceTransform(BoneIndex);
+
+	switch(SimulationSpace)
+	{
+	// Change nothing, already in component space
+	case AnimPhysSimSpaceType::Component:
+	{
+		break;
+	}
+
+	case AnimPhysSimSpaceType::Actor:
+	{
+		UObject* OuterObject = SkelComp->GetOuter();
+		if(OuterObject && OuterObject->IsA<AActor>())
+		{
+			AActor* OwningActor = Cast<AActor>(OuterObject);
+			check(OwningActor);
+
+			FTransform WorldTransform = Transform * SkelComp->ComponentToWorld;
+			WorldTransform.SetToRelativeTransform(OwningActor->GetActorTransform());
+			Transform = WorldTransform;
+		}
+		
+		break;
+	}
+
+	case AnimPhysSimSpaceType::RootRelative:
+	{
+		const FBoneContainer& RequiredBones = MeshBases.GetPose().GetBoneContainer();
+
+		FCompactPoseBoneIndex RootBoneCompactIndex(0);
+
+		FTransform RelativeBoneTransform = MeshBases.GetComponentSpaceTransform(RootBoneCompactIndex);
+		Transform = Transform.GetRelativeTransform(RelativeBoneTransform);
+
+		break;
+	}
+
+	case AnimPhysSimSpaceType::BoneRelative:
+	{
+		const FBoneContainer& RequiredBones = MeshBases.GetPose().GetBoneContainer();
+		if(RelativeSpaceBone.IsValid(RequiredBones))
+		{
+			FTransform RelativeBoneTransform = MeshBases.GetComponentSpaceTransform(RelativeSpaceBone.GetCompactPoseIndex(RequiredBones));
+			Transform = Transform.GetRelativeTransform(RelativeBoneTransform);
+		}
+
+		break;
+	}
+
+	case AnimPhysSimSpaceType::World:
+	{
+		// Out to world space
+		FAnimationRuntime::ConvertCSTransformToBoneSpace(SkelComp, MeshBases, Transform, BoneIndex, BCS_WorldSpace);
+	}
+
+	default:
+		break;
+	}
+
+	return Transform;
+}
+
+FTransform FAnimNode_AnimDynamics::GetComponentSpaceTransformFromSimSpace(USkeletalMeshComponent* SkelComp, FCSPose<FCompactPose>& MeshBases, const FCompactPoseBoneIndex& BoneIndex, const FTransform& InSimTransform)
+{
+	FTransform OutTransform = InSimTransform;
+
+	switch(SimulationSpace)
+	{
+		// Change nothing, already in component space
+	case AnimPhysSimSpaceType::Component:
+	{
+		break;
+	}
+
+	case AnimPhysSimSpaceType::Actor:
+	{
+		UObject* OuterObject = SkelComp->GetOuter();
+		if(OuterObject && OuterObject->IsA<AActor>())
+		{
+			AActor* OwningActor = Cast<AActor>(OuterObject);
+			check(OwningActor);
+		
+			FTransform ComponentTransform(SkelComp->RelativeRotation, SkelComp->RelativeLocation, SkelComp->RelativeScale3D);
+			OutTransform = OutTransform * ComponentTransform.Inverse();
+		}
+
+		break;
+	}
+
+	case AnimPhysSimSpaceType::RootRelative:
+	{
+		const FBoneContainer& RequiredBones = MeshBases.GetPose().GetBoneContainer();
+
+		FCompactPoseBoneIndex RootBoneCompactIndex(0);
+
+		FTransform RelativeBoneTransform = MeshBases.GetComponentSpaceTransform(RootBoneCompactIndex);
+		OutTransform = OutTransform * RelativeBoneTransform;
+
+		break;
+	}
+
+	case AnimPhysSimSpaceType::BoneRelative:
+	{
+		const FBoneContainer& RequiredBones = MeshBases.GetPose().GetBoneContainer();
+		if(RelativeSpaceBone.IsValid(RequiredBones))
+		{
+			FTransform RelativeBoneTransform = MeshBases.GetComponentSpaceTransform(RelativeSpaceBone.GetCompactPoseIndex(RequiredBones));
+			OutTransform = OutTransform * RelativeBoneTransform;
+		}
+
+		break;
+	}
+	case AnimPhysSimSpaceType::World:
+	{
+		OutTransform *= SkelComp->ComponentToWorld.Inverse();
+	}
+
+	default:
+		break;
+	}
+
+	return OutTransform;
+}
+
+
+FVector FAnimNode_AnimDynamics::TransformWorldVectorToSimSpace(USkeletalMeshComponent* SkelComp, FCSPose<FCompactPose>& MeshBases, const FVector& InVec)
+{
+	FVector OutVec = InVec;
+
+	switch(SimulationSpace)
+	{
+	case AnimPhysSimSpaceType::Component:
+	{
+		OutVec = SkelComp->ComponentToWorld.InverseTransformVectorNoScale(OutVec);
+
+		break;
+	}
+
+	case AnimPhysSimSpaceType::Actor:
+	{
+		UObject* OuterObject = SkelComp->GetOuter();
+		if(OuterObject && OuterObject->IsA<AActor>())
+		{
+			AActor* OwningActor = Cast<AActor>(OuterObject);
+			check(OwningActor);
+
+			OutVec = OwningActor->GetTransform().TransformVectorNoScale(OutVec);
+		}
+
+		break;
+	}
+
+	case AnimPhysSimSpaceType::RootRelative:
+	{
+		const FBoneContainer& RequiredBones = MeshBases.GetPose().GetBoneContainer();
+
+		FCompactPoseBoneIndex RootBoneCompactIndex(0);
+
+		FTransform RelativeBoneTransform = MeshBases.GetComponentSpaceTransform(RootBoneCompactIndex);
+		RelativeBoneTransform = SkelComp->ComponentToWorld * RelativeBoneTransform;
+		OutVec = RelativeBoneTransform.InverseTransformVectorNoScale(OutVec);
+
+		break;
+	}
+
+	case AnimPhysSimSpaceType::BoneRelative:
+	{
+		const FBoneContainer& RequiredBones = MeshBases.GetPose().GetBoneContainer();
+		if(RelativeSpaceBone.IsValid(RequiredBones))
+		{
+			FTransform RelativeBoneTransform = MeshBases.GetComponentSpaceTransform(RelativeSpaceBone.GetCompactPoseIndex(RequiredBones));
+			RelativeBoneTransform = SkelComp->ComponentToWorld * RelativeBoneTransform;
+			OutVec = RelativeBoneTransform.InverseTransformVectorNoScale(OutVec);
+		}
+
+		break;
+	}
+	case AnimPhysSimSpaceType::World:
+	{
+		break;
+	}
+
+	default:
+		break;
+	}
+
+	return OutVec;
+}
+
diff --git a/Engine/Source/Runtime/AnimGraphRuntime/Public/BoneControllers/AnimNode_AnimDynamics.h b/Engine/Source/Runtime/AnimGraphRuntime/Public/BoneControllers/AnimNode_AnimDynamics.h
--- a/Engine/Source/Runtime/AnimGraphRuntime/Public/BoneControllers/AnimNode_AnimDynamics.h
+++ b/Engine/Source/Runtime/AnimGraphRuntime/Public/BoneControllers/AnimNode_AnimDynamics.h
@@ -27,6 +27,16 @@ enum class AnimPhysLinearConstraintType : uint8
 	Limited,
 };
 
+UENUM()
+enum class AnimPhysSimSpaceType : uint8
+{
+	Component UMETA(ToolTip = "Sim origin is the location/orientation of the skeletal mesh component."),
+	Actor UMETA(ToolTip = "Sim origin is the location/orientation of the actor containing the skeletal mesh component."),
+	World UMETA(ToolTip = "Sim origin is the world origin. Teleporting characters is not recommended in this mode."),
+	RootRelative UMETA(ToolTip = "Sim origin is the location/orientation of the root bone."),
+	BoneRelative UMETA(ToolTip = "Sim origin is the location/orientation of the bone specified in RelativeSpaceBone"),
+};
+
 /** Helper mapping a rigid body to a bone reference */
 struct FAnimPhysBoneRigidBody
 {
@@ -157,6 +167,14 @@ struct ANIMGRAPHRUNTIME_API FAnimNode_AnimDynamics : public FAnimNode_SkeletalCo
 
 	FAnimNode_AnimDynamics();
 	
+	/** The space used to run the simulation */
+	UPROPERTY(EditAnywhere, BlueprintReadWrite, Category = Setup)
+	AnimPhysSimSpaceType SimulationSpace;
+
+	/** When in BoneRelative sim space, the simulation will use this bone as the origin */
+	UPROPERTY(EditAnywhere, BlueprintReadWrite, Category = Setup)
+	FBoneReference RelativeSpaceBone;
+
 	/** Set to true to use the solver to simulate a connected chain */
 	UPROPERTY(EditAnywhere, BlueprintReadWrite, Category = Setup)
 	bool bChain;
@@ -255,6 +273,10 @@ struct ANIMGRAPHRUNTIME_API FAnimNode_AnimDynamics : public FAnimNode_SkeletalCo
 	UPROPERTY(EditAnywhere, Category = PlanarLimit, meta = (UIMin = "1", ClampMin = "1"))
 	float SphereCollisionRadius;
 
+	/** An external force to apply to all bodies in the simulation when ticked, specified in world space */
+	UPROPERTY(EditAnywhere, Category = Forces, meta = (PinShownByDefault))
+	FVector ExternalForce;
+
 	// FAnimNode_SkeletalControlBase interface
 	virtual void Initialize(const FAnimationInitializeContext& Context) override;
 	virtual void UpdateInternal(const FAnimationUpdateContext& Context) override;
@@ -293,6 +315,15 @@ protected:
 
 private:
 
+	// Given a bone index, get it's transform in the currently selected simulation space
+	FTransform GetBoneTransformInSimSpace(USkeletalMeshComponent* SkelComp, FCSPose<FCompactPose>& MeshBases, const FCompactPoseBoneIndex& BoneIndex);
+
+	// Given a transform in simulation space, convert it back to component space
+	FTransform GetComponentSpaceTransformFromSimSpace(USkeletalMeshComponent* SkelComp, FCSPose<FCompactPose>& MeshBases, const FCompactPoseBoneIndex& BoneIndex, const FTransform& InSimTransform);
+
+	// Given a world-space vector, convert it into the current simulation space
+	FVector TransformWorldVectorToSimSpace(USkeletalMeshComponent* SkelComp, FCSPose<FCompactPose>& MeshBases, const FVector& InVec);
+
 	// We can't get clean bone positions unless we are in the evaluate step.
 	// Requesting an init or reinit sets this flag for us to pick up during evaluate
 	bool bRequiresInit;
@@ -335,4 +366,7 @@ private:
 
 	// List of bone references for all bodies in this node
 	TArray<FBoneReference> BoundBoneReferences;
+
+	// Gravity direction in sim space
+	FVector SimSpaceGravityDirection;
 };
\ No newline at end of file
diff --git a/Engine/Source/Runtime/Core/Private/GenericPlatform/GenericPlatformMisc.cpp b/Engine/Source/Runtime/Core/Private/GenericPlatform/GenericPlatformMisc.cpp
--- a/Engine/Source/Runtime/Core/Private/GenericPlatform/GenericPlatformMisc.cpp
+++ b/Engine/Source/Runtime/Core/Private/GenericPlatform/GenericPlatformMisc.cpp
@@ -307,8 +307,8 @@ bool FGenericPlatformMisc::SetStoredValue(const FString& InStoreId, const FStrin
 
 	FConfigSection& Section = ConfigFile.FindOrAdd(InSectionName);
 
-	FString& KeyValue = Section.FindOrAdd(*InKeyName);
-	KeyValue = InValue;
+	FConfigValue& KeyValue = Section.FindOrAdd(*InKeyName);
+	KeyValue = FConfigValue(InValue);
 
 	ConfigFile.Dirty = true;
 	return ConfigFile.Write(ConfigPath);
@@ -329,10 +329,10 @@ bool FGenericPlatformMisc::GetStoredValue(const FString& InStoreId, const FStrin
 	const FConfigSection* const Section = ConfigFile.Find(InSectionName);
 	if(Section)
 	{
-		const FString* const KeyValue = Section->Find(*InKeyName);
+		const FConfigValue* const KeyValue = Section->Find(*InKeyName);
 		if(KeyValue)
 		{
-			OutValue = *KeyValue;
+			OutValue = KeyValue->GetValue();
 			return true;
 		}
 	}
diff --git a/Engine/Source/Runtime/Core/Private/HAL/ConsoleManager.cpp b/Engine/Source/Runtime/Core/Private/HAL/ConsoleManager.cpp
--- a/Engine/Source/Runtime/Core/Private/HAL/ConsoleManager.cpp
+++ b/Engine/Source/Runtime/Core/Private/HAL/ConsoleManager.cpp
@@ -921,7 +921,7 @@ void FConsoleManager::LoadHistoryIfNeeded()
 
 			if(Key == History)
 			{
-				HistoryEntries.Add(It.Value);
+				HistoryEntries.Add(It.Value.GetValue());
 			}
 		}
 	}
diff --git a/Engine/Source/Runtime/Core/Private/HAL/OutputDevices.cpp b/Engine/Source/Runtime/Core/Private/HAL/OutputDevices.cpp
--- a/Engine/Source/Runtime/Core/Private/HAL/OutputDevices.cpp
+++ b/Engine/Source/Runtime/Core/Private/HAL/OutputDevices.cpp
@@ -461,7 +461,7 @@ public:
 		{
 			for( FConfigSectionMap::TIterator It(*RefTypes); It; ++It )
 			{
-				ProcessCmdString(It.Key().ToString() + TEXT(" ") + It.Value(), true);
+				ProcessCmdString(It.Key().ToString() + TEXT(" ") + It.Value().GetValue(), true);
 			}
 		}
 #if !UE_BUILD_SHIPPING
diff --git a/Engine/Source/Runtime/Core/Private/HAL/ThreadingBase.cpp b/Engine/Source/Runtime/Core/Private/HAL/ThreadingBase.cpp
--- a/Engine/Source/Runtime/Core/Private/HAL/ThreadingBase.cpp
+++ b/Engine/Source/Runtime/Core/Private/HAL/ThreadingBase.cpp
@@ -17,13 +17,6 @@ FQueuedThreadPool* GThreadPool = nullptr;
 FQueuedThreadPool* GLargeThreadPool = nullptr;
 #endif
 
-CORE_API bool IsInGameThread()
-{
-	// if the game thread is uninitialized, then we are calling this VERY early before other threads will have started up, so it will be the game thread
-	return !GIsGameThreadIdInitialized || FPlatformTLS::GetCurrentThreadId() == GGameThreadId ||
-		FPlatformTLS::GetCurrentThreadId() == GSlateLoadingThreadId;
-}
-
 CORE_API bool IsInSlateThread()
 {
 	// If this explicitly is a slate thread, not just the main thread running slate
diff --git a/Engine/Source/Runtime/Core/Private/Math/Transform.cpp b/Engine/Source/Runtime/Core/Private/Math/Transform.cpp
--- a/Engine/Source/Runtime/Core/Private/Math/Transform.cpp
+++ b/Engine/Source/Runtime/Core/Private/Math/Transform.cpp
@@ -130,7 +130,7 @@ void FTransform::SetToRelativeTransform(const FTransform& ParentTransform)
  	FMatrix BM = ParentTransform.ToMatrixWithScale();
 #endif
 
-	const FVector SafeRecipScale3D = GetSafeScaleReciprocal(ParentTransform.Scale3D);
+	const FVector SafeRecipScale3D = GetSafeScaleReciprocal(ParentTransform.Scale3D, SMALL_NUMBER);
 	const FQuat InverseRot = ParentTransform.Rotation.Inverse();
 
 	Scale3D *= SafeRecipScale3D;	
@@ -152,7 +152,7 @@ FTransform FTransform::GetRelativeTransform(const FTransform& Other) const
 	// where A = this, B = Other
 	FTransform Result;
 
-	FVector SafeRecipScale3D = GetSafeScaleReciprocal(Other.Scale3D);
+	FVector SafeRecipScale3D = GetSafeScaleReciprocal(Other.Scale3D, SMALL_NUMBER);
 	Result.Scale3D = Scale3D*SafeRecipScale3D;	
 
 	if (Other.Rotation.IsNormalized() == false)
diff --git a/Engine/Source/Runtime/Core/Private/Math/TransformVectorized.cpp b/Engine/Source/Runtime/Core/Private/Math/TransformVectorized.cpp
--- a/Engine/Source/Runtime/Core/Private/Math/TransformVectorized.cpp
+++ b/Engine/Source/Runtime/Core/Private/Math/TransformVectorized.cpp
@@ -146,7 +146,7 @@ void FTransform::SetToRelativeTransform(const FTransform& ParentTransform)
 	checkSlow(ParentTransform.IsRotationNormalized());
 
 	// Scale = S(A)/S(B)	
-	VectorRegister VSafeScale3D	= VectorSet_W0(GetSafeScaleReciprocal(ParentTransform.Scale3D));
+	VectorRegister VSafeScale3D	= VectorSet_W0(GetSafeScaleReciprocal(ParentTransform.Scale3D, ScalarRegister(SMALL_NUMBER)));
 	Scale3D = VectorMultiply(Scale3D, VSafeScale3D);
 	
 	//VQTranslation = (  ( T(A).X - T(B).X ),  ( T(A).Y - T(B).Y ), ( T(A).Z - T(B).Z), 0.f );
