Copying //UE4/Dev-Framework to //UE4/Dev-Main (Source: //UE4/Dev-Framework @ 2964666)

diff --git a/Engine/Source/Editor/PropertyEditor/Private/ItemPropertyNode.cpp b/Engine/Source/Editor/PropertyEditor/Private/ItemPropertyNode.cpp
--- a/Engine/Source/Editor/PropertyEditor/Private/ItemPropertyNode.cpp
+++ b/Engine/Source/Editor/PropertyEditor/Private/ItemPropertyNode.cpp
@@ -127,13 +127,13 @@ void FItemPropertyNode::InitChildNodes()
 		}
 
 		// Expand array.
-		for( int32 ArrayIndex = 0 ; ArrayIndex < Property->ArrayDim ; ArrayIndex++ )
+		for( int32 Index = 0 ; Index < Property->ArrayDim ; Index++ )
 		{
 			bool bShouldBeHidden = false;
 			if( ArraySizeEnum )
 			{
 				// The enum at this array index is hidden
-				bShouldBeHidden = ArraySizeEnum->HasMetaData(TEXT("Hidden"), ArrayIndex );
+				bShouldBeHidden = ArraySizeEnum->HasMetaData(TEXT("Hidden"), Index );
 			}
 
 			if( !bShouldBeHidden )
@@ -142,8 +142,8 @@ void FItemPropertyNode::InitChildNodes()
 				FPropertyNodeInitParams InitParams;
 				InitParams.ParentNode = SharedThis(this);
 				InitParams.Property = Property;
-				InitParams.ArrayOffset = ArrayIndex*Property->ElementSize;
-				InitParams.ArrayIndex = ArrayIndex;
+				InitParams.ArrayOffset = Index*Property->ElementSize;
+				InitParams.ArrayIndex = Index;
 				InitParams.bAllowChildren = true;
 				InitParams.bForceHiddenPropertyVisibility = bShouldShowHiddenProperties;
 				InitParams.bCreateDisableEditOnInstanceNodes = bShouldShowDisableEditOnInstance;
@@ -164,15 +164,15 @@ void FItemPropertyNode::InitChildNodes()
 
 		if( Array )
 		{
-			for( int32 ArrayIndex = 0 ; ArrayIndex < FScriptArrayHelper::Num(Array) ; ArrayIndex++ )
+			for( int32 Index = 0 ; Index < FScriptArrayHelper::Num(Array) ; Index++ )
 			{
 				TSharedPtr<FItemPropertyNode> NewItemNode( new FItemPropertyNode );
 
 				FPropertyNodeInitParams InitParams;
 				InitParams.ParentNode = SharedThis(this);
 				InitParams.Property = ArrayProperty->Inner;
-				InitParams.ArrayOffset = ArrayIndex*ArrayProperty->Inner->ElementSize;
-				InitParams.ArrayIndex = ArrayIndex;
+				InitParams.ArrayOffset = Index*ArrayProperty->Inner->ElementSize;
+				InitParams.ArrayIndex = Index;
 				InitParams.bAllowChildren = true;
 				InitParams.bForceHiddenPropertyVisibility = bShouldShowHiddenProperties;
 				InitParams.bCreateDisableEditOnInstanceNodes = bShouldShowDisableEditOnInstance;
diff --git a/Engine/Source/Editor/PropertyEditor/Private/ObjectPropertyNode.cpp b/Engine/Source/Editor/PropertyEditor/Private/ObjectPropertyNode.cpp
--- a/Engine/Source/Editor/PropertyEditor/Private/ObjectPropertyNode.cpp
+++ b/Engine/Source/Editor/PropertyEditor/Private/ObjectPropertyNode.cpp
@@ -125,19 +125,19 @@ bool FObjectPropertyNode::GetReadAddressUncached(FPropertyNode& InNode,
 
 	if( Cast<UArrayProperty>(InItemProperty->GetOuter()) )
 	{
-		FPropertyNode* ParentNode = InNode.GetParentNode();
-		check(ParentNode);
+		FPropertyNode* ParentPropertyNode = InNode.GetParentNode();
+		check(ParentPropertyNode);
 		const UObject* TempObject = GetUObject(0);
 		if( TempObject )
 		{
-			uint8* BaseAddr = ParentNode->GetValueBaseAddress( (uint8*)TempObject );
+			uint8* BaseAddr = ParentPropertyNode->GetValueBaseAddress( (uint8*)TempObject );
 			if( BaseAddr )
 			{
 				const int32 Num = FScriptArrayHelper::Num(BaseAddr);
 				for( int32 ObjIndex = 1 ; ObjIndex < GetNumObjects(); ObjIndex++ )
 				{
 					TempObject = GetUObject(ObjIndex);
-					BaseAddr = ParentNode->GetValueBaseAddress( (uint8*)TempObject );
+					BaseAddr = ParentPropertyNode->GetValueBaseAddress( (uint8*)TempObject );
 
 					if( BaseAddr && Num != FScriptArrayHelper::Num( BaseAddr ) )
 					{
diff --git a/Engine/Source/Editor/PropertyEditor/Private/PropertyEditorToolkit.cpp b/Engine/Source/Editor/PropertyEditor/Private/PropertyEditorToolkit.cpp
--- a/Engine/Source/Editor/PropertyEditor/Private/PropertyEditorToolkit.cpp
+++ b/Engine/Source/Editor/PropertyEditor/Private/PropertyEditorToolkit.cpp
@@ -20,25 +20,25 @@ const FName FPropertyEditorToolkit::TreeTabId( TEXT( "PropertyEditorToolkit_Prop
 const FName FPropertyEditorToolkit::GridTabId( TEXT( "PropertyEditorToolkit_PropertyTable" ) );
 const FName FPropertyEditorToolkit::TreePinAsColumnHeaderId( TEXT( "PropertyEditorToolkit_PinAsColumnHeader" ) );
 
-void FPropertyEditorToolkit::RegisterTabSpawners(const TSharedRef<class FTabManager>& TabManager)
+void FPropertyEditorToolkit::RegisterTabSpawners(const TSharedRef<class FTabManager>& InTabManager)
 {
-	WorkspaceMenuCategory = TabManager->AddLocalWorkspaceMenuCategory(LOCTEXT("WorkspaceMenu_PropertyEditorToolkit", "Property Editor"));
+	WorkspaceMenuCategory = InTabManager->AddLocalWorkspaceMenuCategory(LOCTEXT("WorkspaceMenu_PropertyEditorToolkit", "Property Editor"));
 
-	TabManager->RegisterTabSpawner( GridTabId, FOnSpawnTab::CreateSP(this, &FPropertyEditorToolkit::SpawnTab_PropertyTable) )
+	InTabManager->RegisterTabSpawner( GridTabId, FOnSpawnTab::CreateSP(this, &FPropertyEditorToolkit::SpawnTab_PropertyTable) )
 		.SetDisplayName( LOCTEXT("PropertyTableTab", "Grid") )
 		.SetGroup( WorkspaceMenuCategory.ToSharedRef() )
 		.SetIcon(FSlateIcon(FEditorStyle::GetStyleSetName(), "LevelEditor.Tabs.Viewports"));
 
-	TabManager->RegisterTabSpawner( TreeTabId, FOnSpawnTab::CreateSP(this, &FPropertyEditorToolkit::SpawnTab_PropertyTree) )
+	InTabManager->RegisterTabSpawner( TreeTabId, FOnSpawnTab::CreateSP(this, &FPropertyEditorToolkit::SpawnTab_PropertyTree) )
 		.SetDisplayName( LOCTEXT("PropertiesTab", "Details") )
 		.SetGroup( WorkspaceMenuCategory.ToSharedRef() )
 		.SetIcon(FSlateIcon(FEditorStyle::GetStyleSetName(), "PropertyEditor.Grid.TabIcon"));
 }
 
-void FPropertyEditorToolkit::UnregisterTabSpawners(const TSharedRef<class FTabManager>& TabManager)
+void FPropertyEditorToolkit::UnregisterTabSpawners(const TSharedRef<class FTabManager>& InTabManager)
 {
-	TabManager->UnregisterTabSpawner( GridTabId );
-	TabManager->UnregisterTabSpawner( TreeTabId );
+	InTabManager->UnregisterTabSpawner( GridTabId );
+	InTabManager->UnregisterTabSpawner( TreeTabId );
 }
 
 FPropertyEditorToolkit::FPropertyEditorToolkit()
@@ -492,13 +492,13 @@ FText FPropertyEditorToolkit::GetBaseToolkitName() const
 
 FText FPropertyEditorToolkit::GetToolkitName() const
 {
-	const auto EditingObjects = GetEditingObjects();
+	const TArray<UObject*>& EditingObjs = GetEditingObjects();
 
-	check( EditingObjects.Num() > 0 );
+	check( EditingObjs.Num() > 0 );
 
-	if( EditingObjects.Num() == 1 )
+	if( EditingObjs.Num() == 1 )
 	{
-		const UObject* EditingObject = EditingObjects[ 0 ];
+		const UObject* EditingObject = EditingObjs[ 0 ];
 
 		const bool bDirtyState = EditingObject->GetOutermost()->IsDirty();
 
@@ -512,9 +512,9 @@ FText FPropertyEditorToolkit::GetToolkitName() const
 	{
 		bool bDirtyState = false;
 		UClass* SharedBaseClass = NULL;
-		for( int32 x = 0; x < EditingObjects.Num(); ++x )
+		for( int32 x = 0; x < EditingObjs.Num(); ++x )
 		{
-			UObject* Obj = EditingObjects[ x ];
+			UObject* Obj = EditingObjs[ x ];
 			check( Obj );
 
 			UClass* ObjClass = Cast<UClass>(Obj);
@@ -542,7 +542,7 @@ FText FPropertyEditorToolkit::GetToolkitName() const
 		}
 
 		FFormatNamedArguments Args;
-		Args.Add( TEXT("NumberOfObjects"), EditingObjects.Num() );
+		Args.Add( TEXT("NumberOfObjects"), EditingObjs.Num() );
 		Args.Add( TEXT("ClassName"), FText::FromString( SharedBaseClass->GetName() ) );
 		Args.Add( TEXT("DirtyState"), bDirtyState ? FText::FromString( TEXT( "*" ) ) : FText::GetEmpty() );
 		return FText::Format( LOCTEXT("ToolkitName_MultiObject", "{NumberOfObjects} {ClassName}{DirtyState} Objects - Property Matrix Editor"), Args );
@@ -552,21 +552,21 @@ FText FPropertyEditorToolkit::GetToolkitName() const
 
 FText FPropertyEditorToolkit::GetToolkitToolTipText() const
 {
-	const auto EditingObjects = GetEditingObjects();
+	const TArray<UObject*>& EditingObjs = GetEditingObjects();
 
-	check( EditingObjects.Num() > 0 );
+	check( EditingObjs.Num() > 0 );
 
-	if( EditingObjects.Num() == 1 )
+	if( EditingObjs.Num() == 1 )
 	{
-		const UObject* EditingObject = EditingObjects[ 0 ];
+		const UObject* EditingObject = EditingObjs[ 0 ];
 		return FAssetEditorToolkit::GetToolTipTextForObject(EditingObject);
 	}
 	else
 	{
 		UClass* SharedBaseClass = NULL;
-		for( int32 x = 0; x < EditingObjects.Num(); ++x )
+		for( int32 x = 0; x < EditingObjs.Num(); ++x )
 		{
-			UObject* Obj = EditingObjects[ x ];
+			UObject* Obj = EditingObjs[ x ];
 			check( Obj );
 
 			UClass* ObjClass = Cast<UClass>(Obj);
@@ -591,7 +591,7 @@ FText FPropertyEditorToolkit::GetToolkitToolTipText() const
 		}
 
 		FFormatNamedArguments Args;
-		Args.Add( TEXT("NumberOfObjects"), EditingObjects.Num() );
+		Args.Add( TEXT("NumberOfObjects"), EditingObjs.Num() );
 		Args.Add( TEXT("ClassName"), FText::FromString( SharedBaseClass->GetName() ) );
 		return FText::Format( LOCTEXT("ToolkitName_MultiObjectToolTip", "{NumberOfObjects} {ClassName} Objects - Property Matrix Editor"), Args );
 	}
diff --git a/Engine/Source/Editor/PropertyEditor/Private/UserInterface/PropertyTable/SColumnHeader.h b/Engine/Source/Editor/PropertyEditor/Private/UserInterface/PropertyTable/SColumnHeader.h
--- a/Engine/Source/Editor/PropertyEditor/Private/UserInterface/PropertyTable/SColumnHeader.h
+++ b/Engine/Source/Editor/PropertyEditor/Private/UserInterface/PropertyTable/SColumnHeader.h
@@ -49,14 +49,14 @@ protected:
 			ColumnLabel.ToSharedRef()
 		];
 		
-		TSharedPtr< SToolTip > ToolTip;
+		TSharedPtr< SToolTip > ColumnHeaderToolTip;
 
 		const TWeakObjectPtr< UObject > Object = Column->GetDataSource()->AsUObject();
 		const TSharedPtr< FPropertyPath > PropertyPath = Column->GetDataSource()->AsPropertyPath();
 
 		if( Object.IsValid() )
 		{
-			ToolTip = SNew( SToolTip ).Text( FText::FromString(Object->GetName()) );
+			ColumnHeaderToolTip = SNew( SToolTip ).Text( FText::FromString(Object->GetName()) );
 		}
 		else if ( PropertyPath.IsValid() && PropertyPath->GetNumProperties() > 0 )
 		{
@@ -64,10 +64,10 @@ protected:
 			const FText ToolTipText = PropertyEditorHelpers::GetToolTipText( Property );
 			const FString DocumentationLink = PropertyEditorHelpers::GetDocumentationLink( Property );
 			const FString DocumentationExcerptName = PropertyEditorHelpers::GetDocumentationExcerptName( Property );
-			ToolTip = IDocumentation::Get()->CreateToolTip( ToolTipText, NULL, DocumentationLink, DocumentationExcerptName );
+			ColumnHeaderToolTip = IDocumentation::Get()->CreateToolTip( ToolTipText, NULL, DocumentationLink, DocumentationExcerptName );
 		}
 
-		this->SetToolTip( ToolTip );
+		this->SetToolTip( ColumnHeaderToolTip );
 	}
 
 	virtual FReply OnMouseButtonDown( const FGeometry& MyGeometry, const FPointerEvent& MouseEvent ) override
diff --git a/Engine/Source/Editor/PropertyEditor/Private/UserInterface/PropertyTable/SPropertyTable.h b/Engine/Source/Editor/PropertyEditor/Private/UserInterface/PropertyTable/SPropertyTable.h
--- a/Engine/Source/Editor/PropertyEditor/Private/UserInterface/PropertyTable/SPropertyTable.h
+++ b/Engine/Source/Editor/PropertyEditor/Private/UserInterface/PropertyTable/SPropertyTable.h
@@ -126,9 +126,8 @@ public:
 	virtual FReply OnKeyDown( const FGeometry& MyGeometry, const FKeyEvent& InKeyEvent ) override
 	{
 		FReply Reply = FReply::Unhandled();
-		const ESelectionMode::Type SelectionMode = Table->GetSelectionMode();
 
-		if ( SelectionMode == ESelectionMode::None )
+		if ( Table->GetSelectionMode() == ESelectionMode::None )
 		{
 			return Reply;
 		}
diff --git a/Engine/Source/Editor/PropertyEditor/Private/UserInterface/PropertyTable/SPropertyTableHeaderRow.h b/Engine/Source/Editor/PropertyEditor/Private/UserInterface/PropertyTable/SPropertyTableHeaderRow.h
--- a/Engine/Source/Editor/PropertyEditor/Private/UserInterface/PropertyTable/SPropertyTableHeaderRow.h
+++ b/Engine/Source/Editor/PropertyEditor/Private/UserInterface/PropertyTable/SPropertyTableHeaderRow.h
@@ -48,12 +48,12 @@ public:
 	{
 		ClearColumns();
 
-		const TArray< TSharedRef< IPropertyTableColumn > > Columns = PropertyTable->GetColumns();
+		const TArray< TSharedRef< IPropertyTableColumn > > PropertyTableColumns = PropertyTable->GetColumns();
 
 		ColumnHeaders.Empty();
 		const TSharedRef< IPropertyTable > TableRef = PropertyTable.ToSharedRef();
 
-		for( auto ColumnIter = Columns.CreateConstIterator(); ColumnIter; ++ColumnIter )
+		for( auto ColumnIter = PropertyTableColumns.CreateConstIterator(); ColumnIter; ++ColumnIter )
 		{
 			const TSharedRef< IPropertyTableColumn > Column = *ColumnIter;
 
diff --git a/Engine/Source/Editor/StaticMeshEditor/Private/StaticMeshAutomationTests.cpp b/Engine/Source/Editor/StaticMeshEditor/Private/StaticMeshAutomationTests.cpp
--- a/Engine/Source/Editor/StaticMeshEditor/Private/StaticMeshAutomationTests.cpp
+++ b/Engine/Source/Editor/StaticMeshEditor/Private/StaticMeshAutomationTests.cpp
@@ -198,20 +198,20 @@ bool FStaticMeshEditorTest::RunTest(const FString& Parameters)
 	if (AutomationParameters.ViewportClient)
 	{
 		const FString BaseFileName = TEXT("StaticMeshEditorTest");
-		FString TestName = FString::Printf(TEXT("%s/Base"), *BaseFileName);
+		FString StaticMeshTestName = FString::Printf(TEXT("%s/Base"), *BaseFileName);
 
 		//Wait for the window to load and then take the initial screen shot		
 		ADD_LATENT_AUTOMATION_COMMAND(FWaitLatentCommand(0.5f));
 		if( FAutomationTestFramework::GetInstance().IsScreenshotAllowed() )
 		{
-			AutomationCommon::GetScreenshotPath(TestName, WindowParameters.ScreenshotName, false);
+			AutomationCommon::GetScreenshotPath(StaticMeshTestName, WindowParameters.ScreenshotName, false);
 			ADD_LATENT_AUTOMATION_COMMAND(FTakeEditorScreenshotCommand(WindowParameters));
 			ADD_LATENT_AUTOMATION_COMMAND(FWaitLatentCommand(0.1f));
 		}
 
 		for(auto CurrentCommandType = 0; CurrentCommandType < EditorViewButtonHelper::EStaticMeshFlag::Max; ++CurrentCommandType)
 		{
-			TestName = FString::Printf(TEXT("%s/%s"), *BaseFileName, EditorViewButtonHelper::GetStaticMeshFlagName(AutomationParameters.CommandType));
+			StaticMeshTestName = FString::Printf(TEXT("%s/%s"), *BaseFileName, EditorViewButtonHelper::GetStaticMeshFlagName(AutomationParameters.CommandType));
 			
 			//Toggle the command for the button, take a screenshot, and then re-toggle the command for the button
 			ADD_LATENT_AUTOMATION_COMMAND(EditorViewButtonHelper::FPerformStaticMeshFlagToggle(AutomationParameters));
@@ -220,7 +220,7 @@ bool FStaticMeshEditorTest::RunTest(const FString& Parameters)
 			if( FAutomationTestFramework::GetInstance().IsScreenshotAllowed() )
 			{
 				//Take the screenshot
-				AutomationCommon::GetScreenshotPath(TestName, WindowParameters.ScreenshotName, false);
+				AutomationCommon::GetScreenshotPath(StaticMeshTestName, WindowParameters.ScreenshotName, false);
 				ADD_LATENT_AUTOMATION_COMMAND(FTakeEditorScreenshotCommand(WindowParameters));
 
 				//Wait so the screenshots have a chance to save
diff --git a/Engine/Source/Editor/StaticMeshEditor/Private/StaticMeshEditor.cpp b/Engine/Source/Editor/StaticMeshEditor/Private/StaticMeshEditor.cpp
--- a/Engine/Source/Editor/StaticMeshEditor/Private/StaticMeshEditor.cpp
+++ b/Engine/Source/Editor/StaticMeshEditor/Private/StaticMeshEditor.cpp
@@ -53,42 +53,42 @@ const FName FStaticMeshEditor::SocketManagerTabId( TEXT( "StaticMeshEditor_Socke
 const FName FStaticMeshEditor::CollisionTabId( TEXT( "StaticMeshEditor_Collision" ) );
 
 
-void FStaticMeshEditor::RegisterTabSpawners(const TSharedRef<class FTabManager>& TabManager)
+void FStaticMeshEditor::RegisterTabSpawners(const TSharedRef<class FTabManager>& InTabManager)
 {
-	WorkspaceMenuCategory = TabManager->AddLocalWorkspaceMenuCategory(LOCTEXT("WorkspaceMenu_StaticMeshEditor", "Static Mesh Editor"));
+	WorkspaceMenuCategory = InTabManager->AddLocalWorkspaceMenuCategory(LOCTEXT("WorkspaceMenu_StaticMeshEditor", "Static Mesh Editor"));
 	auto WorkspaceMenuCategoryRef = WorkspaceMenuCategory.ToSharedRef();
 
-	FAssetEditorToolkit::RegisterTabSpawners(TabManager);
+	FAssetEditorToolkit::RegisterTabSpawners(InTabManager);
 
-	TabManager->RegisterTabSpawner( ViewportTabId, FOnSpawnTab::CreateSP(this, &FStaticMeshEditor::SpawnTab_Viewport) )
+	InTabManager->RegisterTabSpawner( ViewportTabId, FOnSpawnTab::CreateSP(this, &FStaticMeshEditor::SpawnTab_Viewport) )
 		.SetDisplayName( LOCTEXT("ViewportTab", "Viewport") )
 		.SetGroup(WorkspaceMenuCategoryRef)
 		.SetIcon(FSlateIcon(FEditorStyle::GetStyleSetName(), "LevelEditor.Tabs.Viewports"));
 
-	TabManager->RegisterTabSpawner( PropertiesTabId, FOnSpawnTab::CreateSP(this, &FStaticMeshEditor::SpawnTab_Properties) )
+	InTabManager->RegisterTabSpawner( PropertiesTabId, FOnSpawnTab::CreateSP(this, &FStaticMeshEditor::SpawnTab_Properties) )
 		.SetDisplayName( LOCTEXT("PropertiesTab", "Details") )
 		.SetGroup(WorkspaceMenuCategoryRef)
 		.SetIcon(FSlateIcon(FEditorStyle::GetStyleSetName(), "LevelEditor.Tabs.Details"));
 
-	TabManager->RegisterTabSpawner( SocketManagerTabId, FOnSpawnTab::CreateSP(this, &FStaticMeshEditor::SpawnTab_SocketManager) )
+	InTabManager->RegisterTabSpawner( SocketManagerTabId, FOnSpawnTab::CreateSP(this, &FStaticMeshEditor::SpawnTab_SocketManager) )
 		.SetDisplayName( LOCTEXT("SocketManagerTab", "Socket Manager") )
 		.SetGroup(WorkspaceMenuCategoryRef)
 		.SetIcon(FSlateIcon(FEditorStyle::GetStyleSetName(), "StaticMeshEditor.Tabs.SocketManager"));
 
-	TabManager->RegisterTabSpawner( CollisionTabId, FOnSpawnTab::CreateSP(this, &FStaticMeshEditor::SpawnTab_Collision) )
+	InTabManager->RegisterTabSpawner( CollisionTabId, FOnSpawnTab::CreateSP(this, &FStaticMeshEditor::SpawnTab_Collision) )
 		.SetDisplayName( LOCTEXT("CollisionTab", "Convex Decomposition") )
 		.SetGroup(WorkspaceMenuCategoryRef)
 		.SetIcon(FSlateIcon(FEditorStyle::GetStyleSetName(), "StaticMeshEditor.Tabs.ConvexDecomposition"));
 }
 
-void FStaticMeshEditor::UnregisterTabSpawners(const TSharedRef<class FTabManager>& TabManager)
+void FStaticMeshEditor::UnregisterTabSpawners(const TSharedRef<class FTabManager>& InTabManager)
 {
-	FAssetEditorToolkit::UnregisterTabSpawners(TabManager);
+	FAssetEditorToolkit::UnregisterTabSpawners(InTabManager);
 
-	TabManager->UnregisterTabSpawner( ViewportTabId );
-	TabManager->UnregisterTabSpawner( PropertiesTabId );
-	TabManager->UnregisterTabSpawner( SocketManagerTabId );
-	TabManager->UnregisterTabSpawner( CollisionTabId );
+	InTabManager->UnregisterTabSpawner( ViewportTabId );
+	InTabManager->UnregisterTabSpawner( PropertiesTabId );
+	InTabManager->UnregisterTabSpawner( SocketManagerTabId );
+	InTabManager->UnregisterTabSpawner( CollisionTabId );
 }
 
 
diff --git a/Engine/Source/Editor/StaticMeshEditor/Private/StaticMeshEditorViewportClient.cpp b/Engine/Source/Editor/StaticMeshEditor/Private/StaticMeshEditorViewportClient.cpp
--- a/Engine/Source/Editor/StaticMeshEditor/Private/StaticMeshEditorViewportClient.cpp
+++ b/Engine/Source/Editor/StaticMeshEditor/Private/StaticMeshEditorViewportClient.cpp
@@ -152,7 +152,7 @@ struct HSMEVertexProxy : public HHitProxy
 };
 IMPLEMENT_HIT_PROXY(HSMEVertexProxy, HHitProxy);
 
-bool FStaticMeshEditorViewportClient::InputWidgetDelta( FViewport* Viewport, EAxisList::Type CurrentAxis, FVector& Drag, FRotator& Rot, FVector& Scale )
+bool FStaticMeshEditorViewportClient::InputWidgetDelta( FViewport* InViewport, EAxisList::Type CurrentAxis, FVector& Drag, FRotator& Rot, FVector& Scale )
 {
 	bool bHandled = false;
 	if (bManipulating)
@@ -823,24 +823,24 @@ void FStaticMeshEditorViewportClient::DrawUVsForMesh(FViewport* InViewport, FCan
 	DrawUVs(InViewport, InCanvas, InTextYPos, LODLevel, UVChannel, SelectedEdgeTexCoords[UVChannel], StaticMeshComponent->StaticMesh->RenderData, NULL);
 }
 
-void FStaticMeshEditorViewportClient::MouseMove(FViewport* Viewport,int32 x, int32 y)
+void FStaticMeshEditorViewportClient::MouseMove(FViewport* InViewport,int32 x, int32 y)
 {
-	FEditorViewportClient::MouseMove(Viewport,x,y);
+	FEditorViewportClient::MouseMove(InViewport,x,y);
 }
 
-bool FStaticMeshEditorViewportClient::InputKey(FViewport* Viewport, int32 ControllerId, FKey Key, EInputEvent Event,float AmountDepressed,bool Gamepad)
+bool FStaticMeshEditorViewportClient::InputKey(FViewport* InViewport, int32 ControllerId, FKey Key, EInputEvent Event,float AmountDepressed,bool Gamepad)
 {
-	bool bHandled = FEditorViewportClient::InputKey(Viewport, ControllerId, Key, Event, AmountDepressed, false);
+	bool bHandled = FEditorViewportClient::InputKey(InViewport, ControllerId, Key, Event, AmountDepressed, false);
 
 	// Handle viewport screenshot.
-	bHandled |= InputTakeScreenshot( Viewport, Key, Event );
+	bHandled |= InputTakeScreenshot( InViewport, Key, Event );
 
 	return bHandled;
 }
 
-bool FStaticMeshEditorViewportClient::InputAxis(FViewport* Viewport, int32 ControllerId, FKey Key, float Delta, float DeltaTime, int32 NumSamples, bool bGamepad)
+bool FStaticMeshEditorViewportClient::InputAxis(FViewport* InViewport, int32 ControllerId, FKey Key, float Delta, float DeltaTime, int32 NumSamples, bool bGamepad)
 {
-	return FEditorViewportClient::InputAxis(Viewport,ControllerId,Key,Delta,DeltaTime,NumSamples,bGamepad);
+	return FEditorViewportClient::InputAxis(InViewport,ControllerId,Key,Delta,DeltaTime,NumSamples,bGamepad);
 }
 
 void FStaticMeshEditorViewportClient::ProcessClick(class FSceneView& InView, class HHitProxy* HitProxy, FKey Key, EInputEvent Event, uint32 HitX, uint32 HitY)
diff --git a/Engine/Source/Editor/StatsViewer/Private/StatsPages/PrimitiveStatsPage.cpp b/Engine/Source/Editor/StatsViewer/Private/StatsPages/PrimitiveStatsPage.cpp
--- a/Engine/Source/Editor/StatsViewer/Private/StatsPages/PrimitiveStatsPage.cpp
+++ b/Engine/Source/Editor/StatsViewer/Private/StatsPages/PrimitiveStatsPage.cpp
@@ -465,8 +465,8 @@ void FPrimitiveStatsPage::OnEditorSelectionChanged( UObject* NewSelection, TWeak
 {
 	if(InParentStatsViewer.IsValid())
 	{
-		int32 ObjectSetIndex = InParentStatsViewer.Pin()->GetObjectSetIndex();
-		if( ObjectSetIndex == PrimitiveObjectSets_SelectedObjects )
+		const int32 ObjSetIndex = InParentStatsViewer.Pin()->GetObjectSetIndex();
+		if( ObjSetIndex == PrimitiveObjectSets_SelectedObjects )
 		{
 			InParentStatsViewer.Pin()->Refresh();
 		}
@@ -477,8 +477,8 @@ void FPrimitiveStatsPage::OnEditorNewCurrentLevel( TWeakPtr< IStatsViewer > InPa
 {
 	if(InParentStatsViewer.IsValid())
 	{
-		int32 ObjectSetIndex = InParentStatsViewer.Pin()->GetObjectSetIndex();
-		if( ObjectSetIndex == PrimitiveObjectSets_CurrentLevel )
+		const int32 ObjSetIndex = InParentStatsViewer.Pin()->GetObjectSetIndex();
+		if( ObjSetIndex == PrimitiveObjectSets_CurrentLevel )
 		{
 			InParentStatsViewer.Pin()->Refresh();
 		}
diff --git a/Engine/Source/Editor/StatsViewer/Private/StatsPages/StaticMeshLightingInfoStatsPage.cpp b/Engine/Source/Editor/StatsViewer/Private/StatsPages/StaticMeshLightingInfoStatsPage.cpp
--- a/Engine/Source/Editor/StatsViewer/Private/StatsPages/StaticMeshLightingInfoStatsPage.cpp
+++ b/Engine/Source/Editor/StatsViewer/Private/StatsPages/StaticMeshLightingInfoStatsPage.cpp
@@ -291,7 +291,7 @@ void FStaticMeshLightingInfoStatsPage::SwapMappingMethodOnSelectedComponents(TWe
 		{
 			const FScopedTransaction Transaction( LOCTEXT("StaticMeshLightingInfoSwap", "DlgStaticMeshLightingInfo:Swap") );
 
-			bool bRefresh = false;
+			bool bNeedsRefresh = false;
 			for (int32 CompIdx = 0; CompIdx < SelectedObjects.Num(); CompIdx++)
 			{
 				UStaticMeshLightingInfo* Entry = SelectedObjects[CompIdx];
@@ -307,10 +307,10 @@ void FStaticMeshLightingInfoStatsPage::SwapMappingMethodOnSelectedComponents(TWe
 					Entry->StaticMeshComponent->ReregisterComponent();
 				}
 
-				bRefresh = true;
+				bNeedsRefresh = true;
 			}
 
-			if (bRefresh)
+			if (bNeedsRefresh)
 			{
 				InParentStatsViewer.Pin()->Refresh();
 			}
@@ -375,7 +375,7 @@ void FStaticMeshLightingInfoStatsPage::SetMappingMethodOnSelectedComponents(TWea
 		{
 			const FScopedTransaction Transaction( LOCTEXT("StaticMeshLightingInfoSet", "DlgStaticMeshLightingInfo:Set") );
 
-			bool bRefresh = false;
+			bool bNeedsRefresh = false;
 			for (int32 CompIdx = 0; CompIdx < SelectedObjects.Num(); CompIdx++)
 			{
 				UStaticMeshLightingInfo* Entry = SelectedObjects[CompIdx];
@@ -391,10 +391,10 @@ void FStaticMeshLightingInfoStatsPage::SetMappingMethodOnSelectedComponents(TWea
 					Entry->StaticMeshComponent->ReregisterComponent();
 				}
 			
-				bRefresh = true;
+				bNeedsRefresh = true;
 			}
 
-			if (bRefresh)
+			if (bNeedsRefresh)
 			{
 				InParentStatsViewer.Pin()->Refresh();
 			}
@@ -512,8 +512,8 @@ void FStaticMeshLightingInfoStatsPage::OnEditorNewCurrentLevel( TWeakPtr<IStatsV
 {
 	if(InParentStatsViewer.IsValid())
 	{
-		int32 ObjectSetIndex = InParentStatsViewer.Pin()->GetObjectSetIndex();
-		if( ObjectSetIndex == StaticMeshLightingInfoObjectSets_CurrentLevel )
+		const int32 ObjSetIndex = InParentStatsViewer.Pin()->GetObjectSetIndex();
+		if( ObjSetIndex == StaticMeshLightingInfoObjectSets_CurrentLevel )
 		{
 			InParentStatsViewer.Pin()->Refresh();
 		}
@@ -524,8 +524,8 @@ void FStaticMeshLightingInfoStatsPage::OnEditorLevelSelected( TWeakPtr<IStatsVie
 {
 	if(InParentStatsViewer.IsValid())
 	{
-		int32 ObjectSetIndex = InParentStatsViewer.Pin()->GetObjectSetIndex();
-		if( ObjectSetIndex == StaticMeshLightingInfoObjectSets_SelectedLevels )
+		const int32 ObjSetIndex = InParentStatsViewer.Pin()->GetObjectSetIndex();
+		if( ObjSetIndex == StaticMeshLightingInfoObjectSets_SelectedLevels )
 		{
 			InParentStatsViewer.Pin()->Refresh();
 		}
diff --git a/Engine/Source/Editor/StatsViewer/Private/StatsPages/TextureStatsPage.cpp b/Engine/Source/Editor/StatsViewer/Private/StatsPages/TextureStatsPage.cpp
--- a/Engine/Source/Editor/StatsViewer/Private/StatsPages/TextureStatsPage.cpp
+++ b/Engine/Source/Editor/StatsViewer/Private/StatsPages/TextureStatsPage.cpp
@@ -354,8 +354,8 @@ void FTextureStatsPage::OnEditorSelectionChanged( UObject* NewSelection, TWeakPt
 {
 	if(InParentStatsViewer.IsValid())
 	{
-		int32 ObjectSetIndex = InParentStatsViewer.Pin()->GetObjectSetIndex();
-		if( ObjectSetIndex == TextureObjectSet_SelectedActors || ObjectSetIndex == TextureObjectSet_SelectedMaterials )
+		const int32 ObjSetIndex = InParentStatsViewer.Pin()->GetObjectSetIndex();
+		if( ObjSetIndex == TextureObjectSet_SelectedActors || ObjSetIndex == TextureObjectSet_SelectedMaterials )
 		{
 			InParentStatsViewer.Pin()->Refresh();
 		}
@@ -366,8 +366,8 @@ void FTextureStatsPage::OnEditorNewCurrentLevel( TWeakPtr< IStatsViewer > InPare
 {
 	if(InParentStatsViewer.IsValid())
 	{
-		int32 ObjectSetIndex = InParentStatsViewer.Pin()->GetObjectSetIndex();
-		if( ObjectSetIndex == TextureObjectSet_CurrentStreamingLevel )
+		const int32 ObjSetIndex = InParentStatsViewer.Pin()->GetObjectSetIndex();
+		if( ObjSetIndex == TextureObjectSet_CurrentStreamingLevel )
 		{
 			InParentStatsViewer.Pin()->Refresh();
 		}
diff --git a/Engine/Source/Editor/TextureEditor/Private/TextureEditorToolkit.cpp b/Engine/Source/Editor/TextureEditor/Private/TextureEditorToolkit.cpp
--- a/Engine/Source/Editor/TextureEditor/Private/TextureEditorToolkit.cpp
+++ b/Engine/Source/Editor/TextureEditor/Private/TextureEditorToolkit.cpp
@@ -51,31 +51,31 @@ FString FTextureEditorToolkit::GetDocumentationLink( ) const
 }
 
 
-void FTextureEditorToolkit::RegisterTabSpawners( const TSharedRef<class FTabManager>& TabManager )
+void FTextureEditorToolkit::RegisterTabSpawners( const TSharedRef<class FTabManager>& InTabManager )
 {
-	WorkspaceMenuCategory = TabManager->AddLocalWorkspaceMenuCategory(LOCTEXT("WorkspaceMenu_TextureEditor", "Texture Editor"));
+	WorkspaceMenuCategory = InTabManager->AddLocalWorkspaceMenuCategory(LOCTEXT("WorkspaceMenu_TextureEditor", "Texture Editor"));
 	auto WorkspaceMenuCategoryRef = WorkspaceMenuCategory.ToSharedRef();
 
-	FAssetEditorToolkit::RegisterTabSpawners(TabManager);
+	FAssetEditorToolkit::RegisterTabSpawners(InTabManager);
 
-	TabManager->RegisterTabSpawner(ViewportTabId, FOnSpawnTab::CreateSP(this, &FTextureEditorToolkit::HandleTabSpawnerSpawnViewport))
+	InTabManager->RegisterTabSpawner(ViewportTabId, FOnSpawnTab::CreateSP(this, &FTextureEditorToolkit::HandleTabSpawnerSpawnViewport))
 		.SetDisplayName(LOCTEXT("ViewportTab", "Viewport"))
 		.SetGroup(WorkspaceMenuCategoryRef)
 		.SetIcon(FSlateIcon(FEditorStyle::GetStyleSetName(), "LevelEditor.Tabs.Viewports"));
 
-	TabManager->RegisterTabSpawner(PropertiesTabId, FOnSpawnTab::CreateSP(this, &FTextureEditorToolkit::HandleTabSpawnerSpawnProperties))
+	InTabManager->RegisterTabSpawner(PropertiesTabId, FOnSpawnTab::CreateSP(this, &FTextureEditorToolkit::HandleTabSpawnerSpawnProperties))
 		.SetDisplayName(LOCTEXT("PropertiesTab", "Details") )
 		.SetGroup(WorkspaceMenuCategoryRef)
 		.SetIcon(FSlateIcon(FEditorStyle::GetStyleSetName(), "LevelEditor.Tabs.Details"));
 }
 
 
-void FTextureEditorToolkit::UnregisterTabSpawners( const TSharedRef<class FTabManager>& TabManager )
+void FTextureEditorToolkit::UnregisterTabSpawners( const TSharedRef<class FTabManager>& InTabManager )
 {
-	FAssetEditorToolkit::UnregisterTabSpawners(TabManager);
+	FAssetEditorToolkit::UnregisterTabSpawners(InTabManager);
 
-	TabManager->UnregisterTabSpawner(ViewportTabId);
-	TabManager->UnregisterTabSpawner(PropertiesTabId);
+	InTabManager->UnregisterTabSpawner(ViewportTabId);
+	InTabManager->UnregisterTabSpawner(PropertiesTabId);
 }
 
 
diff --git a/Engine/Source/Editor/TranslationEditor/Private/TranslationEditor.cpp b/Engine/Source/Editor/TranslationEditor/Private/TranslationEditor.cpp
--- a/Engine/Source/Editor/TranslationEditor/Private/TranslationEditor.cpp
+++ b/Engine/Source/Editor/TranslationEditor/Private/TranslationEditor.cpp
@@ -64,56 +64,56 @@ void FTranslationEditor::Initialize()
 	TranslationColumn->SetOnFontSizeValueCommitted(FOnInt32ValueCommitted::CreateSP(this, &FTranslationEditor::OnTranslationTargetFontSizeCommitt));
 }
 
-void FTranslationEditor::RegisterTabSpawners(const TSharedRef<class FTabManager>& TabManager)
+void FTranslationEditor::RegisterTabSpawners(const TSharedRef<class FTabManager>& InTabManager)
 {
-	WorkspaceMenuCategory = TabManager->AddLocalWorkspaceMenuCategory(LOCTEXT("WorkspaceMenu_TranslationEditor", "Translation Editor"));
+	WorkspaceMenuCategory = InTabManager->AddLocalWorkspaceMenuCategory(LOCTEXT("WorkspaceMenu_TranslationEditor", "Translation Editor"));
 	auto WorkspaceMenuCategoryRef = WorkspaceMenuCategory.ToSharedRef();
 
-	FAssetEditorToolkit::RegisterTabSpawners(TabManager);
+	FAssetEditorToolkit::RegisterTabSpawners(InTabManager);
 
-	TabManager->RegisterTabSpawner( UntranslatedTabId, FOnSpawnTab::CreateSP(this, &FTranslationEditor::SpawnTab_Untranslated) )
+	InTabManager->RegisterTabSpawner( UntranslatedTabId, FOnSpawnTab::CreateSP(this, &FTranslationEditor::SpawnTab_Untranslated) )
 		.SetDisplayName( LOCTEXT("UntranslatedTab", "Untranslated") )
 		.SetGroup( WorkspaceMenuCategoryRef );
 
-	TabManager->RegisterTabSpawner( ReviewTabId, FOnSpawnTab::CreateSP(this, &FTranslationEditor::SpawnTab_Review) )
+	InTabManager->RegisterTabSpawner( ReviewTabId, FOnSpawnTab::CreateSP(this, &FTranslationEditor::SpawnTab_Review) )
 		.SetDisplayName( LOCTEXT("ReviewTab", "Needs Review") )
 		.SetGroup( WorkspaceMenuCategoryRef );
 
-	TabManager->RegisterTabSpawner( CompletedTabId, FOnSpawnTab::CreateSP(this, &FTranslationEditor::SpawnTab_Completed) )
+	InTabManager->RegisterTabSpawner( CompletedTabId, FOnSpawnTab::CreateSP(this, &FTranslationEditor::SpawnTab_Completed) )
 		.SetDisplayName( LOCTEXT("CompletedTab", "Completed") )
 		.SetGroup( WorkspaceMenuCategoryRef );
 
-	TabManager->RegisterTabSpawner( PreviewTabId, FOnSpawnTab::CreateSP(this, &FTranslationEditor::SpawnTab_Preview) )
+	InTabManager->RegisterTabSpawner( PreviewTabId, FOnSpawnTab::CreateSP(this, &FTranslationEditor::SpawnTab_Preview) )
 		.SetDisplayName( LOCTEXT("PreviewTab", "Preview") )
 		.SetGroup( WorkspaceMenuCategoryRef );
 
-	TabManager->RegisterTabSpawner( ContextTabId, FOnSpawnTab::CreateSP(this, &FTranslationEditor::SpawnTab_Context) )
+	InTabManager->RegisterTabSpawner( ContextTabId, FOnSpawnTab::CreateSP(this, &FTranslationEditor::SpawnTab_Context) )
 		.SetDisplayName( LOCTEXT("ContextTab", "Context") )
 		.SetGroup( WorkspaceMenuCategoryRef );
 
-	TabManager->RegisterTabSpawner( HistoryTabId, FOnSpawnTab::CreateSP(this, &FTranslationEditor::SpawnTab_History) )
+	InTabManager->RegisterTabSpawner( HistoryTabId, FOnSpawnTab::CreateSP(this, &FTranslationEditor::SpawnTab_History) )
 		.SetDisplayName( LOCTEXT("HistoryTab", "History") )
 		.SetGroup( WorkspaceMenuCategoryRef );
 
-	TabManager->RegisterTabSpawner( SearchTabId, FOnSpawnTab::CreateSP(this, &FTranslationEditor::SpawnTab_Search) )
+	InTabManager->RegisterTabSpawner( SearchTabId, FOnSpawnTab::CreateSP(this, &FTranslationEditor::SpawnTab_Search) )
 		.SetDisplayName(LOCTEXT("SearchTab", "Search"))
 		.SetGroup( WorkspaceMenuCategoryRef );
 
-	TabManager->RegisterTabSpawner( ChangedOnImportTabId, FOnSpawnTab::CreateSP(this, &FTranslationEditor::SpawnTab_ChangedOnImport) )
+	InTabManager->RegisterTabSpawner( ChangedOnImportTabId, FOnSpawnTab::CreateSP(this, &FTranslationEditor::SpawnTab_ChangedOnImport) )
 		.SetDisplayName(LOCTEXT("ChangedOnImportTab", "Changed On Import"))
 		.SetGroup( WorkspaceMenuCategoryRef );
 }
 
-void FTranslationEditor::UnregisterTabSpawners(const TSharedRef<class FTabManager>& TabManager)
+void FTranslationEditor::UnregisterTabSpawners(const TSharedRef<class FTabManager>& InTabManager)
 {
-	TabManager->UnregisterTabSpawner( UntranslatedTabId );
-	TabManager->UnregisterTabSpawner( ReviewTabId );
-	TabManager->UnregisterTabSpawner( CompletedTabId );
-	TabManager->UnregisterTabSpawner( PreviewTabId );
-	TabManager->UnregisterTabSpawner( ContextTabId );
-	TabManager->UnregisterTabSpawner(HistoryTabId);
-	TabManager->UnregisterTabSpawner(SearchTabId);
-	TabManager->UnregisterTabSpawner(ChangedOnImportTabId);
+	InTabManager->UnregisterTabSpawner( UntranslatedTabId );
+	InTabManager->UnregisterTabSpawner( ReviewTabId );
+	InTabManager->UnregisterTabSpawner( CompletedTabId );
+	InTabManager->UnregisterTabSpawner( PreviewTabId );
+	InTabManager->UnregisterTabSpawner( ContextTabId );
+	InTabManager->UnregisterTabSpawner(HistoryTabId);
+	InTabManager->UnregisterTabSpawner(SearchTabId);
+	InTabManager->UnregisterTabSpawner(ChangedOnImportTabId);
 }
 
 void FTranslationEditor::InitTranslationEditor( const EToolkitMode::Type Mode, const TSharedPtr< class IToolkitHost >& InitToolkitHost )
diff --git a/Engine/Source/Editor/UnrealEd/Private/Commandlets/GatherTextCommandlet.cpp b/Engine/Source/Editor/UnrealEd/Private/Commandlets/GatherTextCommandlet.cpp
--- a/Engine/Source/Editor/UnrealEd/Private/Commandlets/GatherTextCommandlet.cpp
+++ b/Engine/Source/Editor/UnrealEd/Private/Commandlets/GatherTextCommandlet.cpp
@@ -73,15 +73,15 @@ int32 UGatherTextCommandlet::Main( const FString& Params )
 
 	UE_LOG(LogGatherTextCommandlet, Log,TEXT("Beginning GatherText Commandlet."));
 
-	TSharedRef< FManifestInfo > ManifestInfo = MakeShareable( new FManifestInfo() );
-	TSharedPtr< FGatherTextSCC > SourceControlInfo = NULL;
+	TSharedRef< FManifestInfo > CommandletManifestInfo = MakeShareable( new FManifestInfo() );
+	TSharedPtr< FGatherTextSCC > CommandletSourceControlInfo = nullptr;
 
 	if( bEnableSourceControl )
 	{
-		SourceControlInfo = MakeShareable( new FGatherTextSCC() );
+		CommandletSourceControlInfo = MakeShareable( new FGatherTextSCC() );
 
 		FText SCCErrorStr;
-		if( !SourceControlInfo->IsReady( SCCErrorStr ) )
+		if( !CommandletSourceControlInfo->IsReady( SCCErrorStr ) )
 		{
 			UE_LOG( LogGatherTextCommandlet, Error, TEXT("Source Control error: %s"), *SCCErrorStr.ToString() );
 			return -1;
@@ -113,7 +113,7 @@ int32 UGatherTextCommandlet::Main( const FString& Params )
 		UGatherTextCommandletBase* Commandlet = NewObject<UGatherTextCommandletBase>(GetTransientPackage(), CommandletClass);
 		check(Commandlet);
 		Commandlet->AddToRoot();
-		Commandlet->Initialize( ManifestInfo, SourceControlInfo );
+		Commandlet->Initialize( CommandletManifestInfo, CommandletSourceControlInfo );
 
 		// Execute the commandlet.
 		double CommandletExecutionStartTime = FPlatformTime::Seconds();
@@ -144,10 +144,10 @@ int32 UGatherTextCommandlet::Main( const FString& Params )
 		if( 0 != Commandlet->Main( GeneratedCmdLine ) )
 		{
 			UE_LOG(LogGatherTextCommandlet, Error,TEXT("%s-%s reported an error."),*SectionName, *CommandletClassName);
-			if( SourceControlInfo.IsValid() )
+			if( CommandletSourceControlInfo.IsValid() )
 			{
 				FText SCCErrorStr;
-				if( !SourceControlInfo->CleanUp( SCCErrorStr ) )
+				if( !CommandletSourceControlInfo->CleanUp( SCCErrorStr ) )
 				{
 					UE_LOG(LogGatherTextCommandlet, Error, TEXT("%s"), *SCCErrorStr.ToString());
 				}
@@ -158,17 +158,17 @@ int32 UGatherTextCommandlet::Main( const FString& Params )
 		UE_LOG(LogGatherTextCommandlet, Log,TEXT("Completed %s: %s"), *SectionName, *CommandletClassName);
 	}
 
-	if( SourceControlInfo.IsValid() && !bDisableSubmit )
+	if( CommandletSourceControlInfo.IsValid() && !bDisableSubmit )
 	{
 		FText SCCErrorStr;
-		if( SourceControlInfo->CheckinFiles( GetChangelistDescription(GatherTextConfigPath), SCCErrorStr ) )
+		if( CommandletSourceControlInfo->CheckinFiles( GetChangelistDescription(GatherTextConfigPath), SCCErrorStr ) )
 		{
 			UE_LOG(LogGatherTextCommandlet, Log,TEXT("Submitted Localization files."));
 		}
 		else
 		{
 			UE_LOG(LogGatherTextCommandlet, Error, TEXT("%s"), *SCCErrorStr.ToString());
-			if( !SourceControlInfo->CleanUp( SCCErrorStr ) )
+			if( !CommandletSourceControlInfo->CleanUp( SCCErrorStr ) )
 			{
 				UE_LOG(LogGatherTextCommandlet, Error, TEXT("%s"), *SCCErrorStr.ToString());
 			}
diff --git a/Engine/Source/Editor/UnrealEd/Private/Dialogs/SOutputLogDialog.cpp b/Engine/Source/Editor/UnrealEd/Private/Dialogs/SOutputLogDialog.cpp
--- a/Engine/Source/Editor/UnrealEd/Private/Dialogs/SOutputLogDialog.cpp
+++ b/Engine/Source/Editor/UnrealEd/Private/Dialogs/SOutputLogDialog.cpp
@@ -164,9 +164,9 @@ bool SOutputLogDialog::SupportsKeyboardFocus() const
 
 FVector2D SOutputLogDialog::ComputeDesiredSize(float LayoutScaleMultiplier) const
 {
-	FVector2D DesiredSize = SCompoundWidget::ComputeDesiredSize(LayoutScaleMultiplier);
-	DesiredSize.X = FMath::Min(DesiredSize.X, MaxWidth);
-	return DesiredSize;
+	FVector2D OutputLogDesiredSize = SCompoundWidget::ComputeDesiredSize(LayoutScaleMultiplier);
+	OutputLogDesiredSize.X = FMath::Min(OutputLogDesiredSize.X, MaxWidth);
+	return OutputLogDesiredSize;
 }
 
 void SOutputLogDialog::CopyMessageToClipboard( )
diff --git a/Engine/Source/Editor/UnrealEd/Private/EditorActor.cpp b/Engine/Source/Editor/UnrealEd/Private/EditorActor.cpp
--- a/Engine/Source/Editor/UnrealEd/Private/EditorActor.cpp
+++ b/Engine/Source/Editor/UnrealEd/Private/EditorActor.cpp
@@ -1893,7 +1893,7 @@ void UUnrealEdEngine::edactSelectMatchingStaticMesh( bool bAllClasses )
 {
 	TArray<FStaticMeshActor> StaticMeshActors;
 
-	TArray<UWorld*> WorldList;
+	TArray<UWorld*> SelectedWorlds;
 	// Make a list of selected actors with static meshes.
 	for ( FSelectionIterator It( GetSelectedActorIterator() ) ; It ; ++It )
 	{
@@ -1906,24 +1906,24 @@ void UUnrealEdEngine::edactSelectMatchingStaticMesh( bool bAllClasses )
 			if ( ActorInfo.IsStaticMeshActor() )
 			{
 				StaticMeshActors.Add( ActorInfo );
-				WorldList.AddUnique( Actor->GetWorld() );	
+				SelectedWorlds.AddUnique( Actor->GetWorld() );	
 			}				
 		}
 	}
-	if( WorldList.Num() == 0 )
+	if( SelectedWorlds.Num() == 0 )
 	{
 		UE_LOG(LogEditorActor, Log, TEXT("No worlds found in edactSelectMatchingStaticMesh") );
 		return;
 	}
 	// Make sure we have only 1 valid world 
-	check(WorldList.Num() == 1);
+	check(SelectedWorlds.Num() == 1);
 	USelection* SelectedActors = GetSelectedActors();
 	SelectedActors->BeginBatchSelectOperation();
 	SelectedActors->Modify();
 
 	// Loop through all non-hidden actors in visible levels, selecting those that have one of the
 	// static meshes in the list.
-	for( FActorIterator It(WorldList[0]); It; ++It )
+	for( FActorIterator It(SelectedWorlds[0]); It; ++It )
 	{
 		AActor* Actor = *It;
 		if ( !Actor->IsHiddenEd() )
@@ -1963,7 +1963,7 @@ void UUnrealEdEngine::edactSelectMatchingSkeletalMesh(bool bAllClasses)
 	bool bSelectSkelMeshActors = false;
 	bool bSelectPawns = false;
 
-	TArray<UWorld*> WorldList;
+	TArray<UWorld*> SelectedWorlds;
 	// Make a list of skeletal meshes of selected actors, and note what classes we have selected.
 	for ( FSelectionIterator It( GetSelectedActorIterator() ) ; It ; ++It )
 	{
@@ -1976,7 +1976,7 @@ void UUnrealEdEngine::edactSelectMatchingSkeletalMesh(bool bAllClasses)
 		{
 			bSelectSkelMeshActors = true;
 			SelectedMeshes.AddUnique(SkelMeshActor->GetSkeletalMeshComponent()->SkeletalMesh);
-			WorldList.AddUnique(Actor->GetWorld());			
+			SelectedWorlds.AddUnique(Actor->GetWorld());			
 		}
 
 		// Look for Pawn
@@ -1988,17 +1988,17 @@ void UUnrealEdEngine::edactSelectMatchingSkeletalMesh(bool bAllClasses)
 			{
 				bSelectPawns = true;
 				SelectedMeshes.AddUnique(PawnSkeletalMesh->SkeletalMesh);
-				WorldList.AddUnique(Actor->GetWorld());
+				SelectedWorlds.AddUnique(Actor->GetWorld());
 			}
 		}
 	}
-	if( WorldList.Num() == 0 )
+	if( SelectedWorlds.Num() == 0 )
 	{
 		UE_LOG(LogEditorActor, Log, TEXT("No worlds found in edactSelectMatchingSkeletalMesh") );
 		return;
 	}
 	// Make sure we have only 1 valid world 
-	check( WorldList.Num() == 1 );
+	check( SelectedWorlds.Num() == 1 );
 	// If desired, select all class types
 	if(bAllClasses)
 	{
@@ -2011,7 +2011,7 @@ void UUnrealEdEngine::edactSelectMatchingSkeletalMesh(bool bAllClasses)
 	SelectedActors->Modify();
 
 	// Loop through all non-hidden actors in visible levels, selecting those that have one of the skeletal meshes in the list.
-	for( FActorIterator It(WorldList[0]); It; ++It )
+	for( FActorIterator It(SelectedWorlds[0]); It; ++It )
 	{
 		AActor* Actor = *It;
 		if ( !Actor->IsHiddenEd() )
@@ -2059,7 +2059,7 @@ void UUnrealEdEngine::edactSelectMatchingMaterial()
 	// Set for fast lookup of used materials.
 	TSet<UMaterialInterface*> MaterialsInSelection;
 
-	TArray<UWorld*> WorldList;
+	TArray<UWorld*> SelectedWorlds;
 	// For each selected actor, find all the materials used by this actor.
 	for ( FSelectionIterator ActorItr( GetSelectedActorIterator() ) ; ActorItr ; ++ActorItr )
 	{
@@ -2077,18 +2077,18 @@ void UUnrealEdEngine::edactSelectMatchingMaterial()
 				TArray<UMaterialInterface*> UsedMaterials;
 				CurrentComponent->GetUsedMaterials( UsedMaterials );
 				MaterialsInSelection.Append( UsedMaterials );
-				WorldList.AddUnique( CurrentActor->GetWorld() );
+				SelectedWorlds.AddUnique( CurrentActor->GetWorld() );
 			}
 		}
 	}
 
-	if( WorldList.Num() == 0 )
+	if( SelectedWorlds.Num() == 0 )
 	{
 		UE_LOG(LogEditorActor, Log, TEXT("No worlds found in edactSelectMatchingMaterial") );
 		return;
 	}
 	// Make sure we have only 1 valid world 
-	check( WorldList.Num() == 1 );
+	check( SelectedWorlds.Num() == 1 );
 
 	USelection* SelectedActors = GetSelectedActors();
 	SelectedActors->BeginBatchSelectOperation();
@@ -2096,7 +2096,7 @@ void UUnrealEdEngine::edactSelectMatchingMaterial()
 
 	// Now go over every actor and see if any of the actors are using any of the materials that 
 	// we found above.
-	for( FActorIterator ActorIt(WorldList[0]); ActorIt; ++ActorIt )
+	for( FActorIterator ActorIt(SelectedWorlds[0]); ActorIt; ++ActorIt )
 	{
 		AActor* Actor = *ActorIt;
 
@@ -2140,7 +2140,7 @@ void UUnrealEdEngine::edactSelectMatchingEmitter()
 {
 	TArray<UParticleSystem*> SelectedParticleSystemTemplates;
 
-	TArray<UWorld*> WorldList;
+	TArray<UWorld*> SelectedWorlds;
 	// Check all of the currently selected actors to find the relevant particle system templates to use to match
 	for ( FSelectionIterator SelectedIterator( GetSelectedActorIterator() ) ; SelectedIterator ; ++SelectedIterator )
 	{
@@ -2152,23 +2152,23 @@ void UUnrealEdEngine::edactSelectMatchingEmitter()
 		if ( Emitter && Emitter->GetParticleSystemComponent() && Emitter->GetParticleSystemComponent()->Template )
 		{
 			SelectedParticleSystemTemplates.AddUnique( Emitter->GetParticleSystemComponent()->Template );
-			WorldList.AddUnique( Actor->GetWorld() );
+			SelectedWorlds.AddUnique( Actor->GetWorld() );
 		}
 	}
 
-	if( WorldList.Num() == 0 )
+	if( SelectedWorlds.Num() == 0 )
 	{
 		UE_LOG(LogEditorActor, Log, TEXT("No worlds found in edactSelectMatchingEmitter") );
 		return;
 	}
 	// Make sure we have only 1 valid world 
-	check( WorldList.Num() == 1 );
+	check( SelectedWorlds.Num() == 1 );
 
 	USelection* SelectedActors = GetSelectedActors();
 	SelectedActors->BeginBatchSelectOperation();
 	SelectedActors->Modify();
 	// Iterate over all of the non-hidden actors, selecting those who have a particle system template that matches one from the previously-found list
-	for( TActorIterator<AEmitter> ActorIterator(WorldList[0]); ActorIterator; ++ActorIterator )
+	for( TActorIterator<AEmitter> ActorIterator(SelectedWorlds[0]); ActorIterator; ++ActorIterator )
 	{
 		AEmitter* ActorAsEmitter = *ActorIterator;
 		if ( !ActorAsEmitter->IsHiddenEd() )
diff --git a/Engine/Source/Editor/UnrealEd/Private/Fbx/FbxMeshImportData.cpp b/Engine/Source/Editor/UnrealEd/Private/Fbx/FbxMeshImportData.cpp
--- a/Engine/Source/Editor/UnrealEd/Private/Fbx/FbxMeshImportData.cpp
+++ b/Engine/Source/Editor/UnrealEd/Private/Fbx/FbxMeshImportData.cpp
@@ -21,8 +21,8 @@ bool UFbxMeshImportData::CanEditChange(const UProperty* InProperty) const
 		bMutable = Outer->CanEditChange(InProperty);
 	}
 
-	static FName NormalGenerationMethod("NormalGenerationMethod");
-	if( bMutable && InProperty->GetFName() == NormalGenerationMethod )
+	static FName NAME_NormalGenerationMethod("NormalGenerationMethod");
+	if( bMutable && InProperty->GetFName() == NAME_NormalGenerationMethod )
 	{
 		// Normal generation method is ignored if we are importing both tangents and normals
 		return NormalImportMethod == FBXNIM_ImportNormalsAndTangents ? false : true;
diff --git a/Engine/Source/Editor/UnrealEd/Private/Fbx/SSceneBaseMeshListView.cpp b/Engine/Source/Editor/UnrealEd/Private/Fbx/SSceneBaseMeshListView.cpp
--- a/Engine/Source/Editor/UnrealEd/Private/Fbx/SSceneBaseMeshListView.cpp
+++ b/Engine/Source/Editor/UnrealEd/Private/Fbx/SSceneBaseMeshListView.cpp
@@ -96,9 +96,9 @@ void SFbxSSceneBaseMeshListView::RemoveSelectionFromImport()
 
 void SFbxSSceneBaseMeshListView::SetSelectionImportState(bool MarkForImport)
 {
-	TArray<FbxMeshInfoPtr> SelectedItems;
-	GetSelectedItems(SelectedItems);
-	for (auto Item : SelectedItems)
+	TArray<FbxMeshInfoPtr> SelectedFbxMeshes;
+	GetSelectedItems(SelectedFbxMeshes);
+	for (auto Item : SelectedFbxMeshes)
 	{
 		FbxMeshInfoPtr ItemPtr = Item;
 		ItemPtr->bImportAttribute = MarkForImport;
@@ -108,9 +108,9 @@ void SFbxSSceneBaseMeshListView::SetSelectionImportState(bool MarkForImport)
 void SFbxSSceneBaseMeshListView::OnSelectionChanged(FbxMeshInfoPtr Item, ESelectInfo::Type SelectionType)
 {
 	//Change the option selection
-	TArray<FbxMeshInfoPtr> SelectedItems;
-	GetSelectedItems(SelectedItems);
-	for (FbxMeshInfoPtr SelectItem : SelectedItems)
+	TArray<FbxMeshInfoPtr> SelectedFbxMeshes;
+	GetSelectedItems(SelectedFbxMeshes);
+	for (FbxMeshInfoPtr SelectItem : SelectedFbxMeshes)
 	{
 		for (auto kvp : *OverrideNameOptionsMap)
 		{
@@ -151,13 +151,13 @@ void SFbxSSceneBaseMeshListView::AddBakePivotMenu(class FMenuBuilder& MenuBuilde
 
 void SFbxSSceneBaseMeshListView::FillPivotContextMenu(FMenuBuilder& MenuBuilder)
 {
-	TArray<FbxMeshInfoPtr> SelectedItems;
-	int32 SelectCount = GetSelectedItems(SelectedItems);
+	TArray<FbxMeshInfoPtr> SelectedFbxMeshes;
+	int32 SelectCount = GetSelectedItems(SelectedFbxMeshes);
 
 	uint64 InvalidUid = INVALID_UNIQUE_ID;
-	if (SelectedItems.Num() == 1)
+	if (SelectedFbxMeshes.Num() == 1)
 	{
-		FbxMeshInfoPtr Item = SelectedItems[0];
+		FbxMeshInfoPtr Item = SelectedFbxMeshes[0];
 		if (Item->bOriginalTypeChanged)
 			return;
 		MenuBuilder.AddMenuEntry(Item->PivotNodeUid == INVALID_UNIQUE_ID ? LOCTEXT("ResetPivotBakeCurrent", "* No Pivot Bake") : LOCTEXT("ResetPivotBake", "No Pivot Bake"), FText(), FSlateIcon(), FUIAction(FExecuteAction::CreateSP(this, &SFbxSSceneBaseMeshListView::AssignToPivot, InvalidUid)));
@@ -200,9 +200,9 @@ TSharedPtr<FFbxNodeInfo> SFbxSSceneBaseMeshListView::FindNodeInfoByUid(uint64 No
 void SFbxSSceneBaseMeshListView::AssignToPivot(uint64 NodeUid)
 {
 	FbxNodeInfoPtr NodeInfo = FindNodeInfoByUid(NodeUid, SceneInfo);
-	TArray<FbxMeshInfoPtr> SelectedItems;
-	int32 SelectCount = GetSelectedItems(SelectedItems);
-	for (FbxMeshInfoPtr MeshInfo : SelectedItems)
+	TArray<FbxMeshInfoPtr> SelectedFbxMeshes;
+	int32 SelectCount = GetSelectedItems(SelectedFbxMeshes);
+	for (FbxMeshInfoPtr MeshInfo : SelectedFbxMeshes)
 	{
 		if (MeshInfo->bOriginalTypeChanged)
 		{
@@ -239,9 +239,9 @@ void SFbxSSceneBaseMeshListView::AssignToOptions(FString OptionName)
 	{
 		return;
 	}
-	TArray<FbxMeshInfoPtr> SelectedItems;
-	GetSelectedItems(SelectedItems);
-	for (FbxMeshInfoPtr ItemPtr : SelectedItems)
+	TArray<FbxMeshInfoPtr> SelectedFbxMeshes;
+	GetSelectedItems(SelectedFbxMeshes);
+	for (FbxMeshInfoPtr ItemPtr : SelectedFbxMeshes)
 	{
 		if (ItemPtr->bOriginalTypeChanged)
 		{
diff --git a/Engine/Source/Editor/UnrealEd/Private/Fbx/SSceneImportNodeTreeView.cpp b/Engine/Source/Editor/UnrealEd/Private/Fbx/SSceneImportNodeTreeView.cpp
--- a/Engine/Source/Editor/UnrealEd/Private/Fbx/SSceneImportNodeTreeView.cpp
+++ b/Engine/Source/Editor/UnrealEd/Private/Fbx/SSceneImportNodeTreeView.cpp
@@ -307,9 +307,9 @@ TSharedPtr<SWidget> SFbxSceneTreeView::OnOpenContextMenu()
 
 	//Get the different type of the multi selection
 	TSet<TSharedPtr<FFbxAttributeInfo>> SelectAssets;
-	TArray<FbxNodeInfoPtr> SelectedItems;
-	const auto NumSelectedItems = GetSelectedItems(SelectedItems);
-	for (auto Item : SelectedItems)
+	TArray<FbxNodeInfoPtr> SelectedFbxNodeInfos;
+	const auto NumSelectedItems = GetSelectedItems(SelectedFbxNodeInfos);
+	for (auto Item : SelectedFbxNodeInfos)
 	{
 		FbxNodeInfoPtr ItemPtr = Item;
 		if (ItemPtr->AttributeInfo.IsValid())
@@ -359,9 +359,9 @@ void SFbxSceneTreeView::RemoveSelectionFromImport()
 
 void SFbxSceneTreeView::SetSelectionImportState(bool MarkForImport)
 {
-	TArray<FbxNodeInfoPtr> SelectedItems;
-	GetSelectedItems(SelectedItems);
-	for (auto Item : SelectedItems)
+	TArray<FbxNodeInfoPtr> SelectedFbxNodeInfos;
+	GetSelectedItems(SelectedFbxNodeInfos);
+	for (auto Item : SelectedFbxNodeInfos)
 	{
 		FbxNodeInfoPtr ItemPtr = Item;
 		ItemPtr->bImportNode = MarkForImport;
diff --git a/Engine/Source/Editor/UnrealEd/Private/Fbx/SSceneImportStaticMeshListView.cpp b/Engine/Source/Editor/UnrealEd/Private/Fbx/SSceneImportStaticMeshListView.cpp
--- a/Engine/Source/Editor/UnrealEd/Private/Fbx/SSceneImportStaticMeshListView.cpp
+++ b/Engine/Source/Editor/UnrealEd/Private/Fbx/SSceneImportStaticMeshListView.cpp
@@ -291,8 +291,8 @@ TSharedRef< ITableRow > SFbxSceneStaticMeshListView::OnGenerateRowFbxSceneListVi
 
 TSharedPtr<SWidget> SFbxSceneStaticMeshListView::OnOpenContextMenu()
 {
-	TArray<FbxMeshInfoPtr> SelectedItems;
-	int32 SelectCount = GetSelectedItems(SelectedItems);
+	TArray<FbxMeshInfoPtr> SelectedFbxMeshInfos;
+	int32 SelectCount = GetSelectedItems(SelectedFbxMeshInfos);
 	// Build up the menu for a selection
 	const bool bCloseAfterSelection = true;
 	FMenuBuilder MenuBuilder(bCloseAfterSelection, TSharedPtr<FUICommandList>());
diff --git a/Engine/Source/Editor/UnrealEd/Private/Fbx/SSceneMaterialsListView.cpp b/Engine/Source/Editor/UnrealEd/Private/Fbx/SSceneMaterialsListView.cpp
--- a/Engine/Source/Editor/UnrealEd/Private/Fbx/SSceneMaterialsListView.cpp
+++ b/Engine/Source/Editor/UnrealEd/Private/Fbx/SSceneMaterialsListView.cpp
@@ -300,8 +300,8 @@ void SFbxSceneMaterialsListView::UpdateMaterialBasePath()
 
 TSharedPtr<SWidget> SFbxSceneMaterialsListView::OnOpenContextMenu()
 {
-	TArray<FbxMaterialInfoPtr> SelectedItems;
-	int32 SelectCount = GetSelectedItems(SelectedItems);
+	TArray<FbxMaterialInfoPtr> SelectedFbxMaterialInfos;
+	int32 SelectCount = GetSelectedItems(SelectedFbxMaterialInfos);
 	// Build up the menu for a selection
 	const bool bCloseAfterSelection = true;
 	FMenuBuilder MenuBuilder(bCloseAfterSelection, TSharedPtr<FUICommandList>());
@@ -362,9 +362,9 @@ void SFbxSceneMaterialsListView::RemoveSelectionFromImport()
 
 void SFbxSceneMaterialsListView::SetSelectionImportState(bool MarkForImport)
 {
-	TArray<FbxMaterialInfoPtr> SelectedItems;
-	GetSelectedItems(SelectedItems);
-	for (FbxMaterialInfoPtr Item : SelectedItems)
+	TArray<FbxMaterialInfoPtr> SelectedFbxMaterialInfos;
+	GetSelectedItems(SelectedFbxMaterialInfos);
+	for (FbxMaterialInfoPtr Item : SelectedFbxMaterialInfos)
 	{
 		Item->bImportAttribute = MarkForImport;
 	}
@@ -373,10 +373,10 @@ void SFbxSceneMaterialsListView::SetSelectionImportState(bool MarkForImport)
 void SFbxSceneMaterialsListView::OnSelectionChanged(FbxMaterialInfoPtr Item, ESelectInfo::Type SelectionType)
 {
 	//Change the texture list
-	TArray<FbxMaterialInfoPtr> SelectedItems;
-	GetSelectedItems(SelectedItems);
+	TArray<FbxMaterialInfoPtr> SelectedFbxMaterialInfos;
+	GetSelectedItems(SelectedFbxMaterialInfos);
 	TexturesArray->Reset();
-	for (FbxMaterialInfoPtr SelectItem : SelectedItems)
+	for (FbxMaterialInfoPtr SelectItem : SelectedFbxMaterialInfos)
 	{
 		TexturesArray->Append(SelectItem->Textures);
 	}
@@ -392,8 +392,8 @@ void SFbxSceneMaterialsListView::OnToggleSelectAll(ECheckBoxState CheckType)
 
 void SFbxSceneMaterialsListView::AssignMaterialAssetData(const FAssetData& AssetData)
 {
-	TArray<FbxMaterialInfoPtr> SelectedItems;
-	int32 SelectCount = GetSelectedItems(SelectedItems);
+	TArray<FbxMaterialInfoPtr> SelectedFbxMaterialInfos;
+	int32 SelectCount = GetSelectedItems(SelectedFbxMaterialInfos);
 	if (SelectCount != 1)
 	{
 		FSlateApplication::Get().DismissAllMenus();
@@ -406,7 +406,7 @@ void SFbxSceneMaterialsListView::AssignMaterialAssetData(const FAssetData& Asset
 	{
 		if (!ContentObject->HasAnyFlags(RF_Transient) && !ContentObject->IsPendingKill())
 		{
-			for (FbxMaterialInfoPtr ItemPtr : SelectedItems)
+			for (FbxMaterialInfoPtr ItemPtr : SelectedFbxMaterialInfos)
 			{
 				//Override the MeshInfo with the new asset path
 				ItemPtr->SetOverridePath(true);
@@ -420,13 +420,13 @@ void SFbxSceneMaterialsListView::AssignMaterialAssetData(const FAssetData& Asset
 
 void SFbxSceneMaterialsListView::AssignMaterialToExisting()
 {
-	TArray<FbxMaterialInfoPtr> SelectedItems;
-	int32 SelectCount = GetSelectedItems(SelectedItems);
+	TArray<FbxMaterialInfoPtr> SelectedFbxMaterialInfos;
+	int32 SelectCount = GetSelectedItems(SelectedFbxMaterialInfos);
 	if (SelectCount != 1)
 		return;
 	FContentBrowserModule& ContentBrowserModule = FModuleManager::LoadModuleChecked<FContentBrowserModule>("ContentBrowser");
 
-	for (FbxMaterialInfoPtr ItemPtr : SelectedItems)
+	for (FbxMaterialInfoPtr ItemPtr : SelectedFbxMaterialInfos)
 	{
 		FOpenAssetDialogConfig SelectAssetConfig;
 		SelectAssetConfig.DialogTitleOverride = LOCTEXT("FbxChooseMaterialAssetContentPath", "Choose a material asset");
@@ -454,10 +454,10 @@ void SFbxSceneMaterialsListView::AssignMaterialToExisting()
 
 void SFbxSceneMaterialsListView::ResetAssignMaterial()
 {
-	TArray<FbxMaterialInfoPtr> SelectedItems;
-	int32 SelectCount = GetSelectedItems(SelectedItems);
+	TArray<FbxMaterialInfoPtr> SelectedFbxMaterialInfos;
+	int32 SelectCount = GetSelectedItems(SelectedFbxMaterialInfos);
 
-	for (FbxMaterialInfoPtr ItemPtr : SelectedItems)
+	for (FbxMaterialInfoPtr ItemPtr : SelectedFbxMaterialInfos)
 	{
 		if (ItemPtr->bOverridePath)
 		{
diff --git a/Engine/Source/Editor/UnrealEd/Private/Fbx/SSceneReimportNodeTreeView.cpp b/Engine/Source/Editor/UnrealEd/Private/Fbx/SSceneReimportNodeTreeView.cpp
--- a/Engine/Source/Editor/UnrealEd/Private/Fbx/SSceneReimportNodeTreeView.cpp
+++ b/Engine/Source/Editor/UnrealEd/Private/Fbx/SSceneReimportNodeTreeView.cpp
@@ -526,9 +526,9 @@ TSharedPtr<SWidget> SFbxReimportSceneTreeView::OnOpenContextMenu()
 
 	//Get the different type of the multi selection
 	TSet<TSharedPtr<FFbxAttributeInfo>> SelectAssets;
-	TArray<FbxNodeInfoPtr> SelectedItems;
-	const auto NumSelectedItems = GetSelectedItems(SelectedItems);
-	for (auto Item : SelectedItems)
+	TArray<FbxNodeInfoPtr> SelectedFbxNodeInfos;
+	const auto NumSelectedItems = GetSelectedItems(SelectedFbxNodeInfos);
+	for (auto Item : SelectedFbxNodeInfos)
 	{
 		FbxNodeInfoPtr ItemPtr = Item;
 		if (ItemPtr->AttributeInfo.IsValid())
@@ -578,9 +578,9 @@ void SFbxReimportSceneTreeView::RemoveSelectionFromImport()
 
 void SFbxReimportSceneTreeView::SetSelectionImportState(bool MarkForImport)
 {
-	TArray<FbxNodeInfoPtr> SelectedItems;
-	GetSelectedItems(SelectedItems);
-	for (FbxNodeInfoPtr Item : SelectedItems)
+	TArray<FbxNodeInfoPtr> SelectedFbxNodeInfos;
+	GetSelectedItems(SelectedFbxNodeInfos);
+	for (FbxNodeInfoPtr Item : SelectedFbxNodeInfos)
 	{
 		EFbxSceneReimportStatusFlags *ItemStatus = NodeStatusMap->Find(Item->NodeHierarchyPath);
 		if (ItemStatus != nullptr)
diff --git a/Engine/Source/Editor/UnrealEd/Private/Fbx/SSceneReimportSkeletalMeshListView.cpp b/Engine/Source/Editor/UnrealEd/Private/Fbx/SSceneReimportSkeletalMeshListView.cpp
--- a/Engine/Source/Editor/UnrealEd/Private/Fbx/SSceneReimportSkeletalMeshListView.cpp
+++ b/Engine/Source/Editor/UnrealEd/Private/Fbx/SSceneReimportSkeletalMeshListView.cpp
@@ -535,8 +535,8 @@ void SFbxSceneSkeletalMeshReimportListView::OnFinishedChangingProperties(const F
 
 TSharedPtr<SWidget> SFbxSceneSkeletalMeshReimportListView::OnOpenContextMenu()
 {
-	TArray<FbxMeshInfoPtr> SelectedItems;
-	int32 SelectCount = GetSelectedItems(SelectedItems);
+	TArray<FbxMeshInfoPtr> SelectedFbxMeshInfos;
+	int32 SelectCount = GetSelectedItems(SelectedFbxMeshInfos);
 	// Build up the menu for a selection
 	const bool bCloseAfterSelection = true;
 	FMenuBuilder MenuBuilder(bCloseAfterSelection, TSharedPtr<FUICommandList>());
@@ -555,7 +555,7 @@ TSharedPtr<SWidget> SFbxSceneSkeletalMeshReimportListView::OnOpenContextMenu()
 	//AddBakePivotMenu(MenuBuilder);
 
 	bool bShowOptionMenu = false;
-	for (FbxMeshInfoPtr MeshInfo : SelectedItems)
+	for (FbxMeshInfoPtr MeshInfo : SelectedFbxMeshInfos)
 	{
 		EFbxSceneReimportStatusFlags ReimportFlags = *MeshStatusMap->Find(MeshInfo->OriginalImportPath);
 		if ((ReimportFlags & EFbxSceneReimportStatusFlags::Removed) == EFbxSceneReimportStatusFlags::None)
@@ -579,9 +579,9 @@ TSharedPtr<SWidget> SFbxSceneSkeletalMeshReimportListView::OnOpenContextMenu()
 
 void SFbxSceneSkeletalMeshReimportListView::SetSelectionImportState(bool MarkForImport)
 {
-	TArray<FbxMeshInfoPtr> SelectedItems;
-	GetSelectedItems(SelectedItems);
-	for (FbxMeshInfoPtr ItemPtr : SelectedItems)
+	TArray<FbxMeshInfoPtr> SelectedFbxMeshInfos;
+	GetSelectedItems(SelectedFbxMeshInfos);
+	for (FbxMeshInfoPtr ItemPtr : SelectedFbxMeshInfos)
 	{
 		EFbxSceneReimportStatusFlags *ItemStatus = MeshStatusMap->Find(ItemPtr->OriginalImportPath);
 		if (ItemStatus != nullptr)
diff --git a/Engine/Source/Editor/UnrealEd/Private/Fbx/SSceneReimportStaticMeshListView.cpp b/Engine/Source/Editor/UnrealEd/Private/Fbx/SSceneReimportStaticMeshListView.cpp
--- a/Engine/Source/Editor/UnrealEd/Private/Fbx/SSceneReimportStaticMeshListView.cpp
+++ b/Engine/Source/Editor/UnrealEd/Private/Fbx/SSceneReimportStaticMeshListView.cpp
@@ -534,8 +534,8 @@ void SFbxSceneStaticMeshReimportListView::OnFinishedChangingProperties(const FPr
 
 TSharedPtr<SWidget> SFbxSceneStaticMeshReimportListView::OnOpenContextMenu()
 {
-	TArray<FbxMeshInfoPtr> SelectedItems;
-	int32 SelectCount = GetSelectedItems(SelectedItems);
+	TArray<FbxMeshInfoPtr> SelectedFbxMeshInfos;
+	int32 SelectCount = GetSelectedItems(SelectedFbxMeshInfos);
 	// Build up the menu for a selection
 	const bool bCloseAfterSelection = true;
 	FMenuBuilder MenuBuilder(bCloseAfterSelection, TSharedPtr<FUICommandList>());
@@ -570,7 +570,7 @@ TSharedPtr<SWidget> SFbxSceneStaticMeshReimportListView::OnOpenContextMenu()
 	AddBakePivotMenu(MenuBuilder);
 
 	bool bShowOptionMenu = false;
-	for (FbxMeshInfoPtr MeshInfo : SelectedItems)
+	for (FbxMeshInfoPtr MeshInfo : SelectedFbxMeshInfos)
 	{
 		EFbxSceneReimportStatusFlags ReimportFlags = *MeshStatusMap->Find(MeshInfo->OriginalImportPath);
 		if ((ReimportFlags & EFbxSceneReimportStatusFlags::Removed) == EFbxSceneReimportStatusFlags::None)
@@ -594,13 +594,13 @@ TSharedPtr<SWidget> SFbxSceneStaticMeshReimportListView::OnOpenContextMenu()
 
 void SFbxSceneStaticMeshReimportListView::AssignToStaticMesh()
 {
-	TArray<FbxMeshInfoPtr> SelectedItems;
-	int32 SelectCount = GetSelectedItems(SelectedItems);
+	TArray<FbxMeshInfoPtr> SelectedFbxMeshInfos;
+	int32 SelectCount = GetSelectedItems(SelectedFbxMeshInfos);
 	if (SelectCount != 1)
 		return;
 	FContentBrowserModule& ContentBrowserModule = FModuleManager::LoadModuleChecked<FContentBrowserModule>("ContentBrowser");
 
-	for (FbxMeshInfoPtr ItemPtr : SelectedItems)
+	for (FbxMeshInfoPtr ItemPtr : SelectedFbxMeshInfos)
 	{
 		FOpenAssetDialogConfig SelectAssetConfig;
 		SelectAssetConfig.DialogTitleOverride = LOCTEXT("FbxChooseReimportAssetContentPath", "Choose static mesh asset for reimporting the fbx scene content");
@@ -628,11 +628,11 @@ void SFbxSceneStaticMeshReimportListView::AssignToStaticMesh()
 
 bool SFbxSceneStaticMeshReimportListView::ShowResetAssignToStaticMesh()
 {
-	TArray<FbxMeshInfoPtr> SelectedItems;
-	int32 SelectCount = GetSelectedItems(SelectedItems);
+	TArray<FbxMeshInfoPtr> SelectedFbxMeshInfos;
+	int32 SelectCount = GetSelectedItems(SelectedFbxMeshInfos);
 	FContentBrowserModule& ContentBrowserModule = FModuleManager::LoadModuleChecked<FContentBrowserModule>("ContentBrowser");
 
-	for (FbxMeshInfoPtr ItemPtr : SelectedItems)
+	for (FbxMeshInfoPtr ItemPtr : SelectedFbxMeshInfos)
 	{
 		if (ItemPtr->bOverridePath)
 		{
@@ -644,10 +644,10 @@ bool SFbxSceneStaticMeshReimportListView::ShowResetAssignToStaticMesh()
 
 void SFbxSceneStaticMeshReimportListView::ResetAssignToStaticMesh()
 {
-	TArray<FbxMeshInfoPtr> SelectedItems;
-	int32 SelectCount = GetSelectedItems(SelectedItems);
+	TArray<FbxMeshInfoPtr> SelectedFbxMeshInfos;
+	int32 SelectCount = GetSelectedItems(SelectedFbxMeshInfos);
 
-	for (FbxMeshInfoPtr ItemPtr : SelectedItems)
+	for (FbxMeshInfoPtr ItemPtr : SelectedFbxMeshInfos)
 	{
 		if (ItemPtr->bOverridePath)
 		{
@@ -673,9 +673,9 @@ void SFbxSceneStaticMeshReimportListView::ResetAssignToStaticMesh()
 
 void SFbxSceneStaticMeshReimportListView::SetSelectionImportState(bool MarkForImport)
 {
-	TArray<FbxMeshInfoPtr> SelectedItems;
-	GetSelectedItems(SelectedItems);
-	for (FbxMeshInfoPtr ItemPtr : SelectedItems)
+	TArray<FbxMeshInfoPtr> SelectedFbxMeshInfos;
+	GetSelectedItems(SelectedFbxMeshInfos);
+	for (FbxMeshInfoPtr ItemPtr : SelectedFbxMeshInfos)
 	{
 		EFbxSceneReimportStatusFlags *ItemStatus = MeshStatusMap->Find(ItemPtr->OriginalImportPath);
 		if (ItemStatus != nullptr)
diff --git a/Engine/Source/Editor/UnrealEd/Private/Fbx/SSceneSkeletalMeshListView.cpp b/Engine/Source/Editor/UnrealEd/Private/Fbx/SSceneSkeletalMeshListView.cpp
--- a/Engine/Source/Editor/UnrealEd/Private/Fbx/SSceneSkeletalMeshListView.cpp
+++ b/Engine/Source/Editor/UnrealEd/Private/Fbx/SSceneSkeletalMeshListView.cpp
@@ -290,8 +290,8 @@ TSharedRef< ITableRow > SFbxSceneSkeletalMeshListView::OnGenerateRowFbxSceneList
 
 TSharedPtr<SWidget> SFbxSceneSkeletalMeshListView::OnOpenContextMenu()
 {
-	TArray<FbxMeshInfoPtr> SelectedItems;
-	int32 SelectCount = GetSelectedItems(SelectedItems);
+	TArray<FbxMeshInfoPtr> SelectedFbxMeshInfos;
+	int32 SelectCount = GetSelectedItems(SelectedFbxMeshInfos);
 	// Build up the menu for a selection
 	const bool bCloseAfterSelection = true;
 	FMenuBuilder MenuBuilder(bCloseAfterSelection, TSharedPtr<FUICommandList>());
diff --git a/Engine/Source/Editor/UnrealEd/Private/Lightmass/LightmassLandscapeRender.cpp b/Engine/Source/Editor/UnrealEd/Private/Lightmass/LightmassLandscapeRender.cpp
--- a/Engine/Source/Editor/UnrealEd/Private/Lightmass/LightmassLandscapeRender.cpp
+++ b/Engine/Source/Editor/UnrealEd/Private/Lightmass/LightmassLandscapeRender.cpp
@@ -51,20 +51,20 @@ public:
 	/** Default constructor. */
 	FLightmassLandscapeVertexFactory()
 	{
-		FLocalVertexFactory::FDataType Data;
+		FLocalVertexFactory::FDataType VertextData;
 		// Position
-		Data.PositionComponent = STRUCTMEMBER_VERTEXSTREAMCOMPONENT(&LightmassLandscapeVertexBuffer, FLightmassLandscapeVertex, Position, VET_Float3);
+		VertextData.PositionComponent = STRUCTMEMBER_VERTEXSTREAMCOMPONENT(&LightmassLandscapeVertexBuffer, FLightmassLandscapeVertex, Position, VET_Float3);
 		// Tangents
-		Data.TangentBasisComponents[0] = STRUCTMEMBER_VERTEXSTREAMCOMPONENT(&LightmassLandscapeVertexBuffer, FLightmassLandscapeVertex, TangentX, VET_PackedNormal);
-		Data.TangentBasisComponents[1] = STRUCTMEMBER_VERTEXSTREAMCOMPONENT(&LightmassLandscapeVertexBuffer, FLightmassLandscapeVertex, TangentZ, VET_PackedNormal);
+		VertextData.TangentBasisComponents[0] = STRUCTMEMBER_VERTEXSTREAMCOMPONENT(&LightmassLandscapeVertexBuffer, FLightmassLandscapeVertex, TangentX, VET_PackedNormal);
+		VertextData.TangentBasisComponents[1] = STRUCTMEMBER_VERTEXSTREAMCOMPONENT(&LightmassLandscapeVertexBuffer, FLightmassLandscapeVertex, TangentZ, VET_PackedNormal);
 		// Color
-		Data.ColorComponent = STRUCTMEMBER_VERTEXSTREAMCOMPONENT(&LightmassLandscapeVertexBuffer, FLightmassLandscapeVertex, Color, VET_Color);
+		VertextData.ColorComponent = STRUCTMEMBER_VERTEXSTREAMCOMPONENT(&LightmassLandscapeVertexBuffer, FLightmassLandscapeVertex, Color, VET_Color);
 		// UVs (packed two to a stream component)
-		Data.TextureCoordinates.Add(STRUCTMEMBER_VERTEXSTREAMCOMPONENT(&LightmassLandscapeVertexBuffer, FLightmassLandscapeVertex, UVs[0], VET_Half4));
-		Data.TextureCoordinates.Add(STRUCTMEMBER_VERTEXSTREAMCOMPONENT(&LightmassLandscapeVertexBuffer, FLightmassLandscapeVertex, UVs[2], VET_Half4));
+		VertextData.TextureCoordinates.Add(STRUCTMEMBER_VERTEXSTREAMCOMPONENT(&LightmassLandscapeVertexBuffer, FLightmassLandscapeVertex, UVs[0], VET_Half4));
+		VertextData.TextureCoordinates.Add(STRUCTMEMBER_VERTEXSTREAMCOMPONENT(&LightmassLandscapeVertexBuffer, FLightmassLandscapeVertex, UVs[2], VET_Half4));
 
 		// update the data
-		SetData(Data);
+		SetData(VertextData);
 	}
 };
 TGlobalResource<FLightmassLandscapeVertexFactory> LightmassLandscapeVertexFactory;
diff --git a/Engine/Source/Editor/UnrealEd/Private/PerformanceMonitor.cpp b/Engine/Source/Editor/UnrealEd/Private/PerformanceMonitor.cpp
--- a/Engine/Source/Editor/UnrealEd/Private/PerformanceMonitor.cpp
+++ b/Engine/Source/Editor/UnrealEd/Private/PerformanceMonitor.cpp
@@ -290,11 +290,11 @@ void FPerformanceMonitor::Tick(float DeltaTime)
 
 		if (PercentUnderTarget >= CVarPercentThreshold->GetFloat())
 		{
-			static IConsoleVariable* CVarCoarseSampleTime = IConsoleManager::Get().FindConsoleVariable(TEXT("PerfWarn.CoarseSampleTime"));
+			static IConsoleVariable* CoarseSampleTimeCVar = IConsoleManager::Get().FindConsoleVariable(TEXT("PerfWarn.CoarseSampleTime"));
 
 			Arguments.Add(TEXT("Framerate"), CVarMinFPS->GetInt());
 			Arguments.Add(TEXT("Percentage"), FMath::FloorToFloat(PercentUnderTarget));
-			SampleTime = CVarCoarseSampleTime->GetInt();
+			SampleTime = CoarseSampleTimeCVar->GetInt();
 
 			bLowFramerate = true;
 		}
diff --git a/Engine/Source/Editor/UnrealEd/Private/PlayLevel.cpp b/Engine/Source/Editor/UnrealEd/Private/PlayLevel.cpp
--- a/Engine/Source/Editor/UnrealEd/Private/PlayLevel.cpp
+++ b/Engine/Source/Editor/UnrealEd/Private/PlayLevel.cpp
@@ -3871,14 +3871,14 @@ UGameViewportClient * UEditorEngine::GetNextPIEViewport(UGameViewportClient * Cu
 	return NULL;
 }
 
-void UEditorEngine::RemapGamepadControllerIdForPIE(class UGameViewportClient* GameViewport, int32 &ControllerId)
+void UEditorEngine::RemapGamepadControllerIdForPIE(class UGameViewportClient* InGameViewport, int32 &ControllerId)
 {
 	// Increment the controller id if we are the focused window, and RouteGamepadToSecondWindow is true (and we are running multiple clients).
 	// This cause the focused window to NOT handle the input, decrement controllerID, and pass it to the next window.
 	const ULevelEditorPlaySettings* PlayInSettings = GetDefault<ULevelEditorPlaySettings>();
 	const bool CanRouteGamepadToSecondWindow = [&PlayInSettings]{ bool RouteGamepadToSecondWindow(false); return (PlayInSettings->GetRouteGamepadToSecondWindow(RouteGamepadToSecondWindow) && RouteGamepadToSecondWindow); }();
 	const bool CanRunUnderOneProcess = [&PlayInSettings]{ bool RunUnderOneProcess(false); return (PlayInSettings->GetRunUnderOneProcess(RunUnderOneProcess) && RunUnderOneProcess); }();
-	if ( CanRouteGamepadToSecondWindow && CanRunUnderOneProcess && GameViewport->GetWindow().IsValid() && GameViewport->GetWindow()->HasFocusedDescendants())
+	if ( CanRouteGamepadToSecondWindow && CanRunUnderOneProcess && InGameViewport->GetWindow().IsValid() && InGameViewport->GetWindow()->HasFocusedDescendants())
 	{
 		ControllerId++;
 	}
diff --git a/Engine/Source/Editor/UnrealEd/Private/SClassPickerDialog.cpp b/Engine/Source/Editor/UnrealEd/Private/SClassPickerDialog.cpp
--- a/Engine/Source/Editor/UnrealEd/Private/SClassPickerDialog.cpp
+++ b/Engine/Source/Editor/UnrealEd/Private/SClassPickerDialog.cpp
@@ -288,12 +288,12 @@ void SClassPickerDialog::OnCustomAreaExpansionChanged(bool bExpanded)
 
 EVisibility SClassPickerDialog::GetSelectButtonVisibility() const
 {
-	EVisibility Visibility = EVisibility::Hidden;
-	if( ChosenClass != NULL )
+	EVisibility ButtonVisibility = EVisibility::Hidden;
+	if( ChosenClass != nullptr )
 	{
-		Visibility = EVisibility::Visible;
+		ButtonVisibility = EVisibility::Visible;
 	}
-	return Visibility;
+	return ButtonVisibility;
 }
 
 /** Overridden from SWidget: Called when a key is pressed down - capturing copy */
diff --git a/Engine/Source/Editor/UnrealEd/Private/Tests/EditorAutomationTests.cpp b/Engine/Source/Editor/UnrealEd/Private/Tests/EditorAutomationTests.cpp
--- a/Engine/Source/Editor/UnrealEd/Private/Tests/EditorAutomationTests.cpp
+++ b/Engine/Source/Editor/UnrealEd/Private/Tests/EditorAutomationTests.cpp
@@ -299,8 +299,8 @@ bool FLoadAllMapsInEditorTest::RunTest(const FString& Parameters)
 		CVar->Set(0);
 
 		//Create a screen shot filename and path
-		const FString TestName = FString::Printf(TEXT("LoadAllMaps_Editor/%s"), *FPaths::GetBaseFilename(MapName));
-		AutomationCommon::GetScreenshotPath(TestName, WindowParameters.ScreenshotName, true);
+		const FString LoadAllMapsTestName = FString::Printf(TEXT("LoadAllMaps_Editor/%s"), *FPaths::GetBaseFilename(MapName));
+		AutomationCommon::GetScreenshotPath(LoadAllMapsTestName, WindowParameters.ScreenshotName, true);
 
 		//Get the current number of seconds.  This will be used to track how long it took to load the map.
 		MapLoadStartTime = FPlatformTime::Seconds();
diff --git a/Engine/Source/Editor/UnrealEd/Private/Toolkits/SimpleAssetEditor.cpp b/Engine/Source/Editor/UnrealEd/Private/Toolkits/SimpleAssetEditor.cpp
--- a/Engine/Source/Editor/UnrealEd/Private/Toolkits/SimpleAssetEditor.cpp
+++ b/Engine/Source/Editor/UnrealEd/Private/Toolkits/SimpleAssetEditor.cpp
@@ -14,23 +14,23 @@
 const FName FSimpleAssetEditor::ToolkitFName( TEXT( "GenericAssetEditor" ) );
 const FName FSimpleAssetEditor::PropertiesTabId( TEXT( "GenericEditor_Properties" ) );
 
-void FSimpleAssetEditor::RegisterTabSpawners(const TSharedRef<class FTabManager>& TabManager)
+void FSimpleAssetEditor::RegisterTabSpawners(const TSharedRef<class FTabManager>& InTabManager)
 {
-	WorkspaceMenuCategory = TabManager->AddLocalWorkspaceMenuCategory(LOCTEXT("WorkspaceMenu_GenericAssetEditor", "Asset Editor"));
+	WorkspaceMenuCategory = InTabManager->AddLocalWorkspaceMenuCategory(LOCTEXT("WorkspaceMenu_GenericAssetEditor", "Asset Editor"));
 
-	FAssetEditorToolkit::RegisterTabSpawners(TabManager);
+	FAssetEditorToolkit::RegisterTabSpawners(InTabManager);
 
-	TabManager->RegisterTabSpawner( PropertiesTabId, FOnSpawnTab::CreateSP(this, &FSimpleAssetEditor::SpawnPropertiesTab) )
+	InTabManager->RegisterTabSpawner( PropertiesTabId, FOnSpawnTab::CreateSP(this, &FSimpleAssetEditor::SpawnPropertiesTab) )
 		.SetDisplayName( LOCTEXT("PropertiesTab", "Details") )
 		.SetGroup(WorkspaceMenuCategory.ToSharedRef())
 		.SetIcon(FSlateIcon(FEditorStyle::GetStyleSetName(), "LevelEditor.Tabs.Details"));
 }
 
-void FSimpleAssetEditor::UnregisterTabSpawners(const TSharedRef<class FTabManager>& TabManager)
+void FSimpleAssetEditor::UnregisterTabSpawners(const TSharedRef<class FTabManager>& InTabManager)
 {
-	FAssetEditorToolkit::UnregisterTabSpawners(TabManager);
+	FAssetEditorToolkit::UnregisterTabSpawners(InTabManager);
 
-	TabManager->UnregisterTabSpawner( PropertiesTabId );
+	InTabManager->UnregisterTabSpawner( PropertiesTabId );
 }
 
 const FName FSimpleAssetEditor::SimpleEditorAppIdentifier( TEXT( "GenericEditorApp" ) );
@@ -114,16 +114,16 @@ FText FSimpleAssetEditor::GetBaseToolkitName() const
 
 FText FSimpleAssetEditor::GetToolkitName() const
 {
-	const auto EditingObjects = GetEditingObjects();
+	const TArray<UObject*>& EditingObjs = GetEditingObjects();
 
-	check( EditingObjects.Num() > 0 );
+	check( EditingObjs.Num() > 0 );
 
 	FFormatNamedArguments Args;
 	Args.Add( TEXT("ToolkitName"), GetBaseToolkitName() );
 
-	if( EditingObjects.Num() == 1 )
+	if( EditingObjs.Num() == 1 )
 	{
-		const UObject* EditingObject = EditingObjects[ 0 ];
+		const UObject* EditingObject = EditingObjs[ 0 ];
 
 		const bool bDirtyState = EditingObject->GetOutermost()->IsDirty();
 
@@ -134,21 +134,21 @@ FText FSimpleAssetEditor::GetToolkitName() const
 	else
 	{
 		bool bDirtyState = false;
-		UClass* SharedBaseClass = NULL;
-		for( int32 x = 0; x < EditingObjects.Num(); ++x )
+		UClass* SharedBaseClass = nullptr;
+		for( int32 x = 0; x < EditingObjs.Num(); ++x )
 		{
-			UObject* Obj = EditingObjects[ x ];
+			UObject* Obj = EditingObjs[ x ];
 			check( Obj );
 
 			UClass* ObjClass = Cast<UClass>(Obj);
-			if (ObjClass == NULL)
+			if (ObjClass == nullptr)
 			{
 				ObjClass = Obj->GetClass();
 			}
 			check( ObjClass );
 
 			// Initialize with the class of the first object we encounter.
-			if( SharedBaseClass == NULL )
+			if( SharedBaseClass == nullptr )
 			{
 				SharedBaseClass = ObjClass;
 			}
@@ -164,7 +164,7 @@ FText FSimpleAssetEditor::GetToolkitName() const
 			bDirtyState |= Obj->GetOutermost()->IsDirty();
 		}
 
-		Args.Add( TEXT("NumberOfObjects"), EditingObjects.Num() );
+		Args.Add( TEXT("NumberOfObjects"), EditingObjs.Num() );
 		Args.Add( TEXT("ClassName"), FText::FromString( SharedBaseClass->GetName() ) );
 		Args.Add( TEXT("DirtyState"), bDirtyState ? FText::FromString( TEXT( "*" ) ) : FText::GetEmpty() );
 		return FText::Format( LOCTEXT("ToolkitTitle_EditingMultiple", "{NumberOfObjects} {ClassName}{DirtyState} - {ToolkitName}"), Args );
@@ -173,35 +173,35 @@ FText FSimpleAssetEditor::GetToolkitName() const
 
 FText FSimpleAssetEditor::GetToolkitToolTipText() const
 {
-	const auto EditingObjects = GetEditingObjects();
+	const TArray<UObject*>& EditingObjs = GetEditingObjects();
 
-	check( EditingObjects.Num() > 0 );
+	check( EditingObjs.Num() > 0 );
 
 	FFormatNamedArguments Args;
 	Args.Add( TEXT("ToolkitName"), GetBaseToolkitName() );
 
-	if( EditingObjects.Num() == 1 )
+	if( EditingObjs.Num() == 1 )
 	{
-		const UObject* EditingObject = EditingObjects[ 0 ];
+		const UObject* EditingObject = EditingObjs[ 0 ];
 		return FAssetEditorToolkit::GetToolTipTextForObject(EditingObject);
 	}
 	else
 	{
 		UClass* SharedBaseClass = NULL;
-		for( int32 x = 0; x < EditingObjects.Num(); ++x )
+		for( int32 x = 0; x < EditingObjs.Num(); ++x )
 		{
-			UObject* Obj = EditingObjects[ x ];
+			UObject* Obj = EditingObjs[ x ];
 			check( Obj );
 
 			UClass* ObjClass = Cast<UClass>(Obj);
-			if (ObjClass == NULL)
+			if (ObjClass == nullptr)
 			{
 				ObjClass = Obj->GetClass();
 			}
 			check( ObjClass );
 
 			// Initialize with the class of the first object we encounter.
-			if( SharedBaseClass == NULL )
+			if( SharedBaseClass == nullptr )
 			{
 				SharedBaseClass = ObjClass;
 			}
@@ -214,7 +214,7 @@ FText FSimpleAssetEditor::GetToolkitToolTipText() const
 			}
 		}
 
-		Args.Add( TEXT("NumberOfObjects"), EditingObjects.Num() );
+		Args.Add( TEXT("NumberOfObjects"), EditingObjs.Num() );
 		Args.Add( TEXT("ClassName"), FText::FromString( SharedBaseClass->GetName() ) );
 		return FText::Format( LOCTEXT("ToolkitTitle_EditingMultipleToolTip", "{NumberOfObjects} {ClassName} - {ToolkitName}"), Args );
 	}
diff --git a/Engine/Source/Programs/UnrealHeaderTool/Private/UHTMakefile/PropertyBaseArchiveProxy.h b/Engine/Source/Programs/UnrealHeaderTool/Private/UHTMakefile/PropertyBaseArchiveProxy.h
--- a/Engine/Source/Programs/UnrealHeaderTool/Private/UHTMakefile/PropertyBaseArchiveProxy.h
+++ b/Engine/Source/Programs/UnrealHeaderTool/Private/UHTMakefile/PropertyBaseArchiveProxy.h
@@ -57,34 +57,34 @@ inline FArchive& operator<<(FArchive& Ar, EArrayType::Type& ArrayType)
 	return Ar;
 }
 
-inline FArchive& operator<<(FArchive& Ar, ERefQualifier::Type& ArrayType)
+inline FArchive& operator<<(FArchive& Ar, ERefQualifier::Type& RefQualifier)
 {
 	if (Ar.IsLoading())
 	{
 		int32 Value;
 		Ar << Value;
-		ArrayType = (ERefQualifier::Type)Value;
+		RefQualifier = (ERefQualifier::Type)Value;
 	}
 	else if (Ar.IsSaving())
 	{
-		int32 Value = (int32)ArrayType;
+		int32 Value = (int32)RefQualifier;
 		Ar << Value;
 	}
 
 	return Ar;
 }
 
-inline FArchive& operator<<(FArchive& Ar, EPointerType::Type& ArrayType)
+inline FArchive& operator<<(FArchive& Ar, EPointerType::Type& PointerType)
 {
 	if (Ar.IsLoading())
 	{
 		int32 Value;
 		Ar << Value;
-		ArrayType = (EPointerType::Type)Value;
+		PointerType = (EPointerType::Type)Value;
 	}
 	else if (Ar.IsSaving())
 	{
-		int32 Value = (int32)ArrayType;
+		int32 Value = (int32)PointerType;
 		Ar << Value;
 	}
 
diff --git a/Engine/Source/Runtime/AIModule/Classes/AIController.h b/Engine/Source/Runtime/AIModule/Classes/AIController.h
--- a/Engine/Source/Runtime/AIModule/Classes/AIController.h
+++ b/Engine/Source/Runtime/AIModule/Classes/AIController.h
@@ -88,24 +88,24 @@ protected:
 	/** By default AI's logic gets stopped when controlled Pawn is unpossesed. Setting this flag to false
 	 *	will make AI logic persist past loosing controll over a pawn */
 	UPROPERTY(EditAnywhere, BlueprintReadWrite, Category = AI)
-		uint32 bStopAILogicOnUnposses : 1;
+	uint32 bStopAILogicOnUnposses : 1;
 
 public:
 	/** used for alternating LineOfSight traces */
 	UPROPERTY()
-		mutable uint32 bLOSflag : 1;
+	mutable uint32 bLOSflag : 1;
 
 	/** Skip extra line of sight traces to extremities of target being checked. */
 	UPROPERTY()
-		uint32 bSkipExtraLOSChecks : 1;
+	uint32 bSkipExtraLOSChecks : 1;
 
 	/** Is strafing allowed during movement? */
 	UPROPERTY()
-		uint32 bAllowStrafe : 1;
+	uint32 bAllowStrafe : 1;
 
 	/** Specifies if this AI wants its own PlayerState. */
 	UPROPERTY()
-		uint32 bWantsPlayerState : 1;
+	uint32 bWantsPlayerState : 1;
 
 	/** Copy Pawn rotation to ControlRotation, if there is no focus point. */
 	UPROPERTY(EditAnywhere, BlueprintReadWrite, Category = AI)
@@ -115,31 +115,31 @@ private_subobject:
 
 	/** Component used for moving along a path. */
 	DEPRECATED_FORGAME(4.6, "PathFollowingComponent should not be accessed directly, please use GetPathFollowingComponent() function instead. PathFollowingComponent will soon be private and your code will not compile.")
-		UPROPERTY(VisibleDefaultsOnly, Category = AI)
-		UPathFollowingComponent* PathFollowingComponent;
+	UPROPERTY(VisibleDefaultsOnly, Category = AI)
+	UPathFollowingComponent* PathFollowingComponent;
 
 public:
 
 	/** Component responsible for behaviors. */
 	UPROPERTY(BlueprintReadWrite, Category = AI)
-		UBrainComponent* BrainComponent;
+	UBrainComponent* BrainComponent;
 
 	DEPRECATED_FORGAME(4.8, "PerceptionComponent should not be accessed directly, please use GetAIPerceptionComponent() function instead. PerceptionComponent will soon be private and your code will not compile.")
-		UPROPERTY(VisibleDefaultsOnly, Category = AI)
-		UAIPerceptionComponent* PerceptionComponent;
+	UPROPERTY(VisibleDefaultsOnly, Category = AI)
+	UAIPerceptionComponent* PerceptionComponent;
 
 private_subobject:
 	DEPRECATED_FORGAME(4.6, "ActionsComp should not be accessed directly, please use GetActionsComp() function instead. ActionsComp will soon be private and your code will not compile.")
-		UPROPERTY(BlueprintReadOnly, Category = AI, meta = (AllowPrivateAccess = "true"))
-		UPawnActionsComponent* ActionsComp;
+	UPROPERTY(BlueprintReadOnly, Category = AI, meta = (AllowPrivateAccess = "true"))
+	UPawnActionsComponent* ActionsComp;
 
 protected:
 	/** blackboard */
 	UPROPERTY(BlueprintReadOnly, Category = AI, meta = (AllowPrivateAccess = "true"))
-		UBlackboardComponent* Blackboard;
+	UBlackboardComponent* Blackboard;
 
 	UPROPERTY()
-		UGameplayTasksComponent* CachedGameplayTasksComponent;
+	UGameplayTasksComponent* CachedGameplayTasksComponent;
 
 public:
 
@@ -147,15 +147,15 @@ public:
 
 	/** Event called when PossessedPawn is possessed by this controller. */
 	UFUNCTION(BlueprintImplementableEvent)
-		void OnPossess(APawn* PossessedPawn);
+	void OnPossess(APawn* PossessedPawn);
 
 	/** Gets triggered after given pawn has been unpossesed */
 	UFUNCTION(BlueprintImplementableEvent)
-		void OnUnpossess(APawn* UnpossessedPawn);
+	void OnUnpossess(APawn* UnpossessedPawn);
 
 	virtual void SetPawn(APawn* InPawn) override;
 
-	/** Makes AI go toward specified Goal actor (destination will be continuously updated)
+	/** Makes AI go toward specified Goal actor (destination will be continuously updated), aborts any active path following
 	 *  @param AcceptanceRadius - finish move if pawn gets close enough
 	 *  @param bStopOnOverlap - add pawn's radius to AcceptanceRadius
 	 *  @param bUsePathfinding - use navigation data to calculate path (otherwise it will go in straight line)
@@ -165,11 +165,11 @@ public:
 	 *	@note AcceptanceRadius has default value or -1 due to Header Parser not being able to recognize UPathFollowingComponent::DefaultAcceptanceRadius
 	 */
 	UFUNCTION(BlueprintCallable, Category = "AI|Navigation", Meta = (AdvancedDisplay = "bStopOnOverlap,bCanStrafe,bAllowPartialPath"))
-		EPathFollowingRequestResult::Type MoveToActor(AActor* Goal, float AcceptanceRadius = -1, bool bStopOnOverlap = true,
-			bool bUsePathfinding = true, bool bCanStrafe = true,
-			TSubclassOf<UNavigationQueryFilter> FilterClass = NULL, bool bAllowPartialPath = true);
+	EPathFollowingRequestResult::Type MoveToActor(AActor* Goal, float AcceptanceRadius = -1, bool bStopOnOverlap = true,
+		bool bUsePathfinding = true, bool bCanStrafe = true,
+		TSubclassOf<UNavigationQueryFilter> FilterClass = NULL, bool bAllowPartialPath = true);
 
-	/** Makes AI go toward specified Dest location
+	/** Makes AI go toward specified Dest location, aborts any active path following
 	 *  @param AcceptanceRadius - finish move if pawn gets close enough
 	 *  @param bStopOnOverlap - add pawn's radius to AcceptanceRadius
 	 *  @param bUsePathfinding - use navigation data to calculate path (otherwise it will go in straight line)
@@ -180,9 +180,9 @@ public:
 	 *	@note AcceptanceRadius has default value or -1 due to Header Parser not being able to recognize UPathFollowingComponent::DefaultAcceptanceRadius
 	 */
 	UFUNCTION(BlueprintCallable, Category = "AI|Navigation", Meta = (AdvancedDisplay = "bStopOnOverlap,bCanStrafe,bAllowPartialPath"))
-		EPathFollowingRequestResult::Type MoveToLocation(const FVector& Dest, float AcceptanceRadius = -1, bool bStopOnOverlap = true,
-			bool bUsePathfinding = true, bool bProjectDestinationToNavigation = false, bool bCanStrafe = true,
-			TSubclassOf<UNavigationQueryFilter> FilterClass = NULL, bool bAllowPartialPath = true);
+	EPathFollowingRequestResult::Type MoveToLocation(const FVector& Dest, float AcceptanceRadius = -1, bool bStopOnOverlap = true,
+		bool bUsePathfinding = true, bool bProjectDestinationToNavigation = false, bool bCanStrafe = true,
+		TSubclassOf<UNavigationQueryFilter> FilterClass = NULL, bool bAllowPartialPath = true);
 
 	/** Makes AI go toward specified destination */
 	EPathFollowingRequestResult::Type MoveTo(const FAIMoveRequest& MoveRequest);
@@ -190,25 +190,16 @@ public:
 	/** Fills pathfinding query for navigation system from given move request */
 	virtual bool PreparePathfinding(const FAIMoveRequest& MoveRequest, FPathFindingQuery& Query);
 
-	DEPRECATED(4.8, "This function is deprecated, please use version with MoveRequest parameter instead.")
-		virtual bool PreparePathfinding(FPathFindingQuery& Query, const FVector& Dest, AActor* Goal, bool bUsePathfinding = true, TSubclassOf<class UNavigationQueryFilter> FilterClass = NULL);
-
 	/** Executes pathfinding query and starts move request
 	 *  @return RequestID, or 0 when failed
 	 */
 	virtual FAIRequestID RequestPathAndMove(const FAIMoveRequest& MoveRequest, FPathFindingQuery& Query);
 
-	DEPRECATED(4.8, "This function is deprecated, please use version with MoveRequest parameter instead.")
-		virtual FAIRequestID RequestPathAndMove(FPathFindingQuery& Query, AActor* Goal, float AcceptanceRadius, bool bStopOnOverlap, FCustomMoveSharedPtr CustomData);
-
 	/** Handle move requests
 	 *  @return RequestID, or 0 when failed
 	 */
 	virtual FAIRequestID RequestMove(const FAIMoveRequest& MoveRequest, FNavPathSharedPtr Path);
 
-	DEPRECATED(4.8, "This function is deprecated, please use version with MoveRequest parameter instead.")
-		virtual FAIRequestID RequestMove(FNavPathSharedPtr Path, AActor* Goal = NULL, float AcceptanceRadius = UPathFollowingComponent::DefaultAcceptanceRadius, bool bStopOnOverlap = true, FCustomMoveSharedPtr CustomData = NULL);
-
 	/** if AI is currently moving due to request given by RequestToPause, then the move will be paused */
 	bool PauseMove(FAIRequestID RequestToPause);
 
@@ -219,6 +210,9 @@ public:
 	virtual void StopMovement() override;
 
 	/** Called on completing current movement request */
+	virtual void OnMoveCompleted(FAIRequestID RequestID, const FPathFollowingResult& Result);
+
+	DEPRECATED_FORGAME(4.13, "This function is now deprecated, please use version with EPathFollowingResultDetails parameter.")
 	virtual void OnMoveCompleted(FAIRequestID RequestID, EPathFollowingResult::Type Result);
 
 	/** Returns the Move Request ID for the current move */
@@ -226,30 +220,30 @@ public:
 
 	/** Blueprint notification that we've completed the current movement request */
 	UPROPERTY(BlueprintAssignable, meta = (DisplayName = "MoveCompleted"))
-		FAIMoveCompletedSignature ReceiveMoveCompleted;
+	FAIMoveCompletedSignature ReceiveMoveCompleted;
 
 	/** Returns status of path following */
 	UFUNCTION(BlueprintCallable, Category = "AI|Navigation")
-		EPathFollowingStatus::Type GetMoveStatus() const;
+	EPathFollowingStatus::Type GetMoveStatus() const;
 
 	/** Returns true if the current PathFollowingComponent's path is partial (does not reach desired destination). */
 	UFUNCTION(BlueprintCallable, Category = "AI|Navigation")
-		bool HasPartialPath() const;
+	bool HasPartialPath() const;
 
 	/** Returns position of current path segment's end. */
 	UFUNCTION(BlueprintCallable, Category = "AI|Navigation")
-		FVector GetImmediateMoveDestination() const;
+	FVector GetImmediateMoveDestination() const;
 
 	/** Updates state of movement block detection. */
 	UFUNCTION(BlueprintCallable, Category = "AI|Navigation")
-		void SetMoveBlockDetection(bool bEnable);
+	void SetMoveBlockDetection(bool bEnable);
 
 	/** Prepares path finding and path following components. */
 	virtual void InitNavigationControl(UPathFollowingComponent*& PathFollowingComp) override;
 
 	/** Starts executing behavior tree. */
 	UFUNCTION(BlueprintCallable, Category = "AI")
-		virtual bool RunBehaviorTree(UBehaviorTree* BTAsset);
+	virtual bool RunBehaviorTree(UBehaviorTree* BTAsset);
 
 	/**
 	 * Makes AI use the specified Blackboard asset & creates a Blackboard Component if one does not already exist.
@@ -258,7 +252,7 @@ public:
 	 * @return true if we successfully linked the blackboard asset to the blackboard component.
 	 */
 	UFUNCTION(BlueprintCallable, Category = "AI")
-		bool UseBlackboard(UBlackboardData* BlackboardAsset, UBlackboardComponent*& BlackboardComponent);
+	bool UseBlackboard(UBlackboardData* BlackboardAsset, UBlackboardComponent*& BlackboardComponent);
 
 	/** does this AIController allow given UBlackboardComponent sync data with it */
 	virtual bool ShouldSyncBlackboardWith(const UBlackboardComponent& OtherBlackboardComponent) const { return true; }
@@ -266,38 +260,38 @@ public:
 
 protected:
 	UFUNCTION(BlueprintImplementableEvent)
-		void OnUsingBlackBoard(UBlackboardComponent* BlackboardComp, UBlackboardData* BlackboardAsset);
+	void OnUsingBlackBoard(UBlackboardComponent* BlackboardComp, UBlackboardData* BlackboardAsset);
 
 	virtual bool InitializeBlackboard(UBlackboardComponent& BlackboardComp, UBlackboardData& BlackboardAsset);
 
 public:
 	/** Retrieve the final position that controller should be looking at. */
 	UFUNCTION(BlueprintCallable, Category = "AI")
-		FVector GetFocalPoint() const;
+	FVector GetFocalPoint() const;
 
 	FVector GetFocalPointForPriority(EAIFocusPriority::Type InPriority) const;
 
 	/** Retrieve the focal point this controller should focus to on given actor. */
 	UFUNCTION(BlueprintCallable, Category = "AI")
-		virtual FVector GetFocalPointOnActor(const AActor *Actor) const;
+	virtual FVector GetFocalPointOnActor(const AActor *Actor) const;
 
 	/** Set the position that controller should be looking at. */
 	UFUNCTION(BlueprintCallable, Category = "AI", meta = (DisplayName = "SetFocalPoint", Keywords = "focus"))
-		void K2_SetFocalPoint(FVector FP);
+	void K2_SetFocalPoint(FVector FP);
 
 	/** Set Focus for actor, will set FocalPoint as a result. */
 	UFUNCTION(BlueprintCallable, Category = "AI", meta = (DisplayName = "SetFocus"))
-		void K2_SetFocus(AActor* NewFocus);
+	void K2_SetFocus(AActor* NewFocus);
 
 	/** Get the focused actor. */
 	UFUNCTION(BlueprintCallable, Category = "AI")
-		AActor* GetFocusActor() const;
+	AActor* GetFocusActor() const;
 
 	FORCEINLINE AActor* GetFocusActorForPriority(EAIFocusPriority::Type InPriority) const { return FocusInformation.Priorities.IsValidIndex(InPriority) ? FocusInformation.Priorities[InPriority].Actor.Get() : nullptr; }
 
 	/** Clears Focus, will also clear FocalPoint as a result */
 	UFUNCTION(BlueprintCallable, Category = "AI", meta = (DisplayName = "ClearFocus"))
-		void K2_ClearFocus();
+	void K2_ClearFocus();
 
 
 	/**
diff --git a/Engine/Source/Runtime/AIModule/Classes/AITypes.h b/Engine/Source/Runtime/AIModule/Classes/AITypes.h
--- a/Engine/Source/Runtime/AIModule/Classes/AITypes.h
+++ b/Engine/Source/Runtime/AIModule/Classes/AITypes.h
@@ -472,7 +472,8 @@ struct AIMODULE_API FAIMoveRequest
 	FAIMoveRequest& SetCanStrafe(bool bStrafe) { bCanStrafe = bStrafe; return *this; }
 	FAIMoveRequest& SetStopOnOverlap(bool bStop) { bStopOnOverlap = bStop; return *this; }
 	FAIMoveRequest& SetAcceptanceRadius(float Radius) { AcceptanceRadius = Radius; return *this; }
-	FAIMoveRequest& SetUserData(const FCustomMoveSharedPtr& Data) { UserData = Data; return *this; }
+	FAIMoveRequest& SetUserData(const FCustomMoveSharedPtr& InUserData) { UserData = InUserData; return *this; }
+	FAIMoveRequest& SetUserFlags(int32 InUserFlags) { UserFlags = InUserFlags; return *this; }
 
 	/** the request should be either set up to move to a location, of go to a valid actor */
 	bool IsValid() const { return bInitialized && (!bMoveToActor || GoalActor); }
@@ -490,6 +491,7 @@ struct AIMODULE_API FAIMoveRequest
 	bool CanStopOnOverlap() const { return bStopOnOverlap; }
 	float GetAcceptanceRadius() const { return AcceptanceRadius; }
 	const FCustomMoveSharedPtr& GetUserData() const { return UserData; }
+	int32 GetUserFlags() const { return UserFlags; }
 
 	void SetGoalActor(const AActor* InGoalActor);
 	void SetGoalLocation(const FVector& InGoalLocation);
@@ -533,11 +535,9 @@ protected:
 	/** pathfollowing: required distance to goal to complete move */
 	float AcceptanceRadius;
 
-	/** pathfollowing: custom user data */
+	/** custom user data: structure */
 	FCustomMoveSharedPtr UserData;
 
-public:
-	// deprecated
-	DEPRECATED(4.8, "FAIMoveRequest::HasGoalActor is deprecated, please use FAIMoveRequest::IsMoveToActorRequest instead.")
-	bool HasGoalActor() const { return bMoveToActor; }
+	/** custom user data: flags */
+	int32 UserFlags;
 };
diff --git a/Engine/Source/Runtime/AIModule/Classes/Navigation/CrowdFollowingComponent.h b/Engine/Source/Runtime/AIModule/Classes/Navigation/CrowdFollowingComponent.h
--- a/Engine/Source/Runtime/AIModule/Classes/Navigation/CrowdFollowingComponent.h
+++ b/Engine/Source/Runtime/AIModule/Classes/Navigation/CrowdFollowingComponent.h
@@ -51,13 +51,13 @@ class AIMODULE_API UCrowdFollowingComponent : public UPathFollowingComponent, pu
 	// PathFollowingComponent BEGIN
 	virtual void Initialize() override;
 	virtual void Cleanup() override;
-	virtual void AbortMove(const FString& Reason, FAIRequestID RequestID = FAIRequestID::CurrentRequest, bool bResetVelocity = true, bool bSilent = false, uint8 MessageFlags = 0) override;
-	virtual void PauseMove(FAIRequestID RequestID = FAIRequestID::CurrentRequest, bool bResetVelocity = true) override;
+	virtual void AbortMove(const UObject& Instigator, FPathFollowingResultFlags::Type AbortFlags, FAIRequestID RequestID = FAIRequestID::CurrentRequest, EPathFollowingVelocityMode VelocityMode = EPathFollowingVelocityMode::Reset) override;
+	virtual void PauseMove(FAIRequestID RequestID = FAIRequestID::CurrentRequest, EPathFollowingVelocityMode VelocityMode = EPathFollowingVelocityMode::Reset) override;
 	virtual void ResumeMove(FAIRequestID RequestID = FAIRequestID::CurrentRequest) override;
 	virtual FVector GetMoveFocus(bool bAllowStrafe) const override;
 	virtual void OnLanded() override;
 	virtual void FinishUsingCustomLink(INavLinkCustomInterface* CustomNavLink) override;
-	virtual void OnPathFinished(EPathFollowingResult::Type Result) override;
+	virtual void OnPathFinished(const FPathFollowingResult& Result) override;
 	virtual void OnPathUpdated() override;
 	virtual void OnPathfindingQuery(FPathFindingQuery& Query) override;
 	virtual int32 GetCurrentPathElement() const override { return LastPathPolyIndex; }
