Copying //UE4/Dev-Framework to //UE4/Dev-Main (Source: //UE4/Dev-Framework @ 2964666)

diff --git a/Engine/Source/Runtime/AIModule/Classes/Navigation/PathFollowingComponent.h b/Engine/Source/Runtime/AIModule/Classes/Navigation/PathFollowingComponent.h
--- a/Engine/Source/Runtime/AIModule/Classes/Navigation/PathFollowingComponent.h
+++ b/Engine/Source/Runtime/AIModule/Classes/Navigation/PathFollowingComponent.h
@@ -53,15 +53,74 @@ namespace EPathFollowingResult
 		/** Aborted and stopped (failure) */
 		Aborted,
 
-		/** Aborted and replaced with new request */
-		Skipped,
+		/** DEPRECATED, use Aborted result instead */
+		Skipped_DEPRECATED UMETA(Hidden),
 
 		/** Request was invalid */
 		Invalid,
 	};
 }
 
-// left for now, will be removed soon! please use EPathFollowingStatus instead
+namespace FPathFollowingResultFlags
+{
+	typedef uint16 Type;
+
+	const Type None = 0;
+
+	/** Reached destination (EPathFollowingResult::Success) */
+	const Type Success = (1 << 0);
+
+	/** Movement was blocked (EPathFollowingResult::Blocked) */
+	const Type Blocked = (1 << 1);
+
+	/** Agent is not on path (EPathFollowingResult::OffPath) */
+	const Type OffPath = (1 << 2);
+
+	/** Aborted (EPathFollowingResult::Aborted) */
+	const Type ExternalCancel = (1 << 3);
+
+	/** External cancel: owner no longer wants to move */
+	const Type OwnerFinished = (1 << 4);
+
+	/** External cancel: path is no longer valid */
+	const Type InvalidPath = (1 << 5);
+
+	/** External cancel: unable to move */
+	const Type MovementStop = (1 << 6);
+
+	/** External cancel: new movement request was received */
+	const Type NewRequest = (1 << 7);
+
+	/** External cancel: blueprint MoveTo function was called */
+	const Type ForcedScript = (1 << 8);
+
+	/** Can be used to create project specific reasons */
+	const Type FirstGameplayFlagShift = 9;
+
+	const Type ExternalCancelFlagMask = ~(Success | Blocked | OffPath);
+
+	FString ToString(uint16 Value);
+}
+
+struct AIMODULE_API FPathFollowingResult
+{
+	FPathFollowingResultFlags::Type Flags;
+	TEnumAsByte<EPathFollowingResult::Type> Code;
+
+	FPathFollowingResult() : Flags(0), Code(EPathFollowingResult::Invalid)  {}
+	FPathFollowingResult(FPathFollowingResultFlags::Type InFlags);
+	FPathFollowingResult(EPathFollowingResult::Type ResultCode, FPathFollowingResultFlags::Type ExtraFlags);
+
+	bool HasFlag(FPathFollowingResultFlags::Type Flag) const { return (Flags & Flag) != 0; }
+
+	bool IsSuccess() const { return HasFlag(FPathFollowingResultFlags::Success); }
+	bool IsFailure() const { return !HasFlag(FPathFollowingResultFlags::Success); }
+	bool IsSkipped() const { return HasFlag(FPathFollowingResultFlags::ExternalCancel | FPathFollowingResultFlags::NewRequest); }
+	
+	FString ToString() const;
+};
+
+// DEPRECATED, will be removed with GetPathActionType function
 UENUM(BlueprintType)
 namespace EPathFollowingAction
 {
@@ -97,18 +156,22 @@ namespace EPathFollowingDebugTokens
 	};
 }
 
+// DEPRECATED, please use EPathFollowingResultDetails instead, will be removed with deprecated override of AbortMove function
 namespace EPathFollowingMessage
 {
 	enum Type
 	{
-		/** Aborted because no path was found */
 		NoPath,
-
-		/** Aborted because another request came in */
 		OtherRequest,
 	};
 }
 
+enum class EPathFollowingVelocityMode : uint8
+{
+	Reset,
+	Keep,
+};
+
 UCLASS(config=Engine)
 class AIMODULE_API UPathFollowingComponent : public UActorComponent, public IAIResourceInterface
 {
@@ -117,12 +180,13 @@ class AIMODULE_API UPathFollowingComponent : public UActorComponent, public IAIR
 	DECLARE_DELEGATE_TwoParams(FPostProcessMoveSignature, UPathFollowingComponent* /*comp*/, FVector& /*velocity*/);
 	DECLARE_DELEGATE_OneParam(FRequestCompletedSignature, EPathFollowingResult::Type /*Result*/);
 	DECLARE_MULTICAST_DELEGATE_TwoParams(FMoveCompletedSignature, FAIRequestID /*RequestID*/, EPathFollowingResult::Type /*Result*/);
+	DECLARE_MULTICAST_DELEGATE_TwoParams(FMoveComplete, FAIRequestID /*RequestID*/, const FPathFollowingResult& /*Result*/);
 
 	/** delegate for modifying path following velocity */
 	FPostProcessMoveSignature PostProcessMove;
 
 	/** delegate for move completion notify */
-	FMoveCompletedSignature OnMoveFinished;
+	FMoveComplete OnRequestFinished;
 
 	//~ Begin UActorComponent Interface
 	virtual void TickComponent(float DeltaTime, enum ELevelTick TickType, FActorComponentTickFunction *ThisTickFunction) override;
@@ -137,33 +201,27 @@ class AIMODULE_API UPathFollowingComponent : public UActorComponent, public IAIR
 	/** updates cached pointers to relevant owner's components */
 	virtual void UpdateCachedComponents();
 
-	/** start movement along path
-	 *  @returns request ID or 0 when failed */
-	virtual FAIRequestID RequestMove(FNavPathSharedPtr Path, FRequestCompletedSignature OnComplete, const AActor* DestinationActor = NULL, float AcceptanceRadius = UPathFollowingComponent::DefaultAcceptanceRadius, bool bStopOnOverlap = true, FCustomMoveSharedPtr GameData = NULL);
+	/** start movement along path */
+	virtual FAIRequestID RequestMove(const FAIMoveRequest& RequestData, FNavPathSharedPtr InPath);
 
-	/** start movement along path
-	 *  @returns request ID or 0 when failed */
-	FORCEINLINE FAIRequestID RequestMove(FNavPathSharedPtr InPath, const AActor* InDestinationActor = NULL, float InAcceptanceRadius = UPathFollowingComponent::DefaultAcceptanceRadius, bool InStopOnOverlap = true, FCustomMoveSharedPtr InGameData = NULL)
-	{
-		return RequestMove(InPath, UnboundRequestDelegate, InDestinationActor, InAcceptanceRadius, InStopOnOverlap, InGameData);
-	}
+	/** aborts following path */
+	virtual void AbortMove(const UObject& Instigator, FPathFollowingResultFlags::Type AbortFlags, FAIRequestID RequestID = FAIRequestID::CurrentRequest, EPathFollowingVelocityMode VelocityMode = EPathFollowingVelocityMode::Reset);
 
-	/** aborts following path
-	 *  @param RequestID - request to abort, 0 = current
-	 *  @param bResetVelocity - try to stop movement component
-	 *  @param bSilent - finish with Skipped result instead of Aborted */
-	virtual void AbortMove(const FString& Reason, FAIRequestID RequestID = FAIRequestID::CurrentRequest, bool bResetVelocity = true, bool bSilent = false, uint8 MessageFlags = 0);
+	/** create new request and finish it immediately (e.g. already at goal) */
+	void RequestMoveWithImmediateFinish(EPathFollowingResult::Type Result, EPathFollowingVelocityMode VelocityMode = EPathFollowingVelocityMode::Reset);
 
 	/** pause path following
 	*  @param RequestID - request to pause, FAIRequestID::CurrentRequest means pause current request, regardless of its ID */
-	virtual void PauseMove(FAIRequestID RequestID = FAIRequestID::CurrentRequest, bool bResetVelocity = true);
+	virtual void PauseMove(FAIRequestID RequestID = FAIRequestID::CurrentRequest, EPathFollowingVelocityMode VelocityMode = EPathFollowingVelocityMode::Reset);
 
 	/** resume path following
-	*  @param RequestID - request to resume, FAIRequestID::CurrentRequest means restor current request, regardless of its ID*/
+	*  @param RequestID - request to resume, FAIRequestID::CurrentRequest means restor current request, regardless of its ID */
 	virtual void ResumeMove(FAIRequestID RequestID = FAIRequestID::CurrentRequest);
 
 	/** notify about finished movement */
-	virtual void OnPathFinished(EPathFollowingResult::Type Result);
+	virtual void OnPathFinished(const FPathFollowingResult& Result);
+
+	FORCEINLINE void OnPathFinished(EPathFollowingResult::Type ResultCode, uint16 ExtraResultFlags) { OnPathFinished(FPathFollowingResult(ResultCode, ExtraResultFlags)); }
 
 	/** notify about finishing move along current path segment */
 	virtual void OnSegmentFinished();
@@ -210,7 +268,7 @@ class AIMODULE_API UPathFollowingComponent : public UActorComponent, public IAIR
 	/** set threshold for precise reach tests in intermediate goals (minimal test radius)  */
 	void SetPreciseReachThreshold(float AgentRadiusMultiplier, float AgentHalfHeightMultiplier);
 
-	/** set status of last requested move */
+	/** set status of last requested move, works only in Idle state */
 	void SetLastMoveAtGoal(bool bFinishedAtGoal);
 
 	/** @returns estimated cost of unprocessed path segments
@@ -236,14 +294,6 @@ class AIMODULE_API UPathFollowingComponent : public UActorComponent, public IAIR
 	FORCEINLINE FBasedPosition GetCurrentTargetLocationBased() const { return CurrentDestination; }
 	FVector GetCurrentDirection() const;
 
-	/** will be deprecated soon, please use AIController.GetMoveStatus instead! */
-	UFUNCTION(BlueprintCallable, Category="AI|Components|PathFollowing")
-	EPathFollowingAction::Type GetPathActionType() const;
-
-	/** will be deprecated soon, please use AIController.GetImmediateMoveDestination instead! */
-	UFUNCTION(BlueprintCallable, Category="AI|Components|PathFollowing")
-	FVector GetPathDestination() const;
-
 	FORCEINLINE const FNavPathSharedPtr GetPath() const { return Path; }
 	FORCEINLINE bool HasValidPath() const { return Path.IsValid() && Path->IsValid(); }
 	bool HasDirectPath() const;
@@ -293,15 +343,43 @@ class AIMODULE_API UPathFollowingComponent : public UActorComponent, public IAIR
 	virtual bool IsResourceLocked() const override;
 	// IAIResourceInterface end
 
+	/** path observer */
 	void OnPathEvent(FNavigationPath* InPath, ENavPathEvent::Type Event);
 
 	/** helper function for sending a path for visual log */
 	static void LogPathHelper(const AActor* LogOwner, FNavPathSharedPtr InLogPath, const AActor* LogGoalActor);
 	static void LogPathHelper(const AActor* LogOwner, FNavigationPath* InLogPath, const AActor* LogGoalActor);
 
-	DEPRECATED_FORGAME(4.12, "This function is now deprecated and replaced with HandlePathUpdateEvent. Receiving new path pointer for the same move request is no longer supported, please either update data within current path and call FNavigationPath::DoneUpdating or start new move request.")
+	DEPRECATED(4.12, "This function is now deprecated and replaced with HandlePathUpdateEvent. Receiving new path pointer for the same move request is no longer supported, please either update data within current path and call FNavigationPath::DoneUpdating or start new move request.")
 	virtual bool UpdateMove(FNavPathSharedPtr Path, FAIRequestID RequestID = FAIRequestID::CurrentRequest);
 
+	DEPRECATED(4.13, "This function is now deprecated, please use version with FAIMoveRequest parameter instead. Any observers needs to register with OnRequestFinished mutlicast delegate now.")
+	virtual FAIRequestID RequestMove(FNavPathSharedPtr Path, FRequestCompletedSignature OnComplete, const AActor* DestinationActor = NULL, float AcceptanceRadius = UPathFollowingComponent::DefaultAcceptanceRadius, bool bStopOnOverlap = true, FCustomMoveSharedPtr GameData = NULL);
+
+	DEPRECATED(4.13, "This function is now deprecated, please use version with FAIMoveRequest parameter instead.")
+	FAIRequestID RequestMove(FNavPathSharedPtr InPath, const AActor* InDestinationActor = NULL, float InAcceptanceRadius = UPathFollowingComponent::DefaultAcceptanceRadius, bool InStopOnOverlap = true, FCustomMoveSharedPtr InGameData = NULL);
+
+	DEPRECATED(4.13, "This function is now deprecated, please use version with EPathFollowingResultDetails parameter instead.")
+	virtual void AbortMove(const FString& Reason, FAIRequestID RequestID = FAIRequestID::CurrentRequest, bool bResetVelocity = true, bool bSilent = false, uint8 MessageFlags = 0);
+
+	DEPRECATED(4.13, "This function is now deprecated, please use version with EPathFollowingVelocityMode parameter instead.")
+	virtual void PauseMove(FAIRequestID RequestID = FAIRequestID::CurrentRequest, bool bResetVelocity = true);
+
+	UFUNCTION(BlueprintCallable, Category="AI|Components|PathFollowing", meta = (DeprecatedFunction, DeprecationMessage = "This function is now deprecated, please use AIController.GetMoveStatus instead"))
+	EPathFollowingAction::Type GetPathActionType() const;
+
+	UFUNCTION(BlueprintCallable, Category="AI|Components|PathFollowing", meta = (DeprecatedFunction, DeprecationMessage = "This function is now deprecated, please use AIController.GetImmediateMoveDestination instead"))
+	FVector GetPathDestination() const;
+
+	DEPRECATED(4.13, "This function is now deprecated, please version with FPathFollowingResult parameter instead.")
+	virtual void OnPathFinished(EPathFollowingResult::Type Result);
+
+	DEPRECATED(4.13, "This function is now deprecated and no longer supported.")
+	int32 OptimizeSegmentVisibility(int32 StartIndex) { return StartIndex + 1; }
+
+	// This delegate is now deprecated, please use OnRequestFinished instead
+	FMoveCompletedSignature OnMoveFinished_DEPRECATED;
+
 protected:
 
 	/** associated movement component */
@@ -339,9 +417,6 @@ protected:
 	/** game specific data */
 	FCustomMoveSharedPtr GameData;
 
-	/** current request observer */
-	FRequestCompletedSignature OnRequestFinished;
-
 	/** destination actor. Use SetDestinationActor to set this */
 	TWeakObjectPtr<AActor> DestinationActor;
 
@@ -360,10 +435,6 @@ protected:
 	/** timestamp of path update when movement was paused */
 	float PathTimeWhenPaused;
 
-	/** set when paths simplification using visibility tests are needed  (disabled by default because of performance) */
-	UPROPERTY(config)
-	uint32 bUseVisibilityTestsSimplification : 1;
-
 	/** increase acceptance radius with agent's radius */
 	uint32 bStopOnOverlap : 1;
 
@@ -458,8 +529,6 @@ protected:
 	bool HasReachedCurrentTarget(const FVector& CurrentLocation) const;
 
 	/** check if moving agent has reached goal defined by cylinder */
-	DEPRECATED(4.8, "Please use override with AgentRadiusMultiplier instead of this.")
-	bool HasReachedInternal(const FVector& GoalLocation, float GoalRadius, float GoalHalfHeight, const FVector& AgentLocation, float RadiusThreshold, bool bSuccessOnRadiusOverlap) const;
 	bool HasReachedInternal(const FVector& GoalLocation, float GoalRadius, float GoalHalfHeight, const FVector& AgentLocation, float RadiusThreshold, float AgentRadiusMultiplier) const;
 
 	/** check if agent is on path */
@@ -471,9 +540,11 @@ protected:
 	/** switch to next segment on path */
 	FORCEINLINE void SetNextMoveSegment() { SetMoveSegment(GetNextPathIndex()); }
 
-	FORCEINLINE static uint32 GetNextRequestId() { return NextRequestId++; }
+	/** assign new request Id */
 	FORCEINLINE void StoreRequestId() { CurrentRequestId = UPathFollowingComponent::GetNextRequestId(); }
 
+	FORCEINLINE static uint32 GetNextRequestId() { return NextRequestId++; }
+
 	/** Checks if this PathFollowingComponent is already on path, and
 	*	if so determines index of next path point
 	*	@return what PathFollowingComponent thinks should be next path point. INDEX_NONE if given path is invalid
@@ -483,11 +554,6 @@ protected:
 	/** @return index of path point, that should be target of current move segment */
 	virtual int32 DetermineCurrentTargetPathPoint(int32 StartIndex);
 
-	/** Visibility tests to skip some path points if possible
-	@param NextSegmentStartIndex Selected next segment to follow
-	@return  better path segment to follow, to use as MoveSegmentEndIndex */
-	int32 OptimizeSegmentVisibility(int32 StartIndex);
-
 	/** check if movement component is valid or tries to grab one from owner 
 	 *	@param bForce results in looking for owner's movement component even if pointer to one is already cached */
 	virtual bool UpdateMovementComponent(bool bForce = false);
@@ -529,9 +595,6 @@ private:
 	/** timer handle for OnWaitingPathTimeout function */
 	FTimerHandle WaitingForPathTimer;
 
-	/** empty delegate for RequestMove */
-	static FRequestCompletedSignature UnboundRequestDelegate;
-
 public:
 	/** special float constant to symbolize "use default value". This does not contain 
 	 *	value to be used, it's used to detect the fact that it's requested, and 
diff --git a/Engine/Source/Runtime/AIModule/Classes/Perception/AISense.h b/Engine/Source/Runtime/AIModule/Classes/Perception/AISense.h
--- a/Engine/Source/Runtime/AIModule/Classes/Perception/AISense.h
+++ b/Engine/Source/Runtime/AIModule/Classes/Perception/AISense.h
@@ -141,10 +141,10 @@ protected:
 	/** To be called only for BP-generated classes */
 	void ForceSenseID(FAISenseID SenseID);
 
-	DEPRECATED(4.12, "Function removed, use AISenseConfig.GetDebugColor() instead.")
+	DEPRECATED(4.13, "Function removed, use AISenseConfig.GetDebugColor() instead.")
 	FColor GetDebugColor() const { return FColor::White; }
-	DEPRECATED(4.12, "Function removed, use AISenseConfig.GetSenseName() instead.")
+	DEPRECATED(4.13, "Function removed, use AISenseConfig.GetSenseName() instead.")
 	FString GetDebugName() const { return FString(); }
-	DEPRECATED(4.12, "Function removed, use AISenseConfig.DescribeSelfToGameplayDebugger() instead.")
+	DEPRECATED(4.13, "Function removed, use AISenseConfig.DescribeSelfToGameplayDebugger() instead.")
 	virtual FString GetDebugLegend() const { return FString(); };
 };
diff --git a/Engine/Source/Runtime/AIModule/Classes/Tasks/AITask_MoveTo.h b/Engine/Source/Runtime/AIModule/Classes/Tasks/AITask_MoveTo.h
--- a/Engine/Source/Runtime/AIModule/Classes/Tasks/AITask_MoveTo.h
+++ b/Engine/Source/Runtime/AIModule/Classes/Tasks/AITask_MoveTo.h
@@ -91,5 +91,5 @@ protected:
 	virtual void OnPathEvent(FNavigationPath* InPath, ENavPathEvent::Type Event);
 
 	/** event from path following */
-	virtual void OnRequestFinished(FAIRequestID RequestID, EPathFollowingResult::Type Result);
+	virtual void OnRequestFinished(FAIRequestID RequestID, const FPathFollowingResult& Result);
 };
diff --git a/Engine/Source/Runtime/AIModule/Private/AIController.cpp b/Engine/Source/Runtime/AIModule/Private/AIController.cpp
--- a/Engine/Source/Runtime/AIModule/Private/AIController.cpp
+++ b/Engine/Source/Runtime/AIModule/Private/AIController.cpp
@@ -40,7 +40,7 @@ AAIController::AAIController(const FObjectInitializer& ObjectInitializer)
 {
 	bSetControlRotationFromPawnOrientation = true;
 	PathFollowingComponent = CreateDefaultSubobject<UPathFollowingComponent>(TEXT("PathFollowingComponent"));
-	PathFollowingComponent->OnMoveFinished.AddUObject(this, &AAIController::OnMoveCompleted);
+	PathFollowingComponent->OnRequestFinished.AddUObject(this, &AAIController::OnMoveCompleted);
 
 	ActionsComp = CreateDefaultSubobject<UPawnActionsComponent>("ActionsComp");
 
@@ -95,7 +95,7 @@ void AAIController::Reset()
 
 	if (PathFollowingComponent)
 	{
-		PathFollowingComponent->AbortMove(TEXT("controller reset"));
+		PathFollowingComponent->AbortMove(*this, FPathFollowingResultFlags::OwnerFinished);
 	}
 }
 
@@ -562,6 +562,12 @@ void AAIController::InitNavigationControl(UPathFollowingComponent*& PathFollowin
 
 EPathFollowingRequestResult::Type AAIController::MoveToActor(AActor* Goal, float AcceptanceRadius, bool bStopOnOverlap, bool bUsePathfinding, bool bCanStrafe, TSubclassOf<UNavigationQueryFilter> FilterClass, bool bAllowPartialPaths)
 {
+	// abort active movement to keep only one request running
+	if (PathFollowingComponent && PathFollowingComponent->GetStatus() != EPathFollowingStatus::Idle)
+	{
+		PathFollowingComponent->AbortMove(*this, FPathFollowingResultFlags::ForcedScript | FPathFollowingResultFlags::NewRequest);
+	}
+
 	FAIMoveRequest MoveReq(Goal);
 	MoveReq.SetUsePathfinding(bUsePathfinding);
 	MoveReq.SetAllowPartialPath(bAllowPartialPaths);
@@ -575,6 +581,12 @@ EPathFollowingRequestResult::Type AAIController::MoveToActor(AActor* Goal, float
 
 EPathFollowingRequestResult::Type AAIController::MoveToLocation(const FVector& Dest, float AcceptanceRadius, bool bStopOnOverlap, bool bUsePathfinding, bool bProjectDestinationToNavigation, bool bCanStrafe, TSubclassOf<UNavigationQueryFilter> FilterClass, bool bAllowPartialPaths)
 {
+	// abort active movement to keep only one request running
+	if (PathFollowingComponent && PathFollowingComponent->GetStatus() != EPathFollowingStatus::Idle)
+	{
+		PathFollowingComponent->AbortMove(*this, FPathFollowingResultFlags::ForcedScript | FPathFollowingResultFlags::NewRequest);
+	}
+
 	FAIMoveRequest MoveReq(Dest);
 	MoveReq.SetUsePathfinding(bUsePathfinding);
 	MoveReq.SetAllowPartialPath(bAllowPartialPaths);
@@ -589,6 +601,9 @@ EPathFollowingRequestResult::Type AAIController::MoveToLocation(const FVector& D
 
 EPathFollowingRequestResult::Type AAIController::MoveTo(const FAIMoveRequest& MoveRequest)
 {
+	// both MoveToActor and MoveToLocation can be called from blueprints/script and should keep only single movement request at the same time.
+	// this function is entry point of all movement mechanics - do NOT abort in here, since movement may be handled by AITasks, which support stacking 
+
 	SCOPE_CYCLE_COUNTER(STAT_MoveTo);
 	UE_VLOG(this, LogAINavigation, Log, TEXT("MoveTo: %s"), *MoveRequest.ToString());
 
@@ -598,6 +613,12 @@ EPathFollowingRequestResult::Type AAIController::MoveTo(const FAIMoveRequest& Mo
 		return EPathFollowingRequestResult::Failed;
 	}
 
+	if (PathFollowingComponent == nullptr)
+	{
+		UE_VLOG(this, LogAINavigation, Error, TEXT("MoveTo request failed due missing PathFollowingComponent"));
+		return EPathFollowingRequestResult::Failed;
+	}
+
 	EPathFollowingRequestResult::Type Result = EPathFollowingRequestResult::Failed;
 	bool bCanRequestMove = true;
 	bool bAlreadyAtGoal = false;
@@ -626,28 +647,17 @@ EPathFollowingRequestResult::Type AAIController::MoveTo(const FAIMoveRequest& Mo
 			MoveRequest.UpdateGoalLocation(ProjectedLocation.Location);
 		}
 
-		bAlreadyAtGoal = bCanRequestMove && PathFollowingComponent &&
-			PathFollowingComponent->HasReached(MoveRequest.GetGoalLocation(), MoveRequest.GetAcceptanceRadius(), !MoveRequest.CanStopOnOverlap());
+		bAlreadyAtGoal = bCanRequestMove && PathFollowingComponent->HasReached(MoveRequest.GetGoalLocation(), MoveRequest.GetAcceptanceRadius(), !MoveRequest.CanStopOnOverlap());
 	}
 	else 
 	{
-		bAlreadyAtGoal = bCanRequestMove && PathFollowingComponent &&
-			PathFollowingComponent->HasReached(*MoveRequest.GetGoalActor(), MoveRequest.GetAcceptanceRadius(), !MoveRequest.CanStopOnOverlap());
+		bAlreadyAtGoal = bCanRequestMove && PathFollowingComponent->HasReached(*MoveRequest.GetGoalActor(), MoveRequest.GetAcceptanceRadius(), !MoveRequest.CanStopOnOverlap());
 	}
 
 	if (bAlreadyAtGoal)
 	{
-		UE_VLOG(this, LogAINavigation, Log, TEXT("MoveToActor: already at goal!"));
-
-		if (PathFollowingComponent)
-		{
-			// make sure previous move request gets aborted
-			PathFollowingComponent->AbortMove(TEXT("Aborting move due to new move request finishing with AlreadyAtGoal"), FAIRequestID::AnyRequest);
-
-			PathFollowingComponent->SetLastMoveAtGoal(true);
-		}
-
-		OnMoveCompleted(FAIRequestID::CurrentRequest, EPathFollowingResult::Success);
+		UE_VLOG(this, LogAINavigation, Log, TEXT("MoveTo: already at goal!"));
+		PathFollowingComponent->RequestMoveWithImmediateFinish(EPathFollowingResult::Success);
 		Result = EPathFollowingRequestResult::AlreadyAtGoal;
 	}
 	else if (bCanRequestMove)
@@ -665,12 +675,7 @@ EPathFollowingRequestResult::Type AAIController::MoveTo(const FAIMoveRequest& Mo
 
 	if (Result == EPathFollowingRequestResult::Failed)
 	{
-		if (PathFollowingComponent)
-		{
-			PathFollowingComponent->SetLastMoveAtGoal(false);
-		}
-
-		OnMoveCompleted(FAIRequestID::InvalidRequest, EPathFollowingResult::Invalid);
+		PathFollowingComponent->RequestMoveWithImmediateFinish(EPathFollowingResult::Invalid);
 	}
 
 	return Result;
@@ -681,28 +686,17 @@ FAIRequestID AAIController::RequestMove(const FAIMoveRequest& MoveRequest, FNavP
 	uint32 RequestID = FAIRequestID::InvalidRequest;
 	if (PathFollowingComponent)
 	{
-		RequestID = PathFollowingComponent->RequestMove(Path, MoveRequest.GetGoalActor(), MoveRequest.GetAcceptanceRadius(), MoveRequest.CanStopOnOverlap(), MoveRequest.GetUserData());
+		RequestID = PathFollowingComponent->RequestMove(MoveRequest, Path);
 	}
 
 	return RequestID;
 }
 
-// deprecated
-FAIRequestID AAIController::RequestMove(FNavPathSharedPtr Path, AActor* Goal, float AcceptanceRadius, bool bStopOnOverlap, FCustomMoveSharedPtr CustomData)
-{
-	FAIMoveRequest MoveReq(Goal);
-	MoveReq.SetAcceptanceRadius(AcceptanceRadius);
-	MoveReq.SetStopOnOverlap(bStopOnOverlap);
-	MoveReq.SetUserData(CustomData);
-
-	return RequestMove(MoveReq, Path);
-}
-
 bool AAIController::PauseMove(FAIRequestID RequestToPause)
 {
 	if (PathFollowingComponent != NULL && RequestToPause.IsEquivalent(PathFollowingComponent->GetCurrentRequestId()))
 	{
-		PathFollowingComponent->PauseMove(RequestToPause);
+		PathFollowingComponent->PauseMove(RequestToPause, EPathFollowingVelocityMode::Reset);
 		return true;
 	}
 	return false;
@@ -720,8 +714,7 @@ bool AAIController::ResumeMove(FAIRequestID RequestToResume)
 
 void AAIController::StopMovement()
 {
-	UE_VLOG(this, LogAINavigation, Log, TEXT("AAIController::StopMovement: %s STOP MOVEMENT"), *GetNameSafe(GetPawn()));
-	PathFollowingComponent->AbortMove(TEXT("StopMovement"));
+	PathFollowingComponent->AbortMove(*this, FPathFollowingResultFlags::MovementStop);
 }
 
 bool AAIController::PreparePathfinding(const FAIMoveRequest& MoveRequest, FPathFindingQuery& Query)
@@ -770,25 +763,6 @@ bool AAIController::PreparePathfinding(const FAIMoveRequest& MoveRequest, FPathF
 	return false;
 }
 
-// deprecated
-bool AAIController::PreparePathfinding(FPathFindingQuery& Query, const FVector& Dest, AActor* Goal, bool bUsePathfinding, TSubclassOf<class UNavigationQueryFilter> FilterClass)
-{
-	if (Goal)
-	{
-		FAIMoveRequest MoveReq(Goal);
-		MoveReq.SetUsePathfinding(bUsePathfinding);
-		MoveReq.SetNavigationFilter(FilterClass);
-
-		return PreparePathfinding(MoveReq, Query);
-	}
-
-	FAIMoveRequest MoveReq(Dest);
-	MoveReq.SetUsePathfinding(bUsePathfinding);
-	MoveReq.SetNavigationFilter(FilterClass);
-
-	return PreparePathfinding(MoveReq, Query);
-}
-
 FAIRequestID AAIController::RequestPathAndMove(const FAIMoveRequest& MoveRequest, FPathFindingQuery& Query)
 {
 	SCOPE_CYCLE_COUNTER(STAT_AI_Overall);
@@ -825,17 +799,6 @@ FAIRequestID AAIController::RequestPathAndMove(const FAIMoveRequest& MoveRequest
 	return RequestID;
 }
 
-// deprecated
-FAIRequestID AAIController::RequestPathAndMove(FPathFindingQuery& Query, AActor* Goal, float AcceptanceRadius, bool bStopOnOverlap, FCustomMoveSharedPtr CustomData)
-{
-	FAIMoveRequest MoveReq(Goal);
-	MoveReq.SetAcceptanceRadius(AcceptanceRadius);
-	MoveReq.SetStopOnOverlap(bStopOnOverlap);
-	MoveReq.SetUserData(CustomData);
-
-	return RequestPathAndMove(MoveReq, Query);
-}
-
 EPathFollowingStatus::Type AAIController::GetMoveStatus() const
 {
 	return (PathFollowingComponent) ? PathFollowingComponent->GetStatus() : EPathFollowingStatus::Idle;
@@ -864,9 +827,15 @@ void AAIController::SetMoveBlockDetection(bool bEnable)
 	}
 }
 
+void AAIController::OnMoveCompleted(FAIRequestID RequestID, const FPathFollowingResult& Result)
+{
+	ReceiveMoveCompleted.Broadcast(RequestID, Result.Code);
+	OnMoveCompleted(RequestID, Result.Code);
+}
+
 void AAIController::OnMoveCompleted(FAIRequestID RequestID, EPathFollowingResult::Type Result)
 {
-	ReceiveMoveCompleted.Broadcast(RequestID, Result);
+	// deprecated
 }
 
 bool AAIController::RunBehaviorTree(UBehaviorTree* BTAsset)
diff --git a/Engine/Source/Runtime/AIModule/Private/AITypes.cpp b/Engine/Source/Runtime/AIModule/Private/AITypes.cpp
--- a/Engine/Source/Runtime/AIModule/Private/AITypes.cpp
+++ b/Engine/Source/Runtime/AIModule/Private/AITypes.cpp
@@ -160,7 +160,8 @@ FAIMoveRequest::FAIMoveRequest() :
 	GoalActor(nullptr), GoalLocation(FAISystem::InvalidLocation), FilterClass(nullptr),
 	bInitialized(false), bMoveToActor(false),
 	bUsePathfinding(true), bAllowPartialPath(true), bProjectGoalOnNavigation(true),
-	bStopOnOverlap(true), bCanStrafe(false)
+	bStopOnOverlap(true), bCanStrafe(false),
+	UserFlags(0)
 {
 	AcceptanceRadius = UPathFollowingComponent::DefaultAcceptanceRadius;
 }
@@ -169,7 +170,8 @@ FAIMoveRequest::FAIMoveRequest(const AActor* InGoalActor) :
 	GoalActor(const_cast<AActor*>(InGoalActor)), GoalLocation(FAISystem::InvalidLocation), FilterClass(nullptr),
 	bInitialized(true), bMoveToActor(true),
 	bUsePathfinding(true), bAllowPartialPath(true), bProjectGoalOnNavigation(true),
-	bStopOnOverlap(true), bCanStrafe(false)
+	bStopOnOverlap(true), bCanStrafe(false),
+	UserFlags(0)
 {
 	AcceptanceRadius = UPathFollowingComponent::DefaultAcceptanceRadius;
 }
@@ -178,7 +180,8 @@ FAIMoveRequest::FAIMoveRequest(const FVector& InGoalLocation) :
 	GoalActor(nullptr), GoalLocation(InGoalLocation), FilterClass(nullptr),
 	bInitialized(true), bMoveToActor(false),
 	bUsePathfinding(true), bAllowPartialPath(true), bProjectGoalOnNavigation(true),
-	bStopOnOverlap(true), bCanStrafe(false)
+	bStopOnOverlap(true), bCanStrafe(false),
+	UserFlags(0)
 {
 	AcceptanceRadius = UPathFollowingComponent::DefaultAcceptanceRadius;
 }
diff --git a/Engine/Source/Runtime/AIModule/Private/Actions/PawnAction_Move.cpp b/Engine/Source/Runtime/AIModule/Private/Actions/PawnAction_Move.cpp
--- a/Engine/Source/Runtime/AIModule/Private/Actions/PawnAction_Move.cpp
+++ b/Engine/Source/Runtime/AIModule/Private/Actions/PawnAction_Move.cpp
@@ -216,7 +216,7 @@ EPawnActionAbortState::Type UPawnAction_Move::PerformAbort(EAIForceParam::Type S
 
 	if (MyController && MyController->GetPathFollowingComponent())
 	{
-		MyController->GetPathFollowingComponent()->AbortMove(TEXT("BehaviorTree abort"), RequestID);
+		MyController->GetPathFollowingComponent()->AbortMove(*this, FPathFollowingResultFlags::OwnerFinished, RequestID);
 	}
 
 	return Super::PerformAbort(ShouldForce);
@@ -224,7 +224,7 @@ EPawnActionAbortState::Type UPawnAction_Move::PerformAbort(EAIForceParam::Type S
 
 void UPawnAction_Move::HandleAIMessage(UBrainComponent*, const FAIMessage& Message)
 {
-	if (Message.MessageName == UBrainComponent::AIMessage_MoveFinished && Message.HasFlag(EPathFollowingMessage::OtherRequest))
+	if (Message.MessageName == UBrainComponent::AIMessage_MoveFinished && Message.HasFlag(FPathFollowingResultFlags::NewRequest))
 	{
 		// move was aborted by another request from different action, don't finish yet
 		return;
@@ -366,8 +366,8 @@ bool UPawnAction_Move::CheckAlreadyAtGoal(AAIController& Controller, const FVect
 	const bool bAlreadyAtGoal = Controller.GetPathFollowingComponent()->HasReached(TestLocation, Radius);
 	if (bAlreadyAtGoal)
 	{
-		Controller.GetPathFollowingComponent()->AbortMove(TEXT("Aborting move due to new move request finishing with AlreadyAtGoal"), FAIRequestID::AnyRequest, true, false, EPathFollowingMessage::OtherRequest);
-		Controller.GetPathFollowingComponent()->SetLastMoveAtGoal(true);
+		UE_VLOG(&Controller, LogPawnAction, Log, TEXT("New move request already at goal, aborting active movement"));
+		Controller.GetPathFollowingComponent()->RequestMoveWithImmediateFinish(EPathFollowingResult::Success);
 	}
 
 	return bAlreadyAtGoal;
@@ -378,8 +378,8 @@ bool UPawnAction_Move::CheckAlreadyAtGoal(AAIController& Controller, const AActo
 	const bool bAlreadyAtGoal = Controller.GetPathFollowingComponent()->HasReached(TestGoal, Radius);
 	if (bAlreadyAtGoal)
 	{
-		Controller.GetPathFollowingComponent()->AbortMove(TEXT("Aborting move due to new move request finishing with AlreadyAtGoal"), FAIRequestID::AnyRequest, true, false, EPathFollowingMessage::OtherRequest);
-		Controller.GetPathFollowingComponent()->SetLastMoveAtGoal(true);
+		UE_VLOG(&Controller, LogPawnAction, Log, TEXT("New move request already at goal, aborting active movement"));
+		Controller.GetPathFollowingComponent()->RequestMoveWithImmediateFinish(EPathFollowingResult::Success);
 	}
 
 	return bAlreadyAtGoal;
diff --git a/Engine/Source/Runtime/AIModule/Private/BehaviorTree/Tasks/BTTask_MoveTo.cpp b/Engine/Source/Runtime/AIModule/Private/BehaviorTree/Tasks/BTTask_MoveTo.cpp
--- a/Engine/Source/Runtime/AIModule/Private/BehaviorTree/Tasks/BTTask_MoveTo.cpp
+++ b/Engine/Source/Runtime/AIModule/Private/BehaviorTree/Tasks/BTTask_MoveTo.cpp
@@ -221,8 +221,9 @@ EBlackboardNotificationResult UBTTask_MoveTo::OnBlackboardValueChange(const UBla
 			// don't abort move if using AI tasks - it will mess things up
 			if (MyMemory->MoveRequestID.IsValid())
 			{
+				UE_VLOG(BehaviorComp, LogBehaviorTree, Log, TEXT("Blackboard value for goal has changed, aborting current move request"));
 				StopWaitingForMessages(*BehaviorComp);
-				BehaviorComp->GetAIOwner()->GetPathFollowingComponent()->AbortMove(TEXT("Updating move due to BB value change"), MyMemory->MoveRequestID, /*bResetVelocity=*/false, /*bSilent=*/true);
+				BehaviorComp->GetAIOwner()->GetPathFollowingComponent()->AbortMove(*this, FPathFollowingResultFlags::NewRequest, MyMemory->MoveRequestID, EPathFollowingVelocityMode::Keep);
 			}
 
 			if (!bUseGameplayTasks && BehaviorComp->GetAIOwner()->ShouldPostponePathUpdates())
@@ -254,7 +255,7 @@ EBTNodeResult::Type UBTTask_MoveTo::AbortTask(UBehaviorTreeComponent& OwnerComp,
 			AAIController* MyController = OwnerComp.GetAIOwner();
 			if (MyController && MyController->GetPathFollowingComponent())
 			{
-				MyController->GetPathFollowingComponent()->AbortMove(TEXT("BehaviorTree abort"), MyMemory->MoveRequestID);
+				MyController->GetPathFollowingComponent()->AbortMove(*this, FPathFollowingResultFlags::OwnerFinished, MyMemory->MoveRequestID);
 			}
 		}
 		else
diff --git a/Engine/Source/Runtime/AIModule/Private/EnvironmentQuery/EQSTestingPawn.cpp b/Engine/Source/Runtime/AIModule/Private/EnvironmentQuery/EQSTestingPawn.cpp
--- a/Engine/Source/Runtime/AIModule/Private/EnvironmentQuery/EQSTestingPawn.cpp
+++ b/Engine/Source/Runtime/AIModule/Private/EnvironmentQuery/EQSTestingPawn.cpp
@@ -40,11 +40,11 @@ AEQSTestingPawn::AEQSTestingPawn(const FObjectInitializer& ObjectInitializer)
 	EdRenderComp = CreateEditorOnlyDefaultSubobject<UEQSRenderingComponent>(TEXT("EQSRender"));
 	if (HasAnyFlags(RF_ClassDefaultObject) == false)
 	{
-		UArrowComponent* ArrowComponent = FindComponentByClass<UArrowComponent>();
-		if (ArrowComponent != NULL)
+		UArrowComponent* ArrowComp = FindComponentByClass<UArrowComponent>();
+		if (ArrowComp != NULL)
 		{
-			ArrowComponent->SetRelativeScale3D(FVector(2, 2, 2));
-			ArrowComponent->bIsScreenSizeScaled = true;
+			ArrowComp->SetRelativeScale3D(FVector(2, 2, 2));
+			ArrowComp->bIsScreenSizeScaled = true;
 		}
 
 		UBillboardComponent* SpriteComponent = CreateEditorOnlyDefaultSubobject<UBillboardComponent>(TEXT("Sprite"));
diff --git a/Engine/Source/Runtime/AIModule/Private/EnvironmentQuery/EnvQueryInstance.cpp b/Engine/Source/Runtime/AIModule/Private/EnvironmentQuery/EnvQueryInstance.cpp
--- a/Engine/Source/Runtime/AIModule/Private/EnvironmentQuery/EnvQueryInstance.cpp
+++ b/Engine/Source/Runtime/AIModule/Private/EnvironmentQuery/EnvQueryInstance.cpp
@@ -81,14 +81,14 @@ bool FEnvQueryInstance::PrepareContext(UClass* Context, TArray<FEnvQuerySpatialD
 	{
 		UEnvQueryItemType_VectorBase* DefTypeOb = (UEnvQueryItemType_VectorBase*)ContextData.ValueType->GetDefaultObject();
 		const uint16 DefTypeValueSize = DefTypeOb->GetValueSize();
-		uint8* RawData = ContextData.RawData.GetData();
+		uint8* ContextRawData = ContextData.RawData.GetData();
 
 		Data.SetNumUninitialized(ContextData.NumValues);
 		for (int32 ValueIndex = 0; ValueIndex < ContextData.NumValues; ValueIndex++)
 		{
-			Data[ValueIndex].Location = DefTypeOb->GetItemLocation(RawData);
-			Data[ValueIndex].Rotation = DefTypeOb->GetItemRotation(RawData);
-			RawData += DefTypeValueSize;
+			Data[ValueIndex].Location = DefTypeOb->GetItemLocation(ContextRawData);
+			Data[ValueIndex].Rotation = DefTypeOb->GetItemRotation(ContextRawData);
+			ContextRawData += DefTypeValueSize;
 		}
 	}
 
@@ -109,13 +109,13 @@ bool FEnvQueryInstance::PrepareContext(UClass* Context, TArray<FVector>& Data)
 	{
 		UEnvQueryItemType_VectorBase* DefTypeOb = (UEnvQueryItemType_VectorBase*)ContextData.ValueType->GetDefaultObject();
 		const uint16 DefTypeValueSize = DefTypeOb->GetValueSize();
-		uint8* RawData = (uint8*)ContextData.RawData.GetData();
+		uint8* ContextRawData = (uint8*)ContextData.RawData.GetData();
 
 		Data.SetNumUninitialized(ContextData.NumValues);
 		for (int32 ValueIndex = 0; ValueIndex < ContextData.NumValues; ValueIndex++)
 		{
-			Data[ValueIndex] = DefTypeOb->GetItemLocation(RawData);
-			RawData += DefTypeValueSize;
+			Data[ValueIndex] = DefTypeOb->GetItemLocation(ContextRawData);
+			ContextRawData += DefTypeValueSize;
 		}
 	}
 
@@ -136,13 +136,13 @@ bool FEnvQueryInstance::PrepareContext(UClass* Context, TArray<FRotator>& Data)
 	{
 		UEnvQueryItemType_VectorBase* DefTypeOb = (UEnvQueryItemType_VectorBase*)ContextData.ValueType->GetDefaultObject();
 		const uint16 DefTypeValueSize = DefTypeOb->GetValueSize();
-		uint8* RawData = ContextData.RawData.GetData();
+		uint8* ContextRawData = ContextData.RawData.GetData();
 
 		Data.SetNumUninitialized(ContextData.NumValues);
 		for (int32 ValueIndex = 0; ValueIndex < ContextData.NumValues; ValueIndex++)
 		{
-			Data[ValueIndex] = DefTypeOb->GetItemRotation(RawData);
-			RawData += DefTypeValueSize;
+			Data[ValueIndex] = DefTypeOb->GetItemRotation(ContextRawData);
+			ContextRawData += DefTypeValueSize;
 		}
 	}
 
@@ -163,17 +163,17 @@ bool FEnvQueryInstance::PrepareContext(UClass* Context, TArray<AActor*>& Data)
 	{
 		UEnvQueryItemType_ActorBase* DefTypeOb = (UEnvQueryItemType_ActorBase*)ContextData.ValueType->GetDefaultObject();
 		const uint16 DefTypeValueSize = DefTypeOb->GetValueSize();
-		uint8* RawData = ContextData.RawData.GetData();
+		uint8* ContextRawData = ContextData.RawData.GetData();
 
 		Data.Reserve(ContextData.NumValues);
 		for (int32 ValueIndex = 0; ValueIndex < ContextData.NumValues; ValueIndex++)
 		{
-			AActor* Actor = DefTypeOb->GetActor(RawData);
+			AActor* Actor = DefTypeOb->GetActor(ContextRawData);
 			if (Actor)
 			{
 				Data.Add(Actor);
 			}
-			RawData += DefTypeValueSize;
+			ContextRawData += DefTypeValueSize;
 		}
 	}
 
@@ -714,9 +714,9 @@ uint32 FEnvQueryInstance::GetAllocatedSize() const
 	MemSize += ItemDetails.GetAllocatedSize();
 	MemSize += Options.GetAllocatedSize();
 
-	for (int32 OptionIndex = 0; OptionIndex < Options.Num(); OptionIndex++)
+	for (int32 OptionCount = 0; OptionCount < Options.Num(); OptionCount++)
 	{
-		MemSize += Options[OptionIndex].GetAllocatedSize();
+		MemSize += Options[OptionCount].GetAllocatedSize();
 	}
 
 	return MemSize;
@@ -736,17 +736,17 @@ uint32 FEnvQueryInstance::GetContextAllocatedSize() const
 
 FBox FEnvQueryInstance::GetBoundingBox() const
 {
-	const TArray<FEnvQueryItem>& Items = 
+	const TArray<FEnvQueryItem>& QueryItems = 
 #if USE_EQS_DEBUGGER
 	DebugData.DebugItems.Num() > 0 ? DebugData.DebugItems : 
 #endif // USE_EQS_DEBUGGER
-		this->Items;
+		Items;
 
-	const TArray<uint8>& RawData = 
+	const TArray<uint8>& QueryRawData = 
 #if USE_EQS_DEBUGGER
 		DebugData.RawData.Num() > 0 ? DebugData.RawData : 
 #endif // USE_EQS_DEBUGGER
-		this->RawData;
+		RawData;
 
 	FBox BBox(0);
 
@@ -754,9 +754,9 @@ FBox FEnvQueryInstance::GetBoundingBox() const
 	{
 		UEnvQueryItemType_VectorBase* DefTypeOb = ItemType->GetDefaultObject<UEnvQueryItemType_VectorBase>();
 
-		for (int32 Index = 0; Index < Items.Num(); ++Index)
+		for (int32 Index = 0; Index < QueryItems.Num(); ++Index)
 		{		
-			BBox += DefTypeOb->GetItemLocation(RawData.GetData() + Items[Index].DataOffset);
+			BBox += DefTypeOb->GetItemLocation(QueryRawData.GetData() + QueryItems[Index].DataOffset);
 		}
 	}
 
diff --git a/Engine/Source/Runtime/AIModule/Private/GameplayDebugger/GameplayDebuggerCategory_AI.cpp b/Engine/Source/Runtime/AIModule/Private/GameplayDebugger/GameplayDebuggerCategory_AI.cpp
--- a/Engine/Source/Runtime/AIModule/Private/GameplayDebugger/GameplayDebuggerCategory_AI.cpp
+++ b/Engine/Source/Runtime/AIModule/Private/GameplayDebugger/GameplayDebuggerCategory_AI.cpp
@@ -42,7 +42,10 @@ void FGameplayDebuggerCategory_AI::FRepData::Serialize(FArchive& Ar)
 	Ar << CurrentAIAssets;
 	Ar << NavDataInfo;
 	Ar << MontageInfo;
-	Ar << GameplayTaskInfo;
+	Ar << TaskQueueInfo;
+	Ar << TickingTaskInfo;
+	Ar << NumTasksInQueue;
+	Ar << NumTickingTasks;
 
 	uint32 BitFlags =
 		((bIsUsingPathFollowing ? 1 : 0) << 0) |
@@ -80,6 +83,18 @@ void FGameplayDebuggerCategory_AI::FRepDataPath::Serialize(FArchive& Ar)
 	Ar << PathPoints;
 }
 
+static FString DescribeTaskHelper(const UGameplayTask& TaskOb)
+{
+	const UObject* OwnerOb = Cast<const UObject>(TaskOb.GetTaskOwner());
+	return FString::Printf(TEXT("\n  {white}%s%s {%s}%s:%d {white}Owner:{yellow}%s {white}Res:{yellow}%s"),
+		*TaskOb.GetName(),
+		TaskOb.GetInstanceName() != NAME_None ? *FString::Printf(TEXT(" {yellow}[%s]"), *TaskOb.GetInstanceName().ToString()) : TEXT(""),
+		TaskOb.IsActive() ? TEXT("green") : TEXT("orange"),
+		*TaskOb.GetTaskStateName(), TaskOb.GetPriority(),
+		*GetNameSafe(OwnerOb),
+		TaskOb.GetRequiredResources().IsEmpty() ? TEXT("None") : *TaskOb.GetRequiredResources().GetDebugDescription());
+}
+
 void FGameplayDebuggerCategory_AI::CollectData(APlayerController* OwnerPC, AActor* DebugActor)
 {
 	APawn* MyPawn = Cast<APawn>(DebugActor);
@@ -129,24 +144,24 @@ void FGameplayDebuggerCategory_AI::CollectData(APlayerController* OwnerPC, AActo
 		DataPack.bIsUsingGameplayTasks = (TasksComponent != nullptr);
 		if (TasksComponent)
 		{
-			const FString TickQueueDesc = TasksComponent->GetTickingTasksDescription();
-			const FString PriorityQueueDesc = TasksComponent->GetTasksPriorityQueueDescription();
-
-			if (TickQueueDesc.Len())
+			for (FConstGameplayTaskIterator It = TasksComponent->GetTickingTaskIterator(); It; ++It)
 			{
-				DataPack.GameplayTaskInfo = TEXT("Ticking tasks: {yellow}");
-				DataPack.GameplayTaskInfo += TickQueueDesc;
+				const UGameplayTask* TaskOb = *It;
+				if (TaskOb)
+				{
+					DataPack.TickingTaskInfo += DescribeTaskHelper(*TaskOb);
+					DataPack.NumTickingTasks++;
+				}
 			}
 
-			if (PriorityQueueDesc.Len())
+			for (FConstGameplayTaskIterator It = TasksComponent->GetPriorityQueueIterator(); It; ++It)
 			{
-				if (DataPack.GameplayTaskInfo.Len())
+				const UGameplayTask* TaskOb = *It;
+				if (TaskOb)
 				{
-					DataPack.GameplayTaskInfo += TEXT("{white}\n");
+					DataPack.TaskQueueInfo += DescribeTaskHelper(*TaskOb);
+					DataPack.NumTasksInQueue++;
 				}
-
-				DataPack.GameplayTaskInfo = TEXT("Task queue: {yellow}");
-				DataPack.GameplayTaskInfo += PriorityQueueDesc;
 			}
 		}
 
@@ -158,6 +173,11 @@ void FGameplayDebuggerCategory_AI::CollectData(APlayerController* OwnerPC, AActo
 
 		CollectPathData(MyController);
 	}
+	else
+	{
+		PathDataPack.PathCorridor.Reset();
+		PathDataPack.PathPoints.Reset();
+	}
 }
 
 void FGameplayDebuggerCategory_AI::CollectPathData(AAIController* DebugAI)
@@ -294,9 +314,14 @@ void FGameplayDebuggerCategory_AI::DrawData(APlayerController* OwnerPC, FGamepla
 		CanvasContext.Printf(TEXT("Active task: {yellow}%s"), *DataPack.CurrentAITask);
 	}
 
-	if (DataPack.bIsUsingGameplayTasks && DataPack.GameplayTaskInfo.Len())
+	if (DataPack.bIsUsingGameplayTasks)
 	{
-		CanvasContext.Print(DataPack.GameplayTaskInfo);
+		if (DataPack.NumTickingTasks > 0)
+		{
+			CanvasContext.Printf(TEXT("Ticking tasks: {yellow}%d%s"), DataPack.NumTickingTasks, *DataPack.TickingTaskInfo);
+		}
+
+		CanvasContext.Printf(TEXT("Gameplay tasks: {yellow}%d%s"), DataPack.NumTasksInQueue, *DataPack.TaskQueueInfo);
 	}
 
 	if (DataPack.bIsUsingCharacter)
diff --git a/Engine/Source/Runtime/AIModule/Private/GameplayDebugger/GameplayDebuggerCategory_AI.h b/Engine/Source/Runtime/AIModule/Private/GameplayDebugger/GameplayDebuggerCategory_AI.h
--- a/Engine/Source/Runtime/AIModule/Private/GameplayDebugger/GameplayDebuggerCategory_AI.h
+++ b/Engine/Source/Runtime/AIModule/Private/GameplayDebugger/GameplayDebuggerCategory_AI.h
@@ -37,7 +37,10 @@ protected:
 		FString CurrentAIAssets;
 		FString NavDataInfo;
 		FString MontageInfo;
-		FString GameplayTaskInfo;
+		FString TaskQueueInfo;
+		FString TickingTaskInfo;
+		int16 NumTasksInQueue;
+		int16 NumTickingTasks;
 		uint32 bHasController : 1;
 		uint32 bPathHasGoalActor : 1;
 		uint32 bIsUsingPathFollowing : 1;
@@ -45,7 +48,7 @@ protected:
 		uint32 bIsUsingBehaviorTree : 1;
 		uint32 bIsUsingGameplayTasks : 1;
 
-		FRepData() : NextPathPointIndex(0), bHasController(false), bPathHasGoalActor(false),
+		FRepData() : NextPathPointIndex(0), NumTasksInQueue(0), NumTickingTasks(0), bHasController(false), bPathHasGoalActor(false),
 			bIsUsingPathFollowing(false), bIsUsingCharacter(false), bIsUsingBehaviorTree(false), bIsUsingGameplayTasks(false)
 		{
 		}
diff --git a/Engine/Source/Runtime/AIModule/Private/Navigation/CrowdFollowingComponent.cpp b/Engine/Source/Runtime/AIModule/Private/Navigation/CrowdFollowingComponent.cpp
--- a/Engine/Source/Runtime/AIModule/Private/Navigation/CrowdFollowingComponent.cpp
+++ b/Engine/Source/Runtime/AIModule/Private/Navigation/CrowdFollowingComponent.cpp
@@ -472,7 +472,7 @@ void UCrowdFollowingComponent::BeginDestroy()
 	Cleanup();
 }
 
-void UCrowdFollowingComponent::AbortMove(const FString& Reason, FAIRequestID RequestID, bool bResetVelocity, bool bSilent, uint8 MessageFlags)
+void UCrowdFollowingComponent::AbortMove(const UObject& Instigator, FPathFollowingResultFlags::Type AbortFlags, FAIRequestID RequestID, EPathFollowingVelocityMode VelocityMode)
 {
 	if (IsCrowdSimulationEnabled() && (Status != EPathFollowingStatus::Idle) && RequestID.IsEquivalent(GetCurrentRequestId()))
 	{
@@ -483,10 +483,10 @@ void UCrowdFollowingComponent::AbortMove(const FString& Reason, FAIRequestID Req
 		}
 	}
 
-	Super::AbortMove(Reason, RequestID, bResetVelocity, bSilent, MessageFlags);
+	Super::AbortMove(Instigator, AbortFlags, RequestID, VelocityMode);
 }
 
-void UCrowdFollowingComponent::PauseMove(FAIRequestID RequestID, bool bResetVelocity)
+void UCrowdFollowingComponent::PauseMove(FAIRequestID RequestID, EPathFollowingVelocityMode VelocityMode)
 {
 	if (IsCrowdSimulationEnabled() && (Status != EPathFollowingStatus::Paused) && RequestID.IsEquivalent(GetCurrentRequestId()))
 	{
@@ -497,7 +497,7 @@ void UCrowdFollowingComponent::PauseMove(FAIRequestID RequestID, bool bResetVelo
 		}
 	}
 
-	Super::PauseMove(RequestID, bResetVelocity);
+	Super::PauseMove(RequestID, VelocityMode);
 }
 
 void UCrowdFollowingComponent::ResumeMove(FAIRequestID RequestID)
@@ -565,7 +565,7 @@ void UCrowdFollowingComponent::FinishUsingCustomLink(INavLinkCustomInterface* Cu
 	}
 }
 
-void UCrowdFollowingComponent::OnPathFinished(EPathFollowingResult::Type Result)
+void UCrowdFollowingComponent::OnPathFinished(const FPathFollowingResult& Result)
 {
 	UCrowdManager* CrowdManager = UCrowdManager::GetCurrent(GetWorld());
 	if (IsCrowdSimulationEnabled() && CrowdManager)
@@ -838,7 +838,7 @@ void UCrowdFollowingComponent::UpdatePathSegment()
 
 	if (!Path.IsValid() || MovementComp == NULL)
 	{
-		AbortMove(TEXT("no path"), FAIRequestID::CurrentRequest, true, false, EPathFollowingMessage::NoPath);
+		OnPathFinished(FPathFollowingResult(EPathFollowingResult::Aborted, FPathFollowingResultFlags::InvalidPath));
 		return;
 	}
 
@@ -846,7 +846,7 @@ void UCrowdFollowingComponent::UpdatePathSegment()
 	{
 		if (!Path->IsWaitingForRepath())
 		{
-			AbortMove(TEXT("no path"), FAIRequestID::CurrentRequest, true, false, EPathFollowingMessage::NoPath);
+			OnPathFinished(FPathFollowingResult(EPathFollowingResult::Aborted, FPathFollowingResultFlags::InvalidPath));
 		}
 		return;
 	}
@@ -862,7 +862,7 @@ void UCrowdFollowingComponent::UpdatePathSegment()
 		{
 			// check if collided with goal actor
 			OnSegmentFinished();
-			OnPathFinished(EPathFollowingResult::Success);
+			OnPathFinished(FPathFollowingResult(EPathFollowingResult::Success, FPathFollowingResultFlags::None));
 		}
 		else if (bFinalPathPart)
 		{
@@ -876,7 +876,7 @@ void UCrowdFollowingComponent::UpdatePathSegment()
 			if (bMovedTooFar || HasReachedInternal(GoalLocation, 0.0f, 0.0f, CurrentLocation, AcceptanceRadius, bStopOnOverlap ? MinAgentRadiusPct : 0.0f))
 			{
 				UE_VLOG(GetOwner(), LogCrowdFollowing, Log, TEXT("Last path segment finished due to \'%s\'"), bMovedTooFar ? TEXT("Missing Last Point") : TEXT("Reaching Destination"));
-				OnPathFinished(EPathFollowingResult::Success);
+				OnPathFinished(FPathFollowingResult(EPathFollowingResult::Success, FPathFollowingResultFlags::None));
 			}
 		}
 		else
@@ -905,7 +905,7 @@ void UCrowdFollowingComponent::UpdatePathSegment()
 		const bool bHasNewSample = UpdateBlockDetection();
 		if (bHasNewSample && IsBlocked())
 		{
-			OnPathFinished(EPathFollowingResult::Blocked);
+			OnPathFinished(FPathFollowingResult(EPathFollowingResult::Blocked, FPathFollowingResultFlags::None));
 		}
 	}
 }
diff --git a/Engine/Source/Runtime/AIModule/Private/Navigation/PathFollowingComponent.cpp b/Engine/Source/Runtime/AIModule/Private/Navigation/PathFollowingComponent.cpp
--- a/Engine/Source/Runtime/AIModule/Private/Navigation/PathFollowingComponent.cpp
+++ b/Engine/Source/Runtime/AIModule/Private/Navigation/PathFollowingComponent.cpp
@@ -16,8 +16,6 @@
 
 #include "Navigation/PathFollowingComponent.h"
 
-#define USE_PHYSIC_FOR_VISIBILITY_TESTS 1 // Physic will be used for visibility tests if set or only raycasts on navmesh if not
-
 #if UE_BUILD_TEST || UE_BUILD_SHIPPING
 #define SHIPPING_STATIC static
 #else
@@ -26,12 +24,79 @@
 
 DEFINE_LOG_CATEGORY(LogPathFollowing);
 
+FPathFollowingResult::FPathFollowingResult(uint16 InFlags) : Flags(InFlags)
+{
+	Code =
+		HasFlag(FPathFollowingResultFlags::Success) && !HasFlag(FPathFollowingResultFlags::Blocked | FPathFollowingResultFlags::OffPath | FPathFollowingResultFlags::ExternalCancel) ? EPathFollowingResult::Success :
+		HasFlag(FPathFollowingResultFlags::Blocked) && !HasFlag(FPathFollowingResultFlags::Success | FPathFollowingResultFlags::OffPath | FPathFollowingResultFlags::ExternalCancel) ? EPathFollowingResult::Blocked :
+		HasFlag(FPathFollowingResultFlags::OffPath) && !HasFlag(FPathFollowingResultFlags::Success | FPathFollowingResultFlags::Blocked | FPathFollowingResultFlags::ExternalCancel) ? EPathFollowingResult::OffPath :
+		HasFlag(FPathFollowingResultFlags::ExternalCancel) && !HasFlag(FPathFollowingResultFlags::Success | FPathFollowingResultFlags::Blocked | FPathFollowingResultFlags::OffPath) ? EPathFollowingResult::Aborted :
+		EPathFollowingResult::Invalid;
+}
+
+FPathFollowingResult::FPathFollowingResult(EPathFollowingResult::Type ResultCode, uint16 ExtraFlags) : Code(ResultCode)
+{
+	static uint16 CodeLookup[6] = {
+		FPathFollowingResultFlags::Success,			// EPathFollowingResult::Success
+		FPathFollowingResultFlags::Blocked,			// EPathFollowingResult::Blocked
+		FPathFollowingResultFlags::OffPath,			// EPathFollowingResult::OffPath
+		FPathFollowingResultFlags::ExternalCancel,	// EPathFollowingResult::Aborted
+		0,											// EPathFollowingResult::Skipped_DEPRECATED
+		0,											// EPathFollowingResult::Invalid
+	};
+
+	Flags = CodeLookup[ResultCode] | ExtraFlags;
+}
+
+FString FPathFollowingResult::ToString() const
+{
+	static UEnum* ResultEnum = FindObject<UEnum>(ANY_PACKAGE, TEXT("EPathFollowingResult"));
+	return FString::Printf(TEXT("%s[%s]"), ResultEnum ? *ResultEnum->GetEnumName(Code) : TEXT("??"), *FPathFollowingResultFlags::ToString(Flags));
+}
+
+FString FPathFollowingResultFlags::ToString(uint16 Value)
+{
+	static FString FlagDesc[] = {
+		GET_FUNCTION_NAME_STRING_CHECKED(FPathFollowingResultFlags, Success),
+		GET_FUNCTION_NAME_STRING_CHECKED(FPathFollowingResultFlags, Blocked),
+		GET_FUNCTION_NAME_STRING_CHECKED(FPathFollowingResultFlags, OffPath),
+		GET_FUNCTION_NAME_STRING_CHECKED(FPathFollowingResultFlags, ExternalCancel),
+		GET_FUNCTION_NAME_STRING_CHECKED(FPathFollowingResultFlags, OwnerFinished),
+		GET_FUNCTION_NAME_STRING_CHECKED(FPathFollowingResultFlags, InvalidPath),
+		GET_FUNCTION_NAME_STRING_CHECKED(FPathFollowingResultFlags, MovementStop),
+		GET_FUNCTION_NAME_STRING_CHECKED(FPathFollowingResultFlags, NewRequest),
+		GET_FUNCTION_NAME_STRING_CHECKED(FPathFollowingResultFlags, ForcedScript),
+	};
+
+	FString CombinedDesc;
+	for (int32 Idx = 0; Idx < FirstGameplayFlagShift; Idx++)
+	{
+		if (Value & (1 << Idx))
+		{
+			CombinedDesc += FlagDesc[Idx];
+			CombinedDesc += TEXT(' ');
+		}
+	}
+
+	Value = Value >> FirstGameplayFlagShift;
+	for (int32 Idx = 0; Value; Idx++, Value /= 2)
+	{
+		if (Value & 1)
+		{
+			CombinedDesc += TEXT("Game");
+			CombinedDesc += TTypeToString<int32>::ToString(Idx);
+			CombinedDesc += TEXT(' ');
+		}
+	}
+
+	return CombinedDesc.LeftChop(1);
+}
+
 //----------------------------------------------------------------------//
 // Life cycle                                                        
 //----------------------------------------------------------------------//
 uint32 UPathFollowingComponent::NextRequestId = 1;
 const float UPathFollowingComponent::DefaultAcceptanceRadius = -1.f;
-UPathFollowingComponent::FRequestCompletedSignature UPathFollowingComponent::UnboundRequestDelegate;
 
 UPathFollowingComponent::UPathFollowingComponent(const FObjectInitializer& ObjectInitializer) : Super(ObjectInitializer)
 {
@@ -48,7 +113,6 @@ UPathFollowingComponent::UPathFollowingComponent(const FObjectInitializer& Objec
 	NextSampleIdx = 0;
 	bUseBlockDetection = true;
 	bLastMoveReachedGoal = false;
-	bUseVisibilityTestsSimplification = false;
 	bStopMovementOnFinish = true;
 	bIsUsingMetaPath = false;
 
@@ -150,7 +214,8 @@ void UPathFollowingComponent::OnPathEvent(FNavigationPath* InPath, ENavPathEvent
 				const bool bIsPathAccepted = HandlePathUpdateEvent();
 				if (!bIsPathAccepted)
 				{
-					AbortMove(TEXT("Failed path update"));
+					UE_VLOG(GetOwner(), LogPathFollowing, Log, TEXT(">> updated path was not accepcted, aborting move"));
+					OnPathFinished(EPathFollowingResult::Aborted, FPathFollowingResultFlags::InvalidPath);
 				}
 			}
 			break;
@@ -195,21 +260,10 @@ bool UPathFollowingComponent::HandlePathUpdateEvent()
 	return true;
 }
 
-FAIRequestID UPathFollowingComponent::RequestMove(FNavPathSharedPtr InPath, FRequestCompletedSignature OnComplete,
-	const AActor* InDestinationActor, float InAcceptanceRadius, bool bInStopOnOverlap, FCustomMoveSharedPtr InGameData)
+FAIRequestID UPathFollowingComponent::RequestMove(const FAIMoveRequest& RequestData, FNavPathSharedPtr InPath)
 {
-	UE_VLOG(GetOwner(), LogPathFollowing, Log, TEXT("RequestMove: Path(%s), AcceptRadius(%.1f%s), DestinationActor(%s), GameData(%s)"),
-		*GetPathDescHelper(InPath),
-		InAcceptanceRadius, bInStopOnOverlap ? TEXT(" + agent") : TEXT(""),
-		*GetNameSafe(InDestinationActor),
-		!InGameData.IsValid() ? TEXT("missing") : TEXT("valid"));
-
-	LogPathHelper(GetOwner(), InPath, InDestinationActor);
-
-	if (InAcceptanceRadius == UPathFollowingComponent::DefaultAcceptanceRadius)
-	{
-		InAcceptanceRadius = MyDefaultAcceptanceRadius;
-	}
+	UE_VLOG(GetOwner(), LogPathFollowing, Log, TEXT("RequestMove: Path(%s) %s"), *GetPathDescHelper(InPath), *RequestData.ToString());
+	LogPathHelper(GetOwner(), InPath, RequestData.GetGoalActor());
 
 	if (ResourceLock.IsLocked())
 	{
@@ -217,7 +271,8 @@ FAIRequestID UPathFollowingComponent::RequestMove(FNavPathSharedPtr InPath, FReq
 		return FAIRequestID::InvalidRequest;
 	}
 
-	if (!InPath.IsValid() || InAcceptanceRadius < 0.0f)
+	const float UseAcceptanceRadius = (RequestData.GetAcceptanceRadius() == UPathFollowingComponent::DefaultAcceptanceRadius) ? MyDefaultAcceptanceRadius : RequestData.GetAcceptanceRadius();
+ 	if (!InPath.IsValid() || UseAcceptanceRadius < 0.0f)
 	{
 		UE_VLOG(GetOwner(), LogPathFollowing, Log, TEXT("RequestMove: invalid request"));
 		return FAIRequestID::InvalidRequest;
@@ -230,25 +285,19 @@ FAIRequestID UPathFollowingComponent::RequestMove(FNavPathSharedPtr InPath, FReq
 		return FAIRequestID::InvalidRequest;
 	}
 
-	// update ID first, so any observer notified by AbortMove() could detect new request
-	const uint32 PrevMoveId = CurrentRequestId;
-	
 	// abort previous movement
-	if (Status == EPathFollowingStatus::Paused && Path.IsValid() && InPath.Get() == Path.Get() && DestinationActor == InDestinationActor)
+	if (Status == EPathFollowingStatus::Paused && Path.IsValid() && InPath.Get() == Path.Get() && DestinationActor == RequestData.GetGoalActor())
 	{
 		ResumeMove();
 	}
 	else
 	{
-		if (Status == EPathFollowingStatus::Moving)
+		if (Status != EPathFollowingStatus::Idle)
 		{
-			const bool bResetVelocity = false;
-			const bool bFinishAsSkipped = true;
-			AbortMove(TEXT("new request"), PrevMoveId, bResetVelocity, bFinishAsSkipped, EPathFollowingMessage::OtherRequest);
+			OnPathFinished(EPathFollowingResult::Aborted, FPathFollowingResultFlags::NewRequest);
 		}
 		
 		Reset();
-
 		StoreRequestId();
 
 		// store new data
@@ -272,19 +321,17 @@ FAIRequestID UPathFollowingComponent::RequestMove(FNavPathSharedPtr InPath, FReq
 		PathTimeWhenPaused = 0.0f;
 		OnPathUpdated();
 
-		AcceptanceRadius = InAcceptanceRadius;
-		GameData = InGameData;
-		OnRequestFinished = OnComplete;	
-		bStopOnOverlap = bInStopOnOverlap;
-		SetDestinationActor(InDestinationActor);
+		AcceptanceRadius = UseAcceptanceRadius;
+		bStopOnOverlap = RequestData.CanStopOnOverlap();
+		GameData = RequestData.GetUserData();
+		SetDestinationActor(RequestData.GetGoalActor());
 		UpdateDecelerationData();
 
 	#if ENABLE_VISUAL_LOG
 		const FVector CurrentLocation = MovementComp ? MovementComp->GetActorFeetLocation() : FVector::ZeroVector;
 		const FVector DestLocation = InPath->GetDestinationLocation();
 		const FVector ToDest = DestLocation - CurrentLocation;
-		UE_VLOG(GetOwner(), LogPathFollowing, Log, TEXT("RequestMove: accepted, ID(%u) dist2D(%.0f) distZ(%.0f)"),
-			CurrentRequestId.GetID(), ToDest.Size2D(), FMath::Abs(ToDest.Z));
+		UE_VLOG(GetOwner(), LogPathFollowing, Log, TEXT("RequestMove: accepted, ID(%u) dist2D(%.0f) distZ(%.0f)"), CurrentRequestId, ToDest.Size2D(), FMath::Abs(ToDest.Z));
 	#endif // ENABLE_VISUAL_LOG
 
 		// with async pathfinding paths can be incomplete, movement will start after receiving path event 
@@ -320,46 +367,29 @@ bool UPathFollowingComponent::UpdateMove(FNavPathSharedPtr InPath, FAIRequestID
 	return false;
 }
 
-void UPathFollowingComponent::AbortMove(const FString& Reason, FAIRequestID RequestID, bool bResetVelocity, bool bSilent, uint8 MessageFlags)
+void UPathFollowingComponent::AbortMove(const UObject& Instigator, FPathFollowingResultFlags::Type AbortFlags, FAIRequestID RequestID, EPathFollowingVelocityMode VelocityMode)
 {
-	UE_VLOG(GetOwner(), LogPathFollowing, Log, TEXT("AbortMove: Reason(%s) RequestID(%u)"), *Reason, RequestID);
+	UE_VLOG(GetOwner(), LogPathFollowing, Log, TEXT("AbortMove: RequestID(%u) Instigator(%s)"), RequestID, *Instigator.GetName());
 
 	if ((Status != EPathFollowingStatus::Idle) && RequestID.IsEquivalent(GetCurrentRequestId()))
 	{
-		const EPathFollowingResult::Type FinishResult = bSilent ? EPathFollowingResult::Skipped : EPathFollowingResult::Aborted;
+		bStopMovementOnFinish = (VelocityMode == EPathFollowingVelocityMode::Reset);
+		OnPathFinished(FPathFollowingResult(EPathFollowingResult::Aborted, AbortFlags & FPathFollowingResultFlags::ExternalCancelFlagMask));
+	}
+}
 
-		INavLinkCustomInterface* CustomNavLink = Cast<INavLinkCustomInterface>(CurrentCustomLinkOb.Get());
-		if (CustomNavLink)
+void UPathFollowingComponent::RequestMoveWithImmediateFinish(EPathFollowingResult::Type Result, EPathFollowingVelocityMode VelocityMode)
 		{
-			CustomNavLink->OnLinkMoveFinished(this);
-		}
-
-		// save data required for observers before reseting temporary variables
-		const uint32 AbortedMoveId = RequestID ? RequestID : CurrentRequestId;
-		FRequestCompletedSignature SavedReqFinished = OnRequestFinished;
-
-		Reset();
-		bLastMoveReachedGoal = false;
-		UpdateMoveFocus();
-
-		GetWorld()->GetTimerManager().ClearTimer(WaitingForPathTimer);
-		if (bResetVelocity && MovementComp && MovementComp->CanStopPathFollowing())
+	if (Status != EPathFollowingStatus::Idle)
 		{
-			MovementComp->StopMovementKeepPathing();
+		OnPathFinished(EPathFollowingResult::Aborted, FPathFollowingResultFlags::NewRequest);
 		}
 
-		// notify observers after state was reset (they can request another move)
-		SavedReqFinished.ExecuteIfBound(FinishResult);
-		OnMoveFinished.Broadcast(AbortedMoveId, FinishResult);
-
-		FAIMessage Msg(UBrainComponent::AIMessage_MoveFinished, this, AbortedMoveId, FAIMessage::Failure);
-		Msg.SetFlag(MessageFlags);
-
-		FAIMessage::Send(Cast<AController>(GetOwner()), Msg);
-	}
+	StoreRequestId();
+	OnPathFinished(Result, FPathFollowingResultFlags::None);
 }
 
-void UPathFollowingComponent::PauseMove(FAIRequestID RequestID, bool bResetVelocity)
+void UPathFollowingComponent::PauseMove(FAIRequestID RequestID, EPathFollowingVelocityMode VelocityMode)
 {
 	UE_VLOG(GetOwner(), LogPathFollowing, Log, TEXT("PauseMove: RequestID(%u)"), RequestID);
 	if (Status == EPathFollowingStatus::Paused)
@@ -369,7 +399,7 @@ void UPathFollowingComponent::PauseMove(FAIRequestID RequestID, bool bResetVeloc
 
 	if (RequestID.IsEquivalent(GetCurrentRequestId()))
 	{
-		if (bResetVelocity && MovementComp && MovementComp->CanStopPathFollowing())
+		if ((VelocityMode == EPathFollowingVelocityMode::Reset) && MovementComp && MovementComp->CanStopPathFollowing())
 		{
 			MovementComp->StopMovementKeepPathing();
 		}
@@ -380,8 +410,6 @@ void UPathFollowingComponent::PauseMove(FAIRequestID RequestID, bool bResetVeloc
 
 		UpdateMoveFocus();
 	}
-
-	// TODO: pause path updates with goal movement
 }
 
 void UPathFollowingComponent::ResumeMove(FAIRequestID RequestID)
@@ -415,13 +443,12 @@ void UPathFollowingComponent::ResumeMove(FAIRequestID RequestID)
 		}
 		else
 		{
-			OnPathFinished(EPathFollowingResult::OffPath);
+			OnPathFinished(EPathFollowingResult::OffPath, FPathFollowingResultFlags::None);
 		}
 	}
 	else
 	{
-		UE_VLOG(GetOwner(), LogPathFollowing, Log, TEXT("ResumeMove: RequestID(%u) is neither \'AnyRequest\' not CurrentRequestId(%u). Ignoring.")
-			, RequestID, CurrentRequestId);
+		UE_VLOG(GetOwner(), LogPathFollowing, Log, TEXT("ResumeMove: RequestID(%u) is neither \'AnyRequest\' not CurrentRequestId(%u). Ignoring."), RequestID, CurrentRequestId);
 	}
 }
 
@@ -445,12 +472,19 @@ bool UPathFollowingComponent::ShouldCheckPathOnResume() const
 	return bCheckPath;
 }
 
-void UPathFollowingComponent::OnPathFinished(EPathFollowingResult::Type Result)
+void UPathFollowingComponent::OnPathFinished(const FPathFollowingResult& Result)
 {
-	UE_VLOG(GetOwner(), LogPathFollowing, Log, TEXT("OnPathFinished: %s"), *GetResultDesc(Result));
+	UE_VLOG(GetOwner(), LogPathFollowing, Log, TEXT("OnPathFinished: %s"), *Result.ToString());
+	
+	INavLinkCustomInterface* CustomNavLink = Cast<INavLinkCustomInterface>(CurrentCustomLinkOb.Get());
+	if (CustomNavLink)
+	{
+		CustomNavLink->OnLinkMoveFinished(this);
+		CurrentCustomLinkOb.Reset();
+	}
 	
 	// update meta path if needed
-	if (bIsUsingMetaPath && Result == EPathFollowingResult::Success && MovementComp)
+	if (bIsUsingMetaPath && Result.IsSuccess() && MovementComp)
 	{
 		FMetaNavMeshPath* MetaNavPath = Path->CastPath<FMetaNavMeshPath>();
 		const bool bNeedPathUpdate = MetaNavPath && MetaNavPath->ConditionalMoveToNextSection(MovementComp->GetActorFeetLocation(), EMetaPathUpdateReason::PathFinished);
@@ -462,26 +496,25 @@ void UPathFollowingComponent::OnPathFinished(EPathFollowingResult::Type Result)
 	}
 
 	// save move status
-	bLastMoveReachedGoal = (Result == EPathFollowingResult::Success) && !HasPartialPath();
+	bLastMoveReachedGoal = Result.IsSuccess() && !HasPartialPath();
 
 	// save data required for observers before reseting temporary variables
 	const FAIRequestID FinishedMoveId = CurrentRequestId;
-	FRequestCompletedSignature SavedReqFinished = OnRequestFinished;
 
 	Reset();
 	UpdateMoveFocus();
 
-	if (MovementComp && MovementComp->CanStopPathFollowing() && 
-		bStopMovementOnFinish && !MovementComp->UseAccelerationForPathFollowing())
+	if (bStopMovementOnFinish && MovementComp && MovementComp->CanStopPathFollowing() && !MovementComp->UseAccelerationForPathFollowing())
 	{
 		MovementComp->StopMovementKeepPathing();
 	}
 
 	// notify observers after state was reset (they can request another move)
-	SavedReqFinished.ExecuteIfBound(Result);
-	OnMoveFinished.Broadcast(FinishedMoveId, Result);
+	OnRequestFinished.Broadcast(FinishedMoveId, Result);
 
-	FAIMessage::Send(Cast<AController>(GetOwner()), FAIMessage(UBrainComponent::AIMessage_MoveFinished, this, FinishedMoveId, (Result == EPathFollowingResult::Success)));
+	FAIMessage Msg(UBrainComponent::AIMessage_MoveFinished, this, FinishedMoveId, Result.IsSuccess());
+	Msg.SetFlag(Result.Flags & 0xff);
+	FAIMessage::Send(Cast<AController>(GetOwner()), Msg);
 }
 
 void UPathFollowingComponent::OnSegmentFinished()
@@ -516,7 +549,6 @@ void UPathFollowingComponent::TickComponent(float DeltaTime, enum ELevelTick Tic
 	if (Status == EPathFollowingStatus::Moving)
 	{
 		// check finish conditions, update current segment if needed
-		// may result in changing current status!
 		UpdatePathSegment();
 	}
 
@@ -575,17 +607,20 @@ void UPathFollowingComponent::Reset()
 	Path.Reset();
 	GameData.Reset();
 	DestinationActor.Reset();
-	OnRequestFinished.Unbind();
+	CurrentDestination.Clear();
 	AcceptanceRadius = MyDefaultAcceptanceRadius;
 	bStopOnOverlap = true;
 	bCollidedWithGoal = false;
 	bIsUsingMetaPath = false;
 
 	CurrentRequestId = FAIRequestID::InvalidRequest;
-
-	CurrentDestination.Clear();
-
 	Status = EPathFollowingStatus::Idle;
+
+	if (WaitingForPathTimer.IsValid())
+	{
+		GetWorld()->GetTimerManager().ClearTimer(WaitingForPathTimer);
+		WaitingForPathTimer.Invalidate();
+	}
 }
 
 int32 UPathFollowingComponent::DetermineStartingPathPoint(const FNavigationPath* ConsideredPath) const
@@ -697,137 +732,14 @@ void UPathFollowingComponent::SetMoveSegment(int32 SegmentStartIndex)
 
 int32 UPathFollowingComponent::DetermineCurrentTargetPathPoint(int32 StartIndex)
 {
-	if (!bUseVisibilityTestsSimplification ||
-		Path->GetPathPoints()[StartIndex].CustomLinkId ||
-		Path->GetPathPoints()[StartIndex + 1].CustomLinkId)
-	{
-		return StartIndex + 1;
-	}
-
-
-#if WITH_RECAST
-	FNavMeshNodeFlags Flags0(Path->GetPathPoints()[StartIndex].Flags);
-	FNavMeshNodeFlags Flags1(Path->GetPathPoints()[StartIndex + 1].Flags);
-
-	const bool bOffMesh0 = (Flags0.PathFlags & RECAST_STRAIGHTPATH_OFFMESH_CONNECTION) != 0;
-	const bool bOffMesh1 = (Flags1.PathFlags & RECAST_STRAIGHTPATH_OFFMESH_CONNECTION) != 0;
-
-	if ((Flags0.Area != Flags1.Area) || (bOffMesh0 && bOffMesh1))
-	{
-		return StartIndex + 1;
-	}
-#endif
-
-	return OptimizeSegmentVisibility(StartIndex);
-}
-
-int32 UPathFollowingComponent::OptimizeSegmentVisibility(int32 StartIndex)
-{
-	//SCOPE_CYCLE_COUNTER(STAT_Navigation_PathVisibilityOptimisation);
-
-	const AAIController* MyAI = Cast<AAIController>(GetOwner());
-	const ANavigationData* NavData = Path.IsValid() ? Path->GetNavigationDataUsed() : NULL;
-	if (NavData == NULL || MyAI == NULL)
-	{
 		return StartIndex + 1;
 	}
 
-	const bool bIsDirect = (Path->CastPath<FAbstractNavigationPath>() != NULL);
-	if (bIsDirect)
-	{
-		// can't optimize anything without real corridor
-		return StartIndex + 1;
-	}
-
-	const APawn* MyPawn = MyAI->GetPawn();
-	const float PawnHalfHeight = (MyAI->GetCharacter() && MyAI->GetCharacter()->GetCapsuleComponent()) ? MyAI->GetCharacter()->GetCapsuleComponent()->GetScaledCapsuleHalfHeight() : 0.0f;
-	const FVector PawnLocation = MyPawn->GetActorLocation();
-
-#if USE_PHYSIC_FOR_VISIBILITY_TESTS
-	static const FName NAME_NavAreaTrace = TEXT("NavmeshVisibilityTrace");
-	float Radius, HalfHeight;
-	MyPawn->GetSimpleCollisionCylinder(Radius, HalfHeight);
-	const FCollisionShape CollisionShape = FCollisionShape::MakeCapsule(Radius, HalfHeight);
-	FCollisionQueryParams CapsuleParams(NAME_NavAreaTrace, false, MyAI->GetCharacter());
-	UPrimitiveComponent* PrimitiveComponent = Cast<UPrimitiveComponent>(MyPawn->GetRootComponent());
-	if (PrimitiveComponent)
-	{
-		CapsuleParams.AddIgnoredActors(PrimitiveComponent->MoveIgnoreActors);
-	}
-
-#endif
-	FSharedNavQueryFilter QueryFilter = Path->GetFilter()->GetCopy();
-#if WITH_RECAST
-	const uint8 StartArea = FNavMeshNodeFlags(Path->GetPathPoints()[StartIndex].Flags).Area;
-	TArray<float> CostArray;
-	CostArray.Init(DT_UNWALKABLE_POLY_COST, DT_MAX_AREAS);
-	CostArray[StartArea] = 0;
-	QueryFilter->SetAllAreaCosts(CostArray);
-#endif
-
-	const int32 MaxPoints = Path->GetPathPoints().Num();
-	int32 Index = StartIndex + 2;
-	Path->ShortcutNodeRefs.Reset();
-
-	for (; Index <= MaxPoints; ++Index)
-	{
-		if (!Path->GetPathPoints().IsValidIndex(Index))
-		{
-			break;
-		}
-
-		const FNavPathPoint& PathPt1 = Path->GetPathPoints()[Index];
-		const FVector PathPt1Location = *Path->GetPathPointLocation(Index);
-		const FVector AdjustedDestination = PathPt1Location + FVector(0.0f, 0.0f, PawnHalfHeight);
-
-		FVector HitLocation;
-#if WITH_RECAST
-		ARecastNavMesh::FRaycastResult RaycastResult;
-		const bool RaycastHitResult = ARecastNavMesh::NavMeshRaycast(NavData, PawnLocation, AdjustedDestination, HitLocation, QueryFilter, MyAI, RaycastResult);
-		if (Path.IsValid())
-		{
-			Path->ShortcutNodeRefs.Reserve(RaycastResult.CorridorPolysCount);
-			Path->ShortcutNodeRefs.SetNumUninitialized(RaycastResult.CorridorPolysCount);
-
-			FPlatformMemory::Memcpy(Path->ShortcutNodeRefs.GetData(), RaycastResult.CorridorPolys, RaycastResult.CorridorPolysCount * sizeof(NavNodeRef));
-		}
-#else
-		const bool RaycastHitResult = NavData->Raycast(PawnLocation, AdjustedDestination, HitLocation, QueryFilter);
-#endif
-		if (RaycastHitResult)
-		{
-			break;
-		}
-#if USE_PHYSIC_FOR_VISIBILITY_TESTS
-		else
-		{
-			FHitResult Result;
-			if (GetWorld()->SweepSingleByChannel(Result, PawnLocation, AdjustedDestination, FQuat::Identity, ECC_WorldStatic, CollisionShape, CapsuleParams))
-			{
-				break;
-			}
-		}
-#endif
-
-#if WITH_RECAST
-		/** don't move to next point if we are changing area, we are at off mesh connection or it's a custom nav link*/
-		if (FNavMeshNodeFlags(Path->GetPathPoints()[StartIndex].Flags).Area != FNavMeshNodeFlags(PathPt1.Flags).Area ||
-			FNavMeshNodeFlags(PathPt1.Flags).PathFlags & RECAST_STRAIGHTPATH_OFFMESH_CONNECTION || 
-			PathPt1.CustomLinkId)
-		{
-			return Index;
-		}
-#endif
-	}
-
-	return Index-1;
-}
-
 void UPathFollowingComponent::UpdatePathSegment()
 {
 	if (!Path.IsValid() || MovementComp == NULL)
 	{
-		AbortMove(TEXT("no path"), FAIRequestID::CurrentRequest, true, false, EPathFollowingMessage::NoPath);
+		OnPathFinished(EPathFollowingResult::Aborted, FPathFollowingResultFlags::InvalidPath);
 		return;
 	}
 
@@ -835,8 +747,9 @@ void UPathFollowingComponent::UpdatePathSegment()
 	{
 		if (!Path->IsWaitingForRepath())
 		{
-			AbortMove(TEXT("no path"), FAIRequestID::CurrentRequest, true, false, EPathFollowingMessage::NoPath);
+			OnPathFinished(EPathFollowingResult::Aborted, FPathFollowingResultFlags::InvalidPath);
 		}
+
 		return;
 	}
 
@@ -852,13 +765,13 @@ void UPathFollowingComponent::UpdatePathSegment()
 		{
 			// check if collided with goal actor
 			OnSegmentFinished();
-			OnPathFinished(EPathFollowingResult::Success);
+			OnPathFinished(EPathFollowingResult::Success, FPathFollowingResultFlags::None);
 		}
 		else if (HasReachedDestination(CurrentLocation))
 		{
 			// always check for destination, acceptance radius may cause it to pass before reaching last segment
 			OnSegmentFinished();
-			OnPathFinished(EPathFollowingResult::Success);
+			OnPathFinished(EPathFollowingResult::Success, FPathFollowingResultFlags::None);
 		}
 		else if (bFollowingLastSegment)
 		{
@@ -869,18 +782,13 @@ void UPathFollowingComponent::UpdatePathSegment()
 				const FVector AgentLocation = DestinationAgent ? DestinationAgent->GetNavAgentLocation() : DestinationActor->GetActorLocation();
 				// note that the condition below requires GoalLocation to be in world space.
 				const FVector GoalLocation = FQuatRotationTranslationMatrix(DestinationActor->GetActorQuat(), AgentLocation).TransformPosition(MoveOffset);
-				FVector HitLocation;
 
-				if (MyNavData == nullptr //|| MyNavData->DoesNodeContainLocation(Path->GetPathPoints().Last().NodeRef, GoalLocation))
-					|| (FVector::DistSquared(GoalLocation, *CurrentDestination) > SMALL_NUMBER &&  MyNavData->Raycast(CurrentLocation, GoalLocation, HitLocation, nullptr) == false))
-				{
 					CurrentDestination.Set(NULL, GoalLocation);
 
 					UE_VLOG(this, LogPathFollowing, Log, TEXT("Moving directly to move goal rather than following last path segment"));
 					UE_VLOG_LOCATION(this, LogPathFollowing, VeryVerbose, GoalLocation, 30, FColor::Green, TEXT("Last-segment-to-actor"));
 					UE_VLOG_SEGMENT(this, LogPathFollowing, VeryVerbose, CurrentLocation, GoalLocation, FColor::Green, TEXT_EMPTY);
 				}
-			}
 
 			UpdateMoveFocus();
 		}
@@ -918,7 +826,7 @@ void UPathFollowingComponent::UpdatePathSegment()
 					UE_VLOG(GetOwner(), LogPathFollowing, Verbose, TEXT("Path blocked, but move segment indices are not valid: start %d, end %d of %d"), MoveSegmentStartIndex, MoveSegmentEndIndex, Path->GetPathPoints().Num());
 				}
 			}
-			OnPathFinished(EPathFollowingResult::Blocked);
+			OnPathFinished(EPathFollowingResult::Blocked, FPathFollowingResultFlags::None);
 		}
 	}
 }
@@ -1058,11 +966,6 @@ bool UPathFollowingComponent::HasReachedCurrentTarget(const FVector& CurrentLoca
 	return HasReachedInternal(CurrentTarget, GoalRadius, GoalHalfHeight, CurrentLocation, CurrentAcceptanceRadius, 0.05f);
 }
 
-bool UPathFollowingComponent::HasReachedInternal(const FVector& GoalLocation, float GoalRadius, float GoalHalfHeight, const FVector& AgentLocation, float RadiusThreshold, bool bSuccessOnRadiusOverlap) const
-{
-	return HasReachedInternal(GoalLocation, GoalRadius, GoalHalfHeight, AgentLocation, RadiusThreshold, bSuccessOnRadiusOverlap ? MinAgentRadiusPct : 0.0f);
-}
-
 bool UPathFollowingComponent::HasReachedInternal(const FVector& GoalLocation, float GoalRadius, float GoalHalfHeight, const FVector& AgentLocation, float RadiusThreshold, float AgentRadiusMultiplier) const
 {
 	if (MovementComp == NULL)
@@ -1440,34 +1343,6 @@ FVector UPathFollowingComponent::GetCurrentDirection() const
 	return MoveSegmentDirection;
 }
 
-// will be deprecated soon, please don't use it!
-EPathFollowingAction::Type UPathFollowingComponent::GetPathActionType() const
-{
-	switch (Status)
-	{
-		case EPathFollowingStatus::Idle:
-			return EPathFollowingAction::NoMove;
-
-		case EPathFollowingStatus::Waiting:
-		case EPathFollowingStatus::Paused:	
-		case EPathFollowingStatus::Moving:
-			return Path.IsValid() == false ? EPathFollowingAction::Error :
-				Path->CastPath<FAbstractNavigationPath>() ? EPathFollowingAction::DirectMove :
-				Path->IsPartial() ? EPathFollowingAction::PartialPath : EPathFollowingAction::PathToGoal;
-
-		default:
-			break;
-	}
-
-	return EPathFollowingAction::Error;
-}
-
-// will be deprecated soon, please don't use it!
-FVector UPathFollowingComponent::GetPathDestination() const
-{
-	return Path.IsValid() ? Path->GetDestinationLocation() : FVector::ZeroVector;
-}
-
 bool UPathFollowingComponent::HasDirectPath() const
 {
 	return Path.IsValid() ? (Path->CastPath<FNavMeshPath>() == NULL) : false;
@@ -1630,7 +1505,7 @@ void UPathFollowingComponent::LockResource(EAIRequestPriority::Type LockSource)
 	if (bWasLocked == false)
 	{
 		UE_VLOG(GetOwner(), LogPathFollowing, Log, TEXT("Locking Move by source %s"), SourceEnum ? *SourceEnum->GetEnumName(Status) : TEXT("invalid"));
-		PauseMove();
+		PauseMove(CurrentRequestId, EPathFollowingVelocityMode::Reset);
 	}
 }
 
@@ -1665,16 +1540,105 @@ bool UPathFollowingComponent::IsResourceLocked() const
 
 void UPathFollowingComponent::SetLastMoveAtGoal(bool bFinishedAtGoal)
 {
+	if (Status == EPathFollowingStatus::Idle)
+	{
 	bLastMoveReachedGoal = bFinishedAtGoal;
 }
+}
 
 void UPathFollowingComponent::OnWaitingPathTimeout()
 {
 	if (Status == EPathFollowingStatus::Waiting)
 	{
 		UE_VLOG(GetOwner(), LogPathFollowing, Warning, TEXT("Waiting for path timeout! Aborting current move"));
-		AbortMove(TEXT("waiting timeout"), CurrentRequestId);
+		OnPathFinished(EPathFollowingResult::Invalid, FPathFollowingResultFlags::None);
+	}
+	}
+
+// deprecated functions
+FAIRequestID UPathFollowingComponent::RequestMove(FNavPathSharedPtr InPath, FRequestCompletedSignature OnComplete, const AActor* InDestinationActor, float InAcceptanceRadius, bool bInStopOnOverlap, FCustomMoveSharedPtr InGameData)
+{
+	FAIMoveRequest MoveReq;
+	if (InDestinationActor)
+	{
+		MoveReq.SetGoalActor(InDestinationActor);
+	}
+	else
+	{
+		MoveReq.SetGoalLocation(InPath.IsValid() && InPath->GetPathPoints().Num() ? InPath->GetPathPoints().Last().Location : FAISystem::InvalidLocation);
+	}
+
+	MoveReq.SetAcceptanceRadius(InAcceptanceRadius);
+	MoveReq.SetStopOnOverlap(bInStopOnOverlap);
+	MoveReq.SetUserData(InGameData);
+
+	return RequestMove(MoveReq, InPath);
+}
+
+FAIRequestID UPathFollowingComponent::RequestMove(FNavPathSharedPtr InPath, const AActor* InDestinationActor, float InAcceptanceRadius, bool bInStopOnOverlap, FCustomMoveSharedPtr InGameData)
+{
+	FAIMoveRequest MoveReq;
+	if (InDestinationActor)
+	{
+		MoveReq.SetGoalActor(InDestinationActor);
+	}
+	else
+	{
+		MoveReq.SetGoalLocation(InPath.IsValid() && InPath->GetPathPoints().Num() ? InPath->GetPathPoints().Last().Location : FAISystem::InvalidLocation);
+	}
+
+	MoveReq.SetAcceptanceRadius(InAcceptanceRadius);
+	MoveReq.SetStopOnOverlap(bInStopOnOverlap);
+	MoveReq.SetUserData(InGameData);
+
+	return RequestMove(MoveReq, InPath);
+}
+
+void UPathFollowingComponent::AbortMove(const FString& Reason, FAIRequestID RequestID, bool bResetVelocity, bool bSilent, uint8 MessageFlags)
+{
+	const FPathFollowingResultFlags::Type AbortDetails =
+		(MessageFlags == EPathFollowingMessage::NoPath) ? FPathFollowingResultFlags::InvalidPath :
+		(MessageFlags == EPathFollowingMessage::OtherRequest) ? FPathFollowingResultFlags::NewRequest :
+		FPathFollowingResultFlags::None;
+
+	UE_VLOG(GetOwner(), LogPathFollowing, Log, TEXT("AbortMove with reason(%s)"), *Reason);
+	AbortMove(*this, AbortDetails, RequestID, bResetVelocity ? EPathFollowingVelocityMode::Reset : EPathFollowingVelocityMode::Keep);
+}
+
+void UPathFollowingComponent::PauseMove(FAIRequestID RequestID, bool bResetVelocity)
+{
+	PauseMove(RequestID, bResetVelocity ? EPathFollowingVelocityMode::Reset : EPathFollowingVelocityMode::Keep);
+}
+
+EPathFollowingAction::Type UPathFollowingComponent::GetPathActionType() const
+{
+	switch (Status)
+	{
+	case EPathFollowingStatus::Idle:
+		return EPathFollowingAction::NoMove;
+
+	case EPathFollowingStatus::Waiting:
+	case EPathFollowingStatus::Paused:
+	case EPathFollowingStatus::Moving:
+		return Path.IsValid() == false ? EPathFollowingAction::Error :
+			Path->CastPath<FAbstractNavigationPath>() ? EPathFollowingAction::DirectMove :
+			Path->IsPartial() ? EPathFollowingAction::PartialPath : EPathFollowingAction::PathToGoal;
+
+	default:
+		break;
 	}
+
+	return EPathFollowingAction::Error;
+}
+
+FVector UPathFollowingComponent::GetPathDestination() const
+{
+	return Path.IsValid() ? Path->GetDestinationLocation() : FVector::ZeroVector;
+}
+
+void UPathFollowingComponent::OnPathFinished(EPathFollowingResult::Type Result)
+{
+	OnPathFinished(Result, FPathFollowingResultFlags::None);
 }
 
 #undef SHIPPING_STATIC
\ No newline at end of file
diff --git a/Engine/Source/Runtime/AIModule/Private/Perception/AISense_Blueprint.cpp b/Engine/Source/Runtime/AIModule/Private/Perception/AISense_Blueprint.cpp
--- a/Engine/Source/Runtime/AIModule/Private/Perception/AISense_Blueprint.cpp
+++ b/Engine/Source/Runtime/AIModule/Private/Perception/AISense_Blueprint.cpp
@@ -74,9 +74,9 @@ void UAISense_Blueprint::OnNewPawn(APawn& NewPawn)
 
 float UAISense_Blueprint::Update()
 {
-	const float TimeUntilNextUpdate = OnUpdate(UnprocessedEvents);
+	const float TimeUntilUpdate = OnUpdate(UnprocessedEvents);
 	UnprocessedEvents.Reset();
-	return TimeUntilNextUpdate;
+	return TimeUntilUpdate;
 }
 
 FAISenseID UAISense_Blueprint::UpdateSenseID()
diff --git a/Engine/Source/Runtime/AIModule/Private/Tasks/AITask_MoveTo.cpp b/Engine/Source/Runtime/AIModule/Private/Tasks/AITask_MoveTo.cpp
--- a/Engine/Source/Runtime/AIModule/Private/Tasks/AITask_MoveTo.cpp
+++ b/Engine/Source/Runtime/AIModule/Private/Tasks/AITask_MoveTo.cpp
@@ -64,7 +64,7 @@ void UAITask_MoveTo::FinishMoveTask(EPathFollowingResult::Type InResult)
 		if (PFComp && PFComp->GetStatus() != EPathFollowingStatus::Idle)
 		{
 			ResetObservers();
-			PFComp->AbortMove(TEXT("AITask_MoveTo finished"), MoveRequestID);
+			PFComp->AbortMove(*this, FPathFollowingResultFlags::OwnerFinished, MoveRequestID);
 		}
 	}
 
@@ -129,7 +129,7 @@ void UAITask_MoveTo::PerformMove()
 
 	case EPathFollowingRequestResult::RequestSuccessful:
 		MoveRequestID = PFComp->GetCurrentRequestId();
-		PathFinishDelegateHandle = PFComp->OnMoveFinished.AddUObject(this, &UAITask_MoveTo::OnRequestFinished);
+		PathFinishDelegateHandle = PFComp->OnRequestFinished.AddUObject(this, &UAITask_MoveTo::OnRequestFinished);
 		SetObservedPath(PFComp->GetPath());
 
 		if (TaskState == EGameplayTaskState::Finished)
@@ -191,7 +191,7 @@ void UAITask_MoveTo::ResetObservers()
 		UPathFollowingComponent* PFComp = OwnerController ? OwnerController->GetPathFollowingComponent() : nullptr;
 		if (PFComp)
 		{
-			PFComp->OnMoveFinished.Remove(PathFinishDelegateHandle);
+			PFComp->OnRequestFinished.Remove(PathFinishDelegateHandle);
 		}
 
 		PathFinishDelegateHandle.Reset();
@@ -231,9 +231,9 @@ void UAITask_MoveTo::ResetTimers()
 	}
 }
 
-void UAITask_MoveTo::OnDestroy(bool bOwnerFinished)
+void UAITask_MoveTo::OnDestroy(bool bInOwnerFinished)
 {
-	Super::OnDestroy(bOwnerFinished);
+	Super::OnDestroy(bInOwnerFinished);
 	
 	ResetObservers();
 	ResetTimers();
@@ -243,23 +243,29 @@ void UAITask_MoveTo::OnDestroy(bool bOwnerFinished)
 		UPathFollowingComponent* PFComp = OwnerController ? OwnerController->GetPathFollowingComponent() : nullptr;
 		if (PFComp && PFComp->GetStatus() != EPathFollowingStatus::Idle)
 		{
-			PFComp->AbortMove(TEXT("AITask_MoveTo finished"), MoveRequestID);
+			PFComp->AbortMove(*this, FPathFollowingResultFlags::OwnerFinished, MoveRequestID);
 		}
 	}
 }
 
-void UAITask_MoveTo::OnRequestFinished(FAIRequestID RequestID, EPathFollowingResult::Type Result)
+void UAITask_MoveTo::OnRequestFinished(FAIRequestID RequestID, const FPathFollowingResult& Result)
 {
 	if (RequestID == MoveRequestID)
 	{
-		// reset request Id, FinishMoveTask doesn't need to update path following's state
-		MoveRequestID = FAIRequestID::InvalidRequest;
-
-		FinishMoveTask(Result);
+		if (Result.IsSkipped() && !Result.HasFlag(FPathFollowingResultFlags::ForcedScript))
+		{
+			UE_VLOG(GetGameplayTasksComponent(), LogGameplayTasks, Log, TEXT("%s> ignoring OnRequestFinished, move was aborted by new request"), *GetName());
+		}
+		else
+		{
+			// reset request Id, FinishMoveTask doesn't need to update path following's state
+			MoveRequestID = FAIRequestID::InvalidRequest;
+			FinishMoveTask(Result.Code);
+		}
 	}
 	else if (IsActive())
 	{
-		UE_VLOG(GetGameplayTasksComponent(), LogGameplayTasks, Warning, TEXT("%s> received HandleMoveFinished with not matching RequestID!"), *GetName());
+		UE_VLOG(GetGameplayTasksComponent(), LogGameplayTasks, Warning, TEXT("%s> received OnRequestFinished with not matching RequestID!"), *GetName());
 	}
 }
 
diff --git a/Engine/Source/Runtime/AnimGraphRuntime/Private/BoneControllers/AnimNode_AnimDynamics.cpp b/Engine/Source/Runtime/AnimGraphRuntime/Private/BoneControllers/AnimNode_AnimDynamics.cpp
--- a/Engine/Source/Runtime/AnimGraphRuntime/Private/BoneControllers/AnimNode_AnimDynamics.cpp
+++ b/Engine/Source/Runtime/AnimGraphRuntime/Private/BoneControllers/AnimNode_AnimDynamics.cpp
@@ -38,7 +38,7 @@ FAnimNode_AnimDynamics::FAnimNode_AnimDynamics()
 , NumSolverIterationsPostUpdate(1)
 , bUsePlanarLimit(true)
 {
-
+	
 }
 
 void FAnimNode_AnimDynamics::Initialize(const FAnimationInitializeContext& Context)
@@ -59,6 +59,11 @@ void FAnimNode_AnimDynamics::Initialize(const FAnimationInitializeContext& Conte
 		PlanarLimit.DrivingBone.Initialize(RequiredBones);
 	}
 
+	for(FAnimPhysSphericalLimit& SphericalLimit : SphericalLimits)
+	{
+		SphericalLimit.DrivingBone.Initialize(RequiredBones);
+	}
+
 	if(SimulationSpace == AnimPhysSimSpaceType::BoneRelative)
 	{
 		RelativeSpaceBone.Initialize(RequiredBones);
@@ -231,10 +236,10 @@ void FAnimNode_AnimDynamics::InitializeBoneReferences(const FBoneContainer& Requ
 
 void FAnimNode_AnimDynamics::GatherDebugData(FNodeDebugData& DebugData)
 {
-	const float ActualAlpha = AlphaScaleBias.ApplyTo(Alpha);
+	const float ActualBiasedAlpha = AlphaScaleBias.ApplyTo(Alpha);
 
 	FString DebugLine = DebugData.GetNodeName(this);
-	DebugLine += FString::Printf(TEXT("(Alpha: %.1f%%)"), ActualAlpha*100.f);
+	DebugLine += FString::Printf(TEXT("(Alpha: %.1f%%)"), ActualBiasedAlpha*100.f);
 
 	DebugData.AddDebugItem(DebugLine);
 	ComponentPose.GatherDebugData(DebugData.BranchFlow(1.f));
@@ -563,7 +568,7 @@ void FAnimNode_AnimDynamics::UpdateLimits(USkeletalMeshComponent* SkelComp, FCSP
 			FAnimPhys::ConstrainConeAngle(NextTimeStep, AngularLimits, PrevBody, BoundBoneTransform.GetRotation().GetAxisX(), &RigidBody, FVector(1.0f, 0.0f, 0.0f), ConstraintSetup.ConeAngle);
 		}
 
-		if(PlanarLimits.Num() > 0)
+		if(PlanarLimits.Num() > 0 && bUsePlanarLimit)
 		{
 			for(FAnimPhysPlanarLimit& PlanarLimit : PlanarLimits)
 			{
@@ -581,6 +586,36 @@ void FAnimNode_AnimDynamics::UpdateLimits(USkeletalMeshComponent* SkelComp, FCSP
 			}
 		}
 
+		if(SphericalLimits.Num() > 0 && bUseSphericalLimits)
+		{
+			for(FAnimPhysSphericalLimit& SphericalLimit : SphericalLimits)
+			{
+				FTransform SphereTransform = FTransform::Identity;
+				SphereTransform.SetTranslation(SphericalLimit.SphereLocalOffset);
+
+				if(SphericalLimit.DrivingBone.IsValid(BoneContainer))
+				{
+					FCompactPoseBoneIndex DrivingBoneIndex = SphericalLimit.DrivingBone.GetCompactPoseIndex(BoneContainer);
+
+					FTransform DrivingBoneTransform = GetBoneTransformInSimSpace(SkelComp, MeshBases, DrivingBoneIndex);
+
+					SphereTransform *= DrivingBoneTransform;
+				}
+
+				switch(SphericalLimit.LimitType)
+				{
+				case ESphericalLimitType::Inner:
+					FAnimPhys::ConstrainSphericalInner(NextTimeStep, LinearLimits, &RigidBody, SphereTransform, SphericalLimit.LimitRadius);
+					break;
+				case ESphericalLimitType::Outer:
+					FAnimPhys::ConstrainSphericalOuter(NextTimeStep, LinearLimits, &RigidBody, SphereTransform, SphericalLimit.LimitRadius);
+					break;
+				default:
+					break;
+				}
+			}
+		}
+
 		// Add spring if we need spring forces
 		if (bAngularSpring || bLinearSpring)
 		{
diff --git a/Engine/Source/Runtime/AnimGraphRuntime/Private/BoneControllers/AnimNode_BoneDrivenController.cpp b/Engine/Source/Runtime/AnimGraphRuntime/Private/BoneControllers/AnimNode_BoneDrivenController.cpp
--- a/Engine/Source/Runtime/AnimGraphRuntime/Private/BoneControllers/AnimNode_BoneDrivenController.cpp
+++ b/Engine/Source/Runtime/AnimGraphRuntime/Private/BoneControllers/AnimNode_BoneDrivenController.cpp
@@ -230,8 +230,8 @@ const float FAnimNode_BoneDrivenController::ExtractSourceValue(const FTransform
 		// Apply the fixed function remapping/clamping
 		if (bUseRange)
 		{
-			const float Alpha = FMath::Clamp(FMath::GetRangePct(RangeMin, RangeMax, FinalDriverValue), 0.0f, 1.0f);
-			FinalDriverValue = FMath::Lerp(RemappedMin, RemappedMax, Alpha);
+			const float ClampedAlpha = FMath::Clamp(FMath::GetRangePct(RangeMin, RangeMax, FinalDriverValue), 0.0f, 1.0f);
+			FinalDriverValue = FMath::Lerp(RemappedMin, RemappedMax, ClampedAlpha);
 		}
 
 		FinalDriverValue *= Multiplier;
diff --git a/Engine/Source/Runtime/AnimGraphRuntime/Private/BoneControllers/AnimNode_ModifyBone.cpp b/Engine/Source/Runtime/AnimGraphRuntime/Private/BoneControllers/AnimNode_ModifyBone.cpp
--- a/Engine/Source/Runtime/AnimGraphRuntime/Private/BoneControllers/AnimNode_ModifyBone.cpp
+++ b/Engine/Source/Runtime/AnimGraphRuntime/Private/BoneControllers/AnimNode_ModifyBone.cpp
@@ -39,7 +39,7 @@ void FAnimNode_ModifyBone::EvaluateBoneTransforms(USkeletalMeshComponent* SkelCo
 	// the way we apply transform is same as FMatrix or FTransform
 	// we apply scale first, and rotation, and translation
 	// if you'd like to translate first, you'll need two nodes that first node does translate and second nodes to rotate.
-	const FBoneContainer BoneContainer = MeshBases.GetPose().GetBoneContainer();
+	const FBoneContainer& BoneContainer = MeshBases.GetPose().GetBoneContainer();
 
 	FCompactPoseBoneIndex CompactPoseBoneToModify = BoneToModify.GetCompactPoseIndex(BoneContainer);
 	FTransform NewBoneTM = MeshBases.GetComponentSpaceTransform(CompactPoseBoneToModify);
@@ -111,4 +111,4 @@ bool FAnimNode_ModifyBone::IsValidToEvaluate(const USkeleton* Skeleton, const FB
 void FAnimNode_ModifyBone::InitializeBoneReferences(const FBoneContainer& RequiredBones) 
 {
 	BoneToModify.Initialize(RequiredBones);
-}
\ No newline at end of file
+}
diff --git a/Engine/Source/Runtime/AnimGraphRuntime/Private/BoneControllers/AnimNode_ObserveBone.cpp b/Engine/Source/Runtime/AnimGraphRuntime/Private/BoneControllers/AnimNode_ObserveBone.cpp
--- a/Engine/Source/Runtime/AnimGraphRuntime/Private/BoneControllers/AnimNode_ObserveBone.cpp
+++ b/Engine/Source/Runtime/AnimGraphRuntime/Private/BoneControllers/AnimNode_ObserveBone.cpp
@@ -27,7 +27,7 @@ void FAnimNode_ObserveBone::GatherDebugData(FNodeDebugData& DebugData)
 
 void FAnimNode_ObserveBone::EvaluateBoneTransforms(USkeletalMeshComponent* SkelComp, FCSPose<FCompactPose>& MeshBases, TArray<FBoneTransform>& OutBoneTransforms)
 {
-	const FBoneContainer BoneContainer = MeshBases.GetPose().GetBoneContainer();
+	const FBoneContainer& BoneContainer = MeshBases.GetPose().GetBoneContainer();
 
 	const FCompactPoseBoneIndex BoneIndex = BoneToObserve.GetCompactPoseIndex(BoneContainer);
 	FTransform BoneTM = MeshBases.GetComponentSpaceTransform(BoneIndex);
diff --git a/Engine/Source/Runtime/AnimGraphRuntime/Private/BoneControllers/AnimNode_SpringBone.cpp b/Engine/Source/Runtime/AnimGraphRuntime/Private/BoneControllers/AnimNode_SpringBone.cpp
--- a/Engine/Source/Runtime/AnimGraphRuntime/Private/BoneControllers/AnimNode_SpringBone.cpp
+++ b/Engine/Source/Runtime/AnimGraphRuntime/Private/BoneControllers/AnimNode_SpringBone.cpp
@@ -51,11 +51,11 @@ void FAnimNode_SpringBone::UpdateInternal(const FAnimationUpdateContext& Context
 
 void FAnimNode_SpringBone::GatherDebugData(FNodeDebugData& DebugData)
 {
-	const float ActualAlpha = AlphaScaleBias.ApplyTo(Alpha);
+	const float ActualBiasedAlpha = AlphaScaleBias.ApplyTo(Alpha);
 
 	//MDW_TODO Add more output info?
 	FString DebugLine = DebugData.GetNodeName(this);
-	DebugLine += FString::Printf(TEXT("(Alpha: %.1f%% RemainingTime: %.3f)"), ActualAlpha*100.f, RemainingTime);
+	DebugLine += FString::Printf(TEXT("(Alpha: %.1f%% RemainingTime: %.3f)"), ActualBiasedAlpha*100.f, RemainingTime);
 
 	DebugData.AddDebugItem(DebugLine);
 	ComponentPose.GatherDebugData(DebugData);
diff --git a/Engine/Source/Runtime/AnimGraphRuntime/Public/BoneControllers/AnimNode_AnimDynamics.h b/Engine/Source/Runtime/AnimGraphRuntime/Public/BoneControllers/AnimNode_AnimDynamics.h
--- a/Engine/Source/Runtime/AnimGraphRuntime/Public/BoneControllers/AnimNode_AnimDynamics.h
+++ b/Engine/Source/Runtime/AnimGraphRuntime/Public/BoneControllers/AnimNode_AnimDynamics.h
@@ -160,6 +160,42 @@ struct FAnimPhysPlanarLimit
 	FTransform PlaneTransform;
 };
 
+/** Whether spheres keep bodies inside, or outside of their shape */
+UENUM()
+enum class ESphericalLimitType : uint8
+{
+	Inner,
+	Outer
+};
+
+USTRUCT()
+struct FAnimPhysSphericalLimit
+{
+	GENERATED_BODY();
+
+	FAnimPhysSphericalLimit()
+		: SphereLocalOffset(FVector::ZeroVector)
+		, LimitRadius(0.0f)
+		, LimitType(ESphericalLimitType::Outer)
+	{}
+
+	/** Bone to attach the sphere to */
+	UPROPERTY(EditAnywhere, Category = SphericalLimit)
+	FBoneReference DrivingBone;
+
+	/** Local offset for the sphere, if no driving bone is set this is in node space, otherwise bone space */
+	UPROPERTY(EditAnywhere, Category = SphericalLimit)
+	FVector SphereLocalOffset;
+
+	/** Radius of the sphere */
+	UPROPERTY(EditAnywhere, Category = SphericalLimit)
+	float LimitRadius;
+
+	/** Whether to lock bodies inside or outside of the sphere */
+	UPROPERTY(EditAnywhere, Category = SphericalLimit)
+	ESphericalLimitType LimitType;
+};
+
 USTRUCT()
 struct ANIMGRAPHRUNTIME_API FAnimNode_AnimDynamics : public FAnimNode_SkeletalControlBase
 {
@@ -259,18 +295,28 @@ struct ANIMGRAPHRUNTIME_API FAnimNode_AnimDynamics : public FAnimNode_SkeletalCo
 	UPROPERTY(EditAnywhere, BlueprintReadWrite, Category = Constraint)
 	FAnimPhysConstraintSetup ConstraintSetup;
 
+	/** Whether to evaluate planar limits */
 	UPROPERTY(EditAnywhere, BlueprintReadWrite, Category=PlanarLimit)
 	bool bUsePlanarLimit;
 
+	/** List of available planar limits for this node */
 	UPROPERTY(EditAnywhere, BlueprintReadWrite, Category=PlanarLimit)
 	TArray<FAnimPhysPlanarLimit> PlanarLimits;
 
+	/** Whether to evaluate spherical limits */
+	UPROPERTY(EditAnywhere, BlueprintReadWrite, Category = SphericalLimit)
+	bool bUseSphericalLimits;
+
+	/** List of available spherical limits for this node */
+	UPROPERTY(EditAnywhere, BlueprintReadWrite, Category = SphericalLimit)
+	TArray<FAnimPhysSphericalLimit> SphericalLimits;
+
 	/** Resolution method for planar limits */
-	UPROPERTY(EditAnywhere, Category = PlanarLimit)
+	UPROPERTY(EditAnywhere, Category = Collision)
 	AnimPhysCollisionType CollisionType;
 
 	/** Radius to use if CollisionType is set to CustomSphere */
-	UPROPERTY(EditAnywhere, Category = PlanarLimit, meta = (UIMin = "1", ClampMin = "1"))
+	UPROPERTY(EditAnywhere, Category = Collision, meta = (UIMin = "1", ClampMin = "1"))
 	float SphereCollisionRadius;
 
 	/** An external force to apply to all bodies in the simulation when ticked, specified in world space */
@@ -371,4 +417,4 @@ private:
 
 	// Gravity direction in sim space
 	FVector SimSpaceGravityDirection;
-};
\ No newline at end of file
+};
diff --git a/Engine/Source/Runtime/Core/Public/Delegates/DelegateSignatureImpl_Variadics.inl b/Engine/Source/Runtime/Core/Public/Delegates/DelegateSignatureImpl_Variadics.inl
--- a/Engine/Source/Runtime/Core/Public/Delegates/DelegateSignatureImpl_Variadics.inl
+++ b/Engine/Source/Runtime/Core/Public/Delegates/DelegateSignatureImpl_Variadics.inl
@@ -873,9 +873,9 @@ protected:
 			{
 				if (IDelegateInstance* OtherInstance = Super::GetDelegateInstanceProtectedHelper(OtherDelegateRef))
 				{
-					FDelegate Temp;
-					((TDelegateInstanceInterface*)OtherInstance)->CreateCopy(Temp);
-					AddInternal(MoveTemp(Temp));
+					FDelegate TempDelegate;
+					((TDelegateInstanceInterface*)OtherInstance)->CreateCopy(TempDelegate);
+					AddInternal(MoveTemp(TempDelegate));
 				}
 			}
 		}
diff --git a/Engine/Source/Runtime/Core/Public/UObject/UnrealNames.inl b/Engine/Source/Runtime/Core/Public/UObject/UnrealNames.inl
--- a/Engine/Source/Runtime/Core/Public/UObject/UnrealNames.inl
+++ b/Engine/Source/Runtime/Core/Public/UObject/UnrealNames.inl
@@ -30,6 +30,8 @@ REGISTER_NAME(23,Int64Property)
 REGISTER_NAME(25,Int16Property)
 REGISTER_NAME(26,Int8Property)
 REGISTER_NAME(27,AssetSubclassOfProperty)
+REGISTER_NAME(28,MapProperty)
+REGISTER_NAME(29,SetProperty)
 
 // Special packages.
 REGISTER_NAME(30,Core)
diff --git a/Engine/Source/Runtime/CoreUObject/Private/Serialization/ArchiveUObject.cpp b/Engine/Source/Runtime/CoreUObject/Private/Serialization/ArchiveUObject.cpp
--- a/Engine/Source/Runtime/CoreUObject/Private/Serialization/ArchiveUObject.cpp
+++ b/Engine/Source/Runtime/CoreUObject/Private/Serialization/ArchiveUObject.cpp
@@ -104,11 +104,11 @@ FArchive& FObjectAndNameAsStringProxyArchive::operator<<(class UObject*& Obj)
 		FString LoadedString;
 		InnerArchive << LoadedString;
 		// look up the object by fully qualified pathname
-		Obj = FindObject<UObject>(NULL, *LoadedString, false);
+		Obj = FindObject<UObject>(nullptr, *LoadedString, false);
 		// If we couldn't find it, and we want to load it, do that
-		if((Obj == NULL) && bLoadIfFindFails)
+		if(Obj && bLoadIfFindFails)
 		{
-			Obj = LoadObject<UObject>(NULL, *LoadedString);
+			Obj = LoadObject<UObject>(nullptr, *LoadedString);
 		}
 	}
 	else
@@ -137,7 +137,7 @@ void FSerializedPropertyScope::PopEditorOnlyProperty()
 }
 #endif
 
-void FArchiveReplaceObjectRefBase::SerializeObject(UObject* ObjectToSerialzie)
+void FArchiveReplaceObjectRefBase::SerializeObject(UObject* ObjectToSerialize)
 {
 	// Simple FReferenceCollector proxy for FArchiveReplaceObjectRefBase
 	class FReplaceObjectRefCollector : public FReferenceCollector
@@ -177,23 +177,23 @@ void FArchiveReplaceObjectRefBase::SerializeObject(UObject* ObjectToSerialzie)
 	// default objects may be serialized during script compilation while the script
 	// and C++ versions of a class are not in sync), so use SerializeTaggedProperties()
 	// rather than the native Serialize() function
-	UClass* ObjectClass = ObjectToSerialzie->GetClass();
-	if (ObjectToSerialzie->HasAnyFlags(RF_ClassDefaultObject))
+	UClass* ObjectClass = ObjectToSerialize->GetClass();
+	if (ObjectToSerialize->HasAnyFlags(RF_ClassDefaultObject))
 	{		
 		StartSerializingDefaults();
 		if (!WantBinaryPropertySerialization() && (IsLoading() || IsSaving()))
 		{
-			ObjectClass->SerializeTaggedProperties(*this, (uint8*)ObjectToSerialzie, ObjectClass, NULL);
+			ObjectClass->SerializeTaggedProperties(*this, (uint8*)ObjectToSerialize, ObjectClass, nullptr);
 		}
 		else
 		{
-			ObjectClass->SerializeBin(*this, ObjectToSerialzie);
+			ObjectClass->SerializeBin(*this, ObjectToSerialize);
 		}
 		StopSerializingDefaults();
 	}
 	else
 	{
-		ObjectToSerialzie->Serialize(*this);
+		ObjectToSerialize->Serialize(*this);
 	}
-	ObjectClass->CallAddReferencedObjects(ObjectToSerialzie, ReplaceRefCollector);
+	ObjectClass->CallAddReferencedObjects(ObjectToSerialize, ReplaceRefCollector);
 }
\ No newline at end of file
diff --git a/Engine/Source/Runtime/CoreUObject/Private/UObject/UObjectGlobals.cpp b/Engine/Source/Runtime/CoreUObject/Private/UObject/UObjectGlobals.cpp
--- a/Engine/Source/Runtime/CoreUObject/Private/UObject/UObjectGlobals.cpp
+++ b/Engine/Source/Runtime/CoreUObject/Private/UObject/UObjectGlobals.cpp
@@ -75,7 +75,7 @@ FCoreUObjectDelegates::FStringAssetReferenceLoaded FCoreUObjectDelegates::String
 FCoreUObjectDelegates::FStringAssetReferenceSaving FCoreUObjectDelegates::StringAssetReferenceSaving;
 FCoreUObjectDelegates::FPackageCreatedForLoad FCoreUObjectDelegates::PackageCreatedForLoad;
 
-/** Check wehether we should report progress or not */
+/** Check whether we should report progress or not */
 bool ShouldReportProgress()
 {
 	return GIsEditor && IsInGameThread() && !IsRunningCommandlet() && !IsAsyncLoading();
@@ -83,7 +83,7 @@ bool ShouldReportProgress()
 
 /**
  * Returns true if code is called form the game thread while collecting garbage.
- * We only have tp guard agains StaticFindObject on the game thread as other threads will be blocked anyway
+ * We only have to guard against StaticFindObject on the game thread as other threads will be blocked anyway
  */
 static FORCEINLINE bool IsGarbageCollectingOnGameThread()
 {
diff --git a/Engine/Source/Runtime/CoreUObject/Public/Serialization/ArchiveUObject.h b/Engine/Source/Runtime/CoreUObject/Public/Serialization/ArchiveUObject.h
--- a/Engine/Source/Runtime/CoreUObject/Public/Serialization/ArchiveUObject.h
+++ b/Engine/Source/Runtime/CoreUObject/Public/Serialization/ArchiveUObject.h
@@ -1033,7 +1033,7 @@ protected:
 	/**
 	* Serializes a single object
 	*/
-	void SerializeObject(UObject* ObjectToSerialzie);
+	void SerializeObject(UObject* ObjectToSerialize);
 };
 
 /*----------------------------------------------------------------------------
