Mac non unity compilation fix #UE4

diff --git a/Engine/Source/Runtime/AIModule/Private/EnvironmentQuery/EQSTestingPawn.cpp b/Engine/Source/Runtime/AIModule/Private/EnvironmentQuery/EQSTestingPawn.cpp
--- a/Engine/Source/Runtime/AIModule/Private/EnvironmentQuery/EQSTestingPawn.cpp
+++ b/Engine/Source/Runtime/AIModule/Private/EnvironmentQuery/EQSTestingPawn.cpp
@@ -7,6 +7,7 @@
 #include "EnvironmentQuery/EnvQueryTypes.h"
 #include "EnvironmentQuery/EnvQueryManager.h"
 #include "EnvironmentQuery/EQSRenderingComponent.h"
+#include "Components/SkeletalMeshComponent.h"
 
 #if WITH_EDITORONLY_DATA
 #include "Components/ArrowComponent.h"
