CIS fix for CrowdManager.cpp

diff --git a/Engine/Source/Runtime/AIModule/Classes/Navigation/CrowdManager.h b/Engine/Source/Runtime/AIModule/Classes/Navigation/CrowdManager.h
--- a/Engine/Source/Runtime/AIModule/Classes/Navigation/CrowdManager.h
+++ b/Engine/Source/Runtime/AIModule/Classes/Navigation/CrowdManager.h
@@ -2,6 +2,7 @@
 
 #pragma once
 #include "Tickable.h"
+#include "DrawDebugHelpers.h"
 #include "CrowdManager.generated.h"
 
 class UWorld;
diff --git a/Engine/Source/Runtime/AIModule/Private/Navigation/CrowdManager.cpp b/Engine/Source/Runtime/AIModule/Private/Navigation/CrowdManager.cpp
--- a/Engine/Source/Runtime/AIModule/Private/Navigation/CrowdManager.cpp
+++ b/Engine/Source/Runtime/AIModule/Private/Navigation/CrowdManager.cpp
@@ -13,8 +13,6 @@
 #include "Navigation/CrowdFollowingComponent.h"
 #include "Navigation/CrowdAgentInterface.h"
 
-#include "DrawDebugHelpers.h"
-
 DECLARE_STATS_GROUP(TEXT("Crowd"), STATGROUP_AICrowd, STATCAT_Advanced);
 
 DECLARE_CYCLE_STAT(TEXT("Nav Tick: crowd simulation"), STAT_AI_Crowd_Tick, STATGROUP_AICrowd);
@@ -1082,6 +1080,7 @@ void UCrowdManager::DebugTick() const
 		}
 	}
 
+#if ENABLE_DRAW_DEBUG
 	// on screen debugging
 	const dtCrowdAgent* SelectedAgent = DetourAgentDebug->idx >= 0 ? DetourCrowd->getAgent(DetourAgentDebug->idx) : NULL;
 	if (SelectedAgent && CrowdDebugDrawing::bDebugSelectedActors)
@@ -1121,6 +1120,7 @@ void UCrowdManager::DebugTick() const
 	{
 		DrawDebugSharedBoundary();
 	}
+#endif // ENABLE_DRAW_DEBUG
 
 	// vislog debugging
 	if (CrowdDebugDrawing::bDebugVisLog)
