Manually merged over CL#3156193 down from Dev-Framework #UE4

Original comment:
Fixed a bug in pathfollowing's new acceptance radius usage resulting is mishandling of fake paths (like Paragon's flow field path)

diff --git a/Engine/Source/Runtime/AIModule/Private/Navigation/PathFollowingComponent.cpp b/Engine/Source/Runtime/AIModule/Private/Navigation/PathFollowingComponent.cpp
--- a/Engine/Source/Runtime/AIModule/Private/Navigation/PathFollowingComponent.cpp
+++ b/Engine/Source/Runtime/AIModule/Private/Navigation/PathFollowingComponent.cpp
@@ -281,7 +281,10 @@ bool UPathFollowingComponent::HandlePathUpdateEvent()
 		float IgnoreGoalRadius, IgnoreGoalHalfHeight;
 		OriginalMoveRequestGoalLocation = FindGoalLocation(*this, *PathGoalActor, Path->GetGoalActorAsNavAgent(), IgnoreGoalRadius, IgnoreGoalHalfHeight);
 	}
-	PreciseAcceptanceRadiusCheckStartNodeIndex = FindPreciseAcceptanceRadiusTestsStartNodeIndex(*Path, OriginalMoveRequestGoalLocation);
+	if (FAISystem::IsValidLocation(OriginalMoveRequestGoalLocation))
+	{
+		PreciseAcceptanceRadiusCheckStartNodeIndex = FindPreciseAcceptanceRadiusTestsStartNodeIndex(*Path, OriginalMoveRequestGoalLocation);
+	}
 
 	OnPathUpdated();
 	GetWorld()->GetTimerManager().ClearTimer(WaitingForPathTimer);
@@ -301,9 +304,8 @@ bool UPathFollowingComponent::HandlePathUpdateEvent()
 void UPathFollowingComponent::SetAcceptanceRadius(const float InAcceptanceRadius)
 {
 	AcceptanceRadius = InAcceptanceRadius;
-	if (Path.IsValid() && Path->IsValid())
+	if (Path.IsValid() && Path->IsValid() && FAISystem::IsValidLocation(OriginalMoveRequestGoalLocation))
 	{
-		ensure(FAISystem::IsValidLocation(OriginalMoveRequestGoalLocation));
 		// figure out when we need to start doing precise end-condition tests
 		PreciseAcceptanceRadiusCheckStartNodeIndex = FindPreciseAcceptanceRadiusTestsStartNodeIndex(*Path, OriginalMoveRequestGoalLocation);
 	}
@@ -1105,7 +1107,7 @@ bool UPathFollowingComponent::HasReachedInternal(const FVector& GoalLocation, fl
 
 int32 UPathFollowingComponent::FindPreciseAcceptanceRadiusTestsStartNodeIndex(const FNavigationPath& PathInstance, const FVector& GoalLocation) const
 {
-	const float DistanceFromEndOfPath = GetFinalAcceptanceRadius(PathInstance, OriginalMoveRequestGoalLocation);
+	const float DistanceFromEndOfPath = GetFinalAcceptanceRadius(PathInstance, GoalLocation);
 	// @todo support based paths
 	
 	int32 TestStartIndex = MAX_int32;
@@ -1139,6 +1141,8 @@ float UPathFollowingComponent::GetFinalAcceptanceRadius(const FNavigationPath& P
 {
 	if (PathInstance.IsPartial())
 	{
+		ensure(FAISystem::IsValidLocation(OriginalGoalLocation));
+
 		const float PathEndToGoalDistance = FVector::Dist(PathEndOverride ? *PathEndOverride : PathInstance.GetEndLocation(), OriginalGoalLocation);
 		const float Remaining = AcceptanceRadius - PathEndToGoalDistance;
 
