Fixed requesting move while another one is in progress always resulting in reseting velocity #UE4

diff --git a/Engine/Source/Runtime/AIModule/Private/Navigation/PathFollowingComponent.cpp b/Engine/Source/Runtime/AIModule/Private/Navigation/PathFollowingComponent.cpp
--- a/Engine/Source/Runtime/AIModule/Private/Navigation/PathFollowingComponent.cpp
+++ b/Engine/Source/Runtime/AIModule/Private/Navigation/PathFollowingComponent.cpp
@@ -347,9 +347,15 @@ FAIRequestID UPathFollowingComponent::RequestMove(const FAIMoveRequest& RequestD
 	{
 		if (Status != EPathFollowingStatus::Idle)
 		{
+			// setting to false to make sure OnPathFinished won't result in stoping 
+			bStopMovementOnFinish = false;
 			OnPathFinished(EPathFollowingResult::Aborted, FPathFollowingResultFlags::NewRequest);
 		}
 		
+		// force-setting this to true since this is the default state, and gets modified only 
+		// when aborting movement.
+		bStopMovementOnFinish = true;
+		
 		Reset();
 		StoreRequestId();
 		
