Copying //UE4/Orion-Staging to //UE4/Main (Source: //Orion/Dev-General @ 3196521)

diff --git a/Engine/Source/Programs/AutomationTool/Scripts/RebuildLightMapsCommand.Automation.cs b/Engine/Source/Programs/AutomationTool/Scripts/RebuildLightMapsCommand.Automation.cs
--- a/Engine/Source/Programs/AutomationTool/Scripts/RebuildLightMapsCommand.Automation.cs
+++ b/Engine/Source/Programs/AutomationTool/Scripts/RebuildLightMapsCommand.Automation.cs
@@ -281,6 +281,8 @@ namespace AutomationScripts.Automation
             Message.CC.Add(new MailAddress("[REDACTED_EMAIL]"));
 			Message.CC.Add(new MailAddress("[REDACTED_EMAIL]"));
 			Message.CC.Add(new MailAddress("[REDACTED_EMAIL]"));
+			Message.CC.Add(new MailAddress("[REDACTED_EMAIL]"));
+			Message.CC.Add(new MailAddress("[REDACTED_EMAIL]"));
 			Message.Subject = String.Format("Nightly lightmap rebuild {0} for {1}", bWasSuccessful ? "[SUCCESS]" : "[FAILED]", Branch);
 			Message.Body = MessageBody;
             /*Attachment Attach = new Attachment();
diff --git a/Engine/Source/Programs/UnrealBuildTool/Configuration/BuildConfiguration.cs b/Engine/Source/Programs/UnrealBuildTool/Configuration/BuildConfiguration.cs
--- a/Engine/Source/Programs/UnrealBuildTool/Configuration/BuildConfiguration.cs
+++ b/Engine/Source/Programs/UnrealBuildTool/Configuration/BuildConfiguration.cs
@@ -38,6 +38,14 @@ namespace UnrealBuildTool
 		[XmlConfig]
 		public static bool bUseAdaptiveUnityBuild;
 
+		/// <summary>
+		/// An experimental new feature that, when enabled, will disable optimization of files that are omitted from Unity so they they become
+		/// easier to debug.
+		/// IMPORTANT: This feature is not complete yet!  Currently, it requires source files to be read-only that you aren't actively working with!
+		/// </summary>
+		[XmlConfig]
+		public static bool bAdaptiveUnityDisablesOptimizations;
+
 		/// <summary>
 		/// The number of source files in a game module before unity build will be activated for that module.  This
 		/// allows small game modules to have faster iterative compile times for single files, at the expense of slower full
@@ -236,6 +244,13 @@ namespace UnrealBuildTool
 		[XmlConfig]
 		public static bool bAllowLTCG;
 
+		/// <summary>
+		/// Whether to allow the use of ASLR (address space layout randomization) if supported. Only
+		/// applies to shipping builds.
+		/// </summary>
+		[XmlConfig]
+		public static bool bAllowASLRInShipping;
+
 		/// <summary>
 		/// Whether to support edit and continue.  Only works on Microsoft compilers in 32-bit compiles.
 		/// </summary>
@@ -567,6 +582,7 @@ namespace UnrealBuildTool
 		public static void LoadDefaults()
 		{
 			bAllowLTCG = false;
+			bAllowASLRInShipping = true;
 			bAllowRemotelyCompiledPCHs = false;
 			bAllowXGE = true;
 			bXGENoWatchdogThread = false;
@@ -627,6 +643,11 @@ namespace UnrealBuildTool
 			// disabled by default until we test it sufficiently
 			bUseAdaptiveUnityBuild = true;
 
+            // This experimental feature may help to improve debugging when working on optimized builds. If a file is omitted from
+            // unity compilation due to being checked-out / editable, an assumption is made that you are working on this file
+            // and may be debugging it. Disabled by default until we test it sufficiently
+            bAdaptiveUnityDisablesOptimizations = false;
+
 			bForceUnityBuild = false;
 
 			// For programmers working on gameplay code, it's convenient to have fast iteration times when changing only single source
diff --git a/Engine/Source/Programs/UnrealBuildTool/Configuration/UEBuildModule.cs b/Engine/Source/Programs/UnrealBuildTool/Configuration/UEBuildModule.cs
--- a/Engine/Source/Programs/UnrealBuildTool/Configuration/UEBuildModule.cs
+++ b/Engine/Source/Programs/UnrealBuildTool/Configuration/UEBuildModule.cs
@@ -841,12 +841,12 @@ namespace UnrealBuildTool
 		/// </summary>
 		public bool bSkipDefinitionsForCompileEnvironment = false;
 
-        public int FindNumberOfGeneratedCppFiles()
-        {
-           return ((null == GeneratedCodeDirectory) || !GeneratedCodeDirectory.Exists()) ? 0
-                : (GeneratedCodeDirectory.EnumerateFileReferences("*.generated.*.cpp", SearchOption.AllDirectories).Count() 
-                 + GeneratedCodeDirectory.EnumerateFileReferences("*.generated.cpp", SearchOption.AllDirectories).Count());
-        }
+		public int FindNumberOfGeneratedCppFiles()
+		{
+			return ((null == GeneratedCodeDirectory) || !GeneratedCodeDirectory.Exists()) ? 0
+				 : (GeneratedCodeDirectory.EnumerateFileReferences("*.generated.*.cpp", SearchOption.AllDirectories).Count()
+				  + GeneratedCodeDirectory.EnumerateFileReferences("*.generated.cpp", SearchOption.AllDirectories).Count());
+		}
 
 		/// <summary>
 		/// Categorizes source files into per-extension buckets
@@ -1117,7 +1117,7 @@ namespace UnrealBuildTool
 								SharedPCHHeaderInfo CurSharedPCHHeaderFile = GlobalCompileEnvironment.SharedPCHHeaderFiles[SharedPCHHeaderFileIndex];
 
 								if (DependencyModule == CurSharedPCHHeaderFile.Module ||
-									(bIsASharedPCHModule && CurSharedPCHHeaderFile.Module == this))	// If we ourselves are a shared PCH module, always at least use our own module as our shared PCH header if we can't find anything better
+									(bIsASharedPCHModule && CurSharedPCHHeaderFile.Module == this)) // If we ourselves are a shared PCH module, always at least use our own module as our shared PCH header if we can't find anything better
 								{
 									SharedPCHModuleName = CurSharedPCHHeaderFile.Module.Name;
 									SharedPCHHeaderFile = CurSharedPCHHeaderFile.PCHHeaderFile;
@@ -1146,7 +1146,7 @@ namespace UnrealBuildTool
 						// Don't allow game modules to use engine PCHs in DebugGame - the optimization settings aren't correct. 
 						// @todo: we should be creating shared PCHs ahead of time, and only using them if our settings match. as it is, the first modules compiled
 						// (which are currently plugins) get to call the shots for how the shared PCH gets built, and that might be a game plugin built in debug...
-						if(Target.Configuration == UnrealTargetConfiguration.DebugGame && SharedPCHHeaderFile.Reference.IsUnderDirectory(UnrealBuildTool.EngineDirectory) && !RulesFile.IsUnderDirectory(UnrealBuildTool.EngineDirectory))
+						if (Target.Configuration == UnrealTargetConfiguration.DebugGame && SharedPCHHeaderFile.Reference.IsUnderDirectory(UnrealBuildTool.EngineDirectory) && !RulesFile.IsUnderDirectory(UnrealBuildTool.EngineDirectory))
 						{
 							SharedPCHModuleName = null;
 							SharedPCHHeaderFile = null;
@@ -1258,7 +1258,7 @@ namespace UnrealBuildTool
 					{
 						if (SharedPCHHeaderFile != null || CPPFilesToBuild.Count >= MinFilesUsingPrecompiledHeader)
 						{
-                            CPPOutput PCHOutput;
+							CPPOutput PCHOutput;
 							if (SharedPCHHeaderFile == null)
 							{
 								PCHOutput = PrecompileHeaderEnvironment.GeneratePCHCreationAction(
@@ -1293,7 +1293,7 @@ namespace UnrealBuildTool
 										SharedPCHCompileEnvironment.Config.Definitions,
 										SharedPCHCompileEnvironment.Config.AdditionalFrameworks);
 								}
-                                
+
 								PCHOutput = PrecompileHeaderEnvironment.GeneratePCHCreationAction(
 									ToolChain,
 									Target,
@@ -1301,7 +1301,7 @@ namespace UnrealBuildTool
 									ModulePCHEnvironment.PrecompiledHeaderIncludeFilename,
 									SharedPCHCompileEnvironment,
 									DirectoryReference.Combine(CompileEnvironment.Config.OutputDirectory, "SharedPCHs"),
-									(CompileEnvironment.Config.PCHOutputDirectory == null)? null : DirectoryReference.Combine(CompileEnvironment.Config.PCHOutputDirectory, "SharedPCHs"),
+									(CompileEnvironment.Config.PCHOutputDirectory == null) ? null : DirectoryReference.Combine(CompileEnvironment.Config.PCHOutputDirectory, "SharedPCHs"),
 									"Shared",
 									false);
 							}
@@ -1329,7 +1329,14 @@ namespace UnrealBuildTool
 						CPPCompileEnvironment = ModulePCHCompileEnvironment;
 					}
 
-					LinkInputFiles.AddRange(ToolChain.CompileCPPFiles(Target, CPPCompileEnvironment, CPPFilesToBuild, Name).ObjectFiles);
+					if (bModuleUsesUnityBuild)
+					{
+						LinkInputFiles.AddRange(CompileUnityFilesWithToolChain(ToolChain, Target, CPPCompileEnvironment, CPPFilesToBuild, Name).ObjectFiles);
+					}
+					else
+					{
+						LinkInputFiles.AddRange(ToolChain.CompileCPPFiles(Target, CPPCompileEnvironment, CPPFilesToBuild, Name).ObjectFiles);
+					}
 					bWasModuleCodeCompiled = true;
 				}
 
@@ -1346,8 +1353,13 @@ namespace UnrealBuildTool
 				if (bModuleUsesUnityBuild)
 				{
 					CPPFilesToCompile = Unity.GenerateUnityCPPs(ToolChain, Target, CPPFilesToCompile, CPPCompileEnvironment, Name);
+					LinkInputFiles.AddRange(CompileUnityFilesWithToolChain(ToolChain, Target, CPPCompileEnvironment, CPPFilesToCompile, Name).ObjectFiles);
 				}
-				LinkInputFiles.AddRange(ToolChain.CompileCPPFiles(Target, CPPCompileEnvironment, CPPFilesToCompile, Name).ObjectFiles);
+				else
+				{
+					LinkInputFiles.AddRange(ToolChain.CompileCPPFiles(Target, CPPCompileEnvironment, CPPFilesToCompile, Name).ObjectFiles);
+				}
+
 			}
 
 			if (AutoGenerateCppInfo != null && AutoGenerateCppInfo.BuildInfo != null && !CPPCompileEnvironment.bHackHeaderGenerator)
@@ -1379,6 +1391,60 @@ namespace UnrealBuildTool
 			return LinkInputFiles;
 		}
 
+		/// <summary>
+		/// Compiles the provided CPP unity files. Will
+		/// </summary>
+		private CPPOutput CompileUnityFilesWithToolChain(UEToolChain ToolChain, UEBuildTarget Target, CPPEnvironment CompileEnvironment, List<FileItem> SourceFiles, string ModuleName)
+		{
+			List<FileItem> NormalFiles = new List<FileItem>();
+			List<FileItem> AdaptiveFiles = new List<FileItem>();
+
+			if (BuildConfiguration.bAdaptiveUnityDisablesOptimizations && !BuildConfiguration.bStressTestUnity)
+			{
+				foreach (FileItem File in SourceFiles)
+				{
+					// Basic check as to whether something in this module is/isn't a unity file...
+					if (File.ToString().StartsWith(Unity.ModulePrefix))
+					{
+						NormalFiles.Add(File);
+					}
+					else
+					{
+						AdaptiveFiles.Add(File);
+					}
+				}
+			}
+			else
+			{
+				NormalFiles.AddRange(SourceFiles);
+			}
+
+			CPPOutput OutputFiles = new CPPOutput();
+
+			if (NormalFiles.Count > 0)
+			{
+				OutputFiles = ToolChain.CompileCPPFiles(Target, CompileEnvironment, NormalFiles, Name);
+			}
+
+			if (AdaptiveFiles.Count > 0)
+			{
+				// Create an unoptmized compilation environment. Need to turn of PCH due to different
+				// compiler settings
+				CPPEnvironment UnoptimziedEnvironment = CompileEnvironment.DeepCopy();
+				UnoptimziedEnvironment.Config.OptimizeCode = ModuleRules.CodeOptimization.Never;
+				UnoptimziedEnvironment.Config.PrecompiledHeaderAction = PrecompiledHeaderAction.None;
+
+				// Compile the files
+				CPPOutput AdaptiveOutput = ToolChain.CompileCPPFiles(Target, UnoptimziedEnvironment, AdaptiveFiles, Name);
+
+				// Merge output
+				OutputFiles.ObjectFiles.AddRange(AdaptiveOutput.ObjectFiles);
+				OutputFiles.DebugDataFiles.AddRange(AdaptiveOutput.DebugDataFiles);
+			}
+
+			return OutputFiles;
+		}
+
 		private PrecompileHeaderEnvironment ApplySharedPCH(CPPEnvironment GlobalCompileEnvironment, CPPEnvironment CompileEnvironment, CPPEnvironment ModuleCompileEnvironment, List<FileItem> CPPFiles, ref FileItem SharedPCHHeaderFile)
 		{
 			// Check to see if we have a PCH header already setup that we can use
@@ -1470,7 +1536,7 @@ namespace UnrealBuildTool
 				bool bFoundAProblemWithPCHs = false;
 
 				FileItem UniquePCH = null;
-				foreach (FileItem CPPFile in SourceFilesFound.CPPFiles)	// @todo ubtmake: We're not caching CPPEnvironments for .c/.mm files, etc.  Even though they don't use PCHs, they still have #includes!  This can break dependency checking!
+				foreach (FileItem CPPFile in SourceFilesFound.CPPFiles) // @todo ubtmake: We're not caching CPPEnvironments for .c/.mm files, etc.  Even though they don't use PCHs, they still have #includes!  This can break dependency checking!
 				{
 					// Build a single list of include paths to search.
 					List<string> IncludePathsToSearch = ModuleCompileEnvironment.Config.CPPIncludeInfo.GetIncludesPathsToSearch(CPPFile);
@@ -1490,7 +1556,7 @@ namespace UnrealBuildTool
 					{
 						UniquePCH = PCH;
 					}
-					else if (!UniquePCH.Info.Name.Equals(PCH.Info.Name, StringComparison.InvariantCultureIgnoreCase))		// @todo ubtmake: We do a string compare on the file name (not path) here, because sometimes the include resolver will pick an Intermediate copy of a PCH header file and throw off our comparisons
+					else if (!UniquePCH.Info.Name.Equals(PCH.Info.Name, StringComparison.InvariantCultureIgnoreCase))       // @todo ubtmake: We do a string compare on the file name (not path) here, because sometimes the include resolver will pick an Intermediate copy of a PCH header file and throw off our comparisons
 					{
 						// OK, looks like we have multiple source files including a different header file first.  We'll keep track of this and print out
 						// helpful information afterwards.
@@ -1621,12 +1687,12 @@ namespace UnrealBuildTool
 		{
 			CPPEnvironment Result = BaseCompileEnvironment.DeepCopy();
 
-            if (Binary == null)
-            {
-                // Adding this check here as otherwise the call to Binary.Config.IntermediateDirectory will give an 
-                // unhandled exception
-                throw new BuildException("UEBuildBinary not set up for module {0}", this.ToString());
-            }
+			if (Binary == null)
+			{
+				// Adding this check here as otherwise the call to Binary.Config.IntermediateDirectory will give an 
+				// unhandled exception
+				throw new BuildException("UEBuildBinary not set up for module {0}", this.ToString());
+			}
 
 			// Override compile environment
 			Result.Config.bFasterWithoutUnity = Rules.bFasterWithoutUnity;
diff --git a/Engine/Source/Programs/UnrealBuildTool/System/Unity.cs b/Engine/Source/Programs/UnrealBuildTool/System/Unity.cs
--- a/Engine/Source/Programs/UnrealBuildTool/System/Unity.cs
+++ b/Engine/Source/Programs/UnrealBuildTool/System/Unity.cs
@@ -22,6 +22,11 @@ namespace UnrealBuildTool
 		/// </summary>
 		public static HashSet<FileItem> CandidateSourceFilesForWorkingSet = new HashSet<FileItem>();
 
+		/// <summary>
+		/// Prefix used for all dynamically created Unity modules
+		/// </summary>
+		public static string ModulePrefix = "Module.";
+
 		/// <summary>
 		/// A class which represents a list of files and the sum of their lengths.
 		/// </summary>
@@ -342,11 +347,11 @@ namespace UnrealBuildTool
 				string UnityCPPFileName;
 				if (AllUnityFiles.Count > 1)
 				{
-					UnityCPPFileName = string.Format("Module.{0}.{1}_of_{2}.cpp", BaseName, CurrentUnityFileCount, AllUnityFiles.Count);
+					UnityCPPFileName = string.Format("{0}{1}.{2}_of_{3}.cpp", ModulePrefix, BaseName, CurrentUnityFileCount, AllUnityFiles.Count);
 				}
 				else
 				{
-					UnityCPPFileName = string.Format("Module.{0}.cpp", BaseName);
+					UnityCPPFileName = string.Format("{0}{1}.cpp", ModulePrefix, BaseName);
 				}
 				FileReference UnityCPPFilePath = FileReference.Combine(CompileEnvironment.Config.OutputDirectory, UnityCPPFileName);
 
diff --git a/Engine/Source/Programs/UnrealPak/Private/UnrealPak.cpp b/Engine/Source/Programs/UnrealPak/Private/UnrealPak.cpp
--- a/Engine/Source/Programs/UnrealPak/Private/UnrealPak.cpp
+++ b/Engine/Source/Programs/UnrealPak/Private/UnrealPak.cpp
@@ -1317,7 +1317,21 @@ bool DiffFilesInPaks(const FString InPakFilename1, const FString InPakFilename2)
 	return true;
 }
 
-bool GenerateHashForFile( FString Filename, uint8 FileHash[16])
+struct FFileInfo
+{
+	uint64 FileSize;
+	uint8 Hash[16];
+};
+
+void GenerateHashForFile(uint8* ByteBuffer, uint64 TotalSize, FFileInfo& FileHash)
+{
+	FMD5 FileHasher;
+	FileHasher.Update(ByteBuffer, TotalSize);
+	FileHasher.Final(FileHash.Hash);
+	FileHash.FileSize = TotalSize;
+}
+
+bool GenerateHashForFile( FString Filename, FFileInfo& FileHash)
 {
 	FArchive* File = IFileManager::Get().CreateFileReader(*Filename);
 
@@ -1333,48 +1347,145 @@ bool GenerateHashForFile( FString Filename, uint8 FileHash[16])
 	delete File;
 	File = NULL;
 
-	FMD5 FileHasher;
-	FileHasher.Update(ByteBuffer, TotalSize);
-
+	GenerateHashForFile(ByteBuffer, TotalSize, FileHash);
+	
 	delete[] ByteBuffer;
+	return true;
+}
 
-	FileHasher.Final(FileHash);
+bool GenerateHashesFromPak(const TCHAR* InPakFilename, TMap<FString, FFileInfo>& FileHashes, bool bUseMountPoint = false)
+{
+	FPakFile PakFile(InPakFilename, FParse::Param(FCommandLine::Get(), TEXT("signed")));
+	if (PakFile.IsValid())
+	{
+		FArchive& PakReader = *PakFile.GetSharedReader(NULL);
+		const int64 BufferSize = 8 * 1024 * 1024; // 8MB buffer for extracting
+		void* Buffer = FMemory::Malloc(BufferSize);
+		int64 CompressionBufferSize = 0;
+		uint8* PersistantCompressionBuffer = NULL;
+		int32 ErrorCount = 0;
+		int32 FileCount = 0;
 
-	return true;
+		FString PakMountPoint = bUseMountPoint ? PakFile.GetMountPoint().Replace(TEXT("../../../"), TEXT("")) : TEXT("");
 
+		for (FPakFile::FFileIterator It(PakFile); It; ++It, ++FileCount)
+		{
+			const FPakEntry& Entry = It.Info();
+			PakReader.Seek(Entry.Offset);
+			uint32 SerializedCrcTest = 0;
+			FPakEntry EntryInfo;
+			EntryInfo.Serialize(PakReader, PakFile.GetInfo().Version);
+			if (EntryInfo == Entry)
+			{
+				// TAutoPtr<FArchive> FileHandle(IFileManager::Get().CreateFileWriter(*DestFilename));
+				TArray<uint8> Bytes;
+				FMemoryWriter MemoryFile(Bytes);
+				FArchive* FileHandle = &MemoryFile;
+				// if (FileHandle.IsValid())
+				{
+					if (Entry.CompressionMethod == COMPRESS_None)
+					{
+						BufferedCopyFile(*FileHandle, PakReader, Entry, Buffer, BufferSize);
+					}
+					else
+					{
+						UncompressCopyFile(*FileHandle, PakReader, Entry, PersistantCompressionBuffer, CompressionBufferSize);
+					}
 
-	// uint8 DestFileHash[20];
+					UE_LOG(LogPakFile, Display, TEXT("Generated hash for \"%s\""), *It.Filename());
+					FFileInfo FileHash;
+					GenerateHashForFile(Bytes.GetData(), Bytes.Num(), FileHash);
 
-}
+					FileHashes.Add(It.Filename(), FileHash);
+				}
+				/*else
+				{
+					UE_LOG(LogPakFile, Error, TEXT("Unable to create file \"%s\"."), *DestFilename);
+					ErrorCount++;
+				}*/
+
+			}
+			else
+			{
+				UE_LOG(LogPakFile, Error, TEXT("Serialized hash mismatch for \"%s\"."), *It.Filename());
+				ErrorCount++;
+			}
+		}
+		FMemory::Free(Buffer);
+		FMemory::Free(PersistantCompressionBuffer);
 
+		UE_LOG(LogPakFile, Log, TEXT("Finished extracting %d files (including %d errors)."), FileCount, ErrorCount);
+
+		return true;
+	}
+	else
+	{
+		UE_LOG(LogPakFile, Error, TEXT("Unable to open pak file \"%s\"."), InPakFilename);
+		return false;
+	}
+}
 
-void RemoveIdenticalFiles( TArray<FPakInputPair>& FilesToPak, const FString& SourceDirectory )
+void RemoveIdenticalFiles( TArray<FPakInputPair>& FilesToPak, const FString& SourceDirectory, const TMap<FString, FFileInfo>& FileHashes )
 {
+	FString HashFilename = SourceDirectory / TEXT("Hashes.txt");
+
+	if (IFileManager::Get().FileExists(*HashFilename) )
+	{
+		FString EntireFile;
+		FFileHelper::LoadFileToString(EntireFile, *HashFilename);
+
+	}
+
 	for ( int I = FilesToPak.Num()-1; I >= 0; --I )
 	{
 		const auto& NewFile = FilesToPak[I]; 
 
+		FString SourceFileNoMountPoint =  NewFile.Dest.Replace(TEXT("../../../"), TEXT(""));
 		FString SourceFilename = SourceDirectory / NewFile.Dest.Replace(TEXT("../../../"), TEXT(""));
-		int64 SourceTotalSize = IFileManager::Get().FileSize(*SourceFilename);
+		
+		const FFileInfo* FoundFileHash = FileHashes.Find(SourceFileNoMountPoint);
+		if (!FoundFileHash)
+		{
+			FoundFileHash = FileHashes.Find(NewFile.Dest);
+		}
+		
+		if ( !FoundFileHash )
+		{
+			UE_LOG(LogPakFile, Display, TEXT("Didn't find hash for %s No mount %s"), *SourceFilename, *SourceFileNoMountPoint);
+		}
 
+		int64 SourceTotalSize = FoundFileHash ? FoundFileHash->FileSize : IFileManager::Get().FileSize(*SourceFilename);
+		
 		FString DestFilename = NewFile.Source;
 		int64 DestTotalSize = IFileManager::Get().FileSize(*DestFilename);
 		
-		if ( SourceTotalSize != DestTotalSize )
+		if (SourceTotalSize != DestTotalSize)
 		{
 			// file size doesn't match 
-			UE_LOG(LogPakFile, Display, TEXT("Source file size for %s %d bytes doesn't match %s %d bytes"), *SourceFilename, SourceTotalSize, *DestFilename, DestTotalSize);
+			UE_LOG(LogPakFile, Display, TEXT("Source file size for %s %d bytes doesn't match %s %d bytes, did find %d"), *SourceFilename, SourceTotalSize, *DestFilename, DestTotalSize, FoundFileHash ? 1 : 0);
 			continue;
 		}
 
-		uint8 SourceFileHash[16];
-		if ( GenerateHashForFile(SourceFilename, SourceFileHash) == false )
+		FFileInfo SourceFileHash;
+		if ( !FoundFileHash) 
 		{
-			// file size doesn't match 
-			UE_LOG(LogPakFile, Display, TEXT("Source file size %s doesn't exist will be included in build"), *SourceFilename);
-			continue;
+			if (GenerateHashForFile(SourceFilename, SourceFileHash) == false)
+			{
+				// file size doesn't match 
+				UE_LOG(LogPakFile, Display, TEXT("Source file size %s doesn't exist will be included in build"), *SourceFilename);
+				continue;
+			}
+			else
+			{
+				UE_LOG(LogPakFile, Warning, TEXT("Generated hash for file %s but it should have been in the FileHashes array %d"), *SourceFilename, FileHashes.Num());
+			}
 		}
-		uint8 DestFileHash[16];
+		else
+		{
+			SourceFileHash = *FoundFileHash;
+		}
+		
+		FFileInfo DestFileHash;
 		if ( GenerateHashForFile( DestFilename, DestFileHash ) == false )
 		{
 			// destination file was removed don't really care about it
@@ -1382,7 +1493,7 @@ void RemoveIdenticalFiles( TArray<FPakInputPair>& FilesToPak, const FString& Sou
 			continue;
 		}
 
-		int32 Diff = FMemory::Memcmp( SourceFileHash, DestFileHash, sizeof( DestFileHash ) );
+		int32 Diff = FMemory::Memcmp( &SourceFileHash, &DestFileHash, sizeof( DestFileHash ) );
 		if ( Diff != 0 )
 		{
 			UE_LOG(LogPakFile, Display, TEXT("Source file hash for %s doesn't match dest file hash %s and will be included in patch"), *SourceFilename, *DestFilename);
@@ -1531,6 +1642,8 @@ INT32_MAIN_INT32_ARGC_TCHAR_ARGV()
 			}
 			else
 			{
+				TMap<FString, FFileInfo> SourceFileHashes;
+
 				if ( CmdLineParameters.GeneratePatch )
 				{
 					FString OutputPath;
@@ -1543,11 +1656,14 @@ INT32_MAIN_INT32_ARGC_TCHAR_ARGV()
 
 					UE_LOG(LogPakFile, Display, TEXT("Generating patch from %s."), *CmdLineParameters.SourcePatchPakFilename );
 
-					if ( ExtractFilesFromPak( *CmdLineParameters.SourcePatchPakFilename, *OutputPath, true ) == false )
+					if ( !GenerateHashesFromPak(*CmdLineParameters.SourcePatchPakFilename, SourceFileHashes) )
 					{
-						UE_LOG(LogPakFile, Error, TEXT("Unable to extract files from source pak file for patch") );
+						if ( ExtractFilesFromPak( *CmdLineParameters.SourcePatchPakFilename, *OutputPath, true ) == false )
+						{
+							UE_LOG(LogPakFile, Error, TEXT("Unable to extract files from source pak file for patch") );
+						}
+						CmdLineParameters.SourcePatchDiffDirectory = OutputPath;
 					}
-					CmdLineParameters.SourcePatchDiffDirectory = OutputPath;
 				}
 
 
@@ -1558,7 +1674,7 @@ INT32_MAIN_INT32_ARGC_TCHAR_ARGV()
 				if ( CmdLineParameters.GeneratePatch )
 				{
 					// if we are generating a patch here we remove files which are already shipped...
-					RemoveIdenticalFiles(FilesToAdd, CmdLineParameters.SourcePatchDiffDirectory);
+					RemoveIdenticalFiles(FilesToAdd, CmdLineParameters.SourcePatchDiffDirectory, SourceFileHashes);
 				}
 
 
diff --git a/Engine/Source/Runtime/AIModule/Classes/AISystem.h b/Engine/Source/Runtime/AIModule/Classes/AISystem.h
--- a/Engine/Source/Runtime/AIModule/Classes/AISystem.h
+++ b/Engine/Source/Runtime/AIModule/Classes/AISystem.h
@@ -6,6 +6,7 @@
 #include "Engine/EngineTypes.h"
 #include "Engine/World.h"
 #include "AI/AISystemBase.h"
+#include "Math/RandomStream.h"
 #include "AISystem.generated.h"
 
 class UBehaviorTreeManager;
@@ -15,6 +16,7 @@ class UAIAsyncTaskBlueprintProxy;
 class UAIHotSpotManager;
 class UBlackboardData;
 class UBlackboardComponent;
+class UNavLocalGridManager;
 
 #define GET_AI_CONFIG_VAR(a) (GetDefault<UAISystem>()->a)
 
@@ -90,12 +92,18 @@ protected:
 	UPROPERTY(Transient)
 	UAIHotSpotManager* HotSpotManager;
 
+	UPROPERTY(Transient)
+	UNavLocalGridManager* NavLocalGrids;
+
 	typedef TMultiMap<TWeakObjectPtr<UBlackboardData>, TWeakObjectPtr<UBlackboardComponent> > FBlackboardDataToComponentsMap;
 
 	/** UBlackboardComponent instances that reference the blackboard data definition */
 	FBlackboardDataToComponentsMap BlackboardDataToComponentsMap;
 
 	FDelegateHandle ActorSpawnedDelegateHandle;
+
+	/** random number stream to be used by all things AI. WIP */
+	static FRandomStream RandomStream;
 	
 public:
 	UAISystem(const FObjectInitializer& ObjectInitializer);
@@ -127,6 +135,9 @@ public:
 	FORCEINLINE UAIHotSpotManager* GetHotSpotManager() { return HotSpotManager; }
 	FORCEINLINE const UAIHotSpotManager* GetHotSpotManager() const { return HotSpotManager; }
 
+	FORCEINLINE UNavLocalGridManager* GetNavLocalGridManager() { return NavLocalGrids; }
+	FORCEINLINE const UNavLocalGridManager* GetNavLocalGridManager() const { return NavLocalGrids; }
+
 	FORCEINLINE static UAISystem* GetCurrentSafe(UWorld* World) 
 	{ 
 		return World != nullptr ? Cast<UAISystem>(World->GetAISystem()) : NULL;
@@ -230,6 +241,9 @@ public:
 
 	virtual void ConditionalLoadDebuggerPlugin();
 
+	static const FRandomStream& GetRandomStream() { return RandomStream; }
+	static void SeedRandomStream(const int32 Seed) { return RandomStream.Initialize(Seed); }
+
 protected:
 	virtual void OnActorSpawned(AActor* SpawnedActor);
 	void LoadDebuggerPlugin();
diff --git a/Engine/Source/Runtime/AIModule/Classes/AITypes.h b/Engine/Source/Runtime/AIModule/Classes/AITypes.h
--- a/Engine/Source/Runtime/AIModule/Classes/AITypes.h
+++ b/Engine/Source/Runtime/AIModule/Classes/AITypes.h
@@ -487,6 +487,9 @@ struct AIMODULE_API FAIMoveRequest
 	bool IsMoveToActorRequest() const { return bMoveToActor; }
 	AActor* GetGoalActor() const { return bMoveToActor ? GoalActor : nullptr; }
 	FVector GetGoalLocation() const { return GoalLocation; }
+	/** retrieves request's requested destination location, GoalActor's location 
+	 *	or GoalLocation, depending on the request itself */
+	FVector GetDestination() const { return bMoveToActor ? (GoalActor ? GoalActor->GetActorLocation() : FAISystem::InvalidLocation) : GoalLocation; }
 
 	bool IsUsingPathfinding() const { return bUsePathfinding; }
 	bool IsUsingPartialPaths() const { return bAllowPartialPath; }
diff --git a/Engine/Source/Runtime/AIModule/Classes/EnvironmentQuery/EnvQueryTest.h b/Engine/Source/Runtime/AIModule/Classes/EnvironmentQuery/EnvQueryTest.h
--- a/Engine/Source/Runtime/AIModule/Classes/EnvironmentQuery/EnvQueryTest.h
+++ b/Engine/Source/Runtime/AIModule/Classes/EnvironmentQuery/EnvQueryTest.h
@@ -12,6 +12,13 @@ class UEnvQueryContext;
 struct FPropertyChangedEvent;
 #endif // WITH_EDITOR
 
+UENUM()
+enum class EEQSNormalizationType : uint8
+{
+	Absolute,	// using 0 as the base of normalization range
+	RelativeToScores	// using lowest item store as the base of normalization range
+};
+
 namespace EnvQueryTestVersion
 {
 	const int32 Initial = 0;
@@ -100,6 +107,10 @@ class AIMODULE_API UEnvQueryTest : public UEnvQueryNode
 	UPROPERTY(EditDefaultsOnly, Category=Score)
 	TEnumAsByte<EEnvQueryTestClamping::Type> ClampMaxType;
 
+	/** Specifies how to determine value span used to normalize scores */
+	UPROPERTY(EditDefaultsOnly, Category = Score)
+	EEQSNormalizationType NormalizationType;
+
 	/** Minimum value to use to normalize the raw test value before applying scoring formula. */
 	UPROPERTY(EditDefaultsOnly, Category=Score)
 	FAIDataProviderFloatValue ScoreClampMin;
diff --git a/Engine/Source/Runtime/AIModule/Classes/EnvironmentQuery/EnvQueryTypes.h b/Engine/Source/Runtime/AIModule/Classes/EnvironmentQuery/EnvQueryTypes.h
--- a/Engine/Source/Runtime/AIModule/Classes/EnvironmentQuery/EnvQueryTypes.h
+++ b/Engine/Source/Runtime/AIModule/Classes/EnvironmentQuery/EnvQueryTypes.h
@@ -642,6 +642,7 @@ struct FEQSQueryDebugData
 	}
 };
 
+
 UCLASS(Abstract)
 class AIMODULE_API UEnvQueryTypes : public UObject
 {
@@ -1004,9 +1005,12 @@ public:
 					break;
 			}
 
-			if (bPassedTest)
+			if (bPassedTest || !bIsFiltering)
 			{
-				SetScoreInternal(1.0f);
+				// even if the item's result is different than expected
+				// but we're not filtering those items out, we still want
+				// to treat this as successful test, just with different score
+				SetScoreInternal(bPassedTest ? 1.0f : 0.f);
 				NumPassedForItem++;
 			}
 
@@ -1021,7 +1025,7 @@ public:
 		/** Force state and score of item
 		 *  Any following SetScore calls for current item will be ignored
 		 */
-		void ForceItemState(EEnvItemStatus::Type InStatus, float Score = UEnvQueryTypes::SkippedItemValue)
+		void ForceItemState(const EEnvItemStatus::Type InStatus, const float Score = UEnvQueryTypes::SkippedItemValue)
 		{
 			bForced = true;
 			bPassed = (InStatus == EEnvItemStatus::Passed);
@@ -1114,7 +1118,11 @@ public:
 		{
 			if (!bForced)
 			{
-				if (!bIsFiltering)
+				if (NumTestsForItem == 0)
+				{
+					bPassed = false;
+				}
+				else if (!bIsFiltering)
 				{
 					bPassed = true;
 				}
@@ -1175,7 +1183,7 @@ struct AIMODULE_API FAIDynamicParam
 };
 
 USTRUCT()
-struct FEQSParametrizedQueryExecutionRequest
+struct AIMODULE_API FEQSParametrizedQueryExecutionRequest
 {
 	GENERATED_USTRUCT_BODY()
 
diff --git a/Engine/Source/Runtime/AIModule/Classes/Navigation/NavLocalGridData.h b/Engine/Source/Runtime/AIModule/Classes/Navigation/NavLocalGridData.h
new file mode 100644
--- /dev/null
+++ b/Engine/Source/Runtime/AIModule/Classes/Navigation/NavLocalGridData.h
@@ -0,0 +1,108 @@
+// Copyright 1998-2016 Epic Games, Inc. All Rights Reserved.
+
+#pragma once
+#include "SimpleCellGrid.h"
+
+/**
+ *  Local navigation grid - simple 2D grid used for navigation.
+ * 
+ *  Cell can be either free or marked as obstacle, connected with 8 neighbors (no walls in between)
+ *
+ *  When used as source for UNavLocalGridManager, each obstacle should define its own grid data.
+ *  Corresponding category in gameplay debugger is hidden by default, please adjust project configs to change that.  
+ *
+ *  TODO: helpers for marking different basic shapes
+ *  TODO: serialization (with maps?)
+ *  TODO: FNavigationPath support?
+ */
+struct FNavLocalGridData : public TSimpleCellGrid<uint8>
+{
+	FNavLocalGridData() : GridId(0) {}
+	FNavLocalGridData(const FVector& Center, float Extent2D);
+	FNavLocalGridData(const FVector& Center, const FVector2D& Extent2D);
+	FNavLocalGridData(const TArray<FNavLocalGridData>& SourceGrids);
+
+	/** mark single cell as obstacle */
+	void MarkPointObstacle(const FVector& Center);
+
+	/** mark box (AABB or rotated) shape as obstacle */
+	void MarkBoxObstacle(const FVector& Center, const FVector& Extent, const FQuat& Quat = FQuat::Identity);
+
+	/** set height of bounds, if not set: ProjectCells will use height of default query box */
+	void SetHeight(float ExtentZ);
+
+	/** get unique Id of grid */
+	const int32 GetGridId() const
+	{
+		return GridId; 
+	}
+
+	/** check if there's an obstacle at cell coords */
+	bool HasObstacleUnsafe(int32 LocationX, int32 LocationY) const
+	{
+		return GetCellAtIndexUnsafe(GetCellIndexUnsafe(LocationX, LocationY)) > 0;
+	}
+
+	/** convert cell index to global world coords with origin in (0,0,0) */
+	FIntVector GetGlobalCoords(int32 CellIdx) const
+	{
+		return FIntVector(GetCellCoordX(CellIdx) + OriginWorldCoord.X, GetCellCoordY(CellIdx) + OriginWorldCoord.Y, OriginWorldCoord.Z); 
+	}
+
+	/** convert global world coords to cell index, return -1 if outside */
+	int32 GetCellIndexFromGlobalCoords2D(const FIntVector& WorldCoords) const;
+
+	/** convert cell index to world location using projected heights */
+	FVector GetProjectedCellCenter(int32 CellIdx) const 
+	{
+		return GetProjectedCellCenter(GetCellCoordX(CellIdx), GetCellCoordY(CellIdx));
+	}
+
+	/** convert cell coords on grid to world location using projected heights */
+	FVector GetProjectedCellCenter(int32 LocationX, int32 LocationY) const
+	{
+		const FVector WorldCoords = GetWorldCellCenter(LocationX, LocationY);
+		return FVector(WorldCoords.X, WorldCoords.Y, CellZ[GetCellIndexUnsafe(LocationX, LocationY)]);
+	}
+
+	/** creates path points from navigation path going through grid
+	 *  @param SourcePath - [in] full navigation path
+	 *  @param EntryLocation - [in] location of agent
+	 *  @param EntrySegmentStart - [in] current move segment on path
+	 *  @param PathPointsInside - [out] points inside grid
+	 *  @param NextSegmentStart - [out] next move segment on path after leaving grid or -1 if path ends inside
+	 */
+	void FindPathForMovingAgent(const FNavigationPath& SourcePath, const FVector& EntryLocation, int32 EntrySegmentStart, TArray<FVector>& PathPointsInside, int32& NextSegmentStart) const;
+
+	/** create path points from StartCoords to EndCoord, returns false when failed */
+	bool FindPath(const FIntVector& StartCoords, const FIntVector& EndCoords, TArray<FIntVector>& PathCoords) const;
+
+	/** project cells on navigation data and marks failed ones as obstacles */
+	void ProjectCells(const ANavigationData& NavData);
+
+	//////////////////////////////////////////////////////////////////////////
+	// FGraphAStar: TGraph
+	typedef int32 FNodeRef;
+
+	int32 GetNeighbourCount(FNodeRef NodeRef) const { return 8; }
+	bool IsValidRef(FNodeRef NodeRef) const { return IsValidIndex(NodeRef); }
+	FNodeRef GetNeighbour(const FNodeRef NodeRef, const int32 NeiIndex) const;
+	//////////////////////////////////////////////////////////////////////////
+
+protected:
+	friend class UNavLocalGridManager;
+	TArray<float> CellZ;
+
+	/** convert PathIndices into pruned PathCoords */
+	void PostProcessPath(const FIntVector& StartCoords, const FIntVector& EndCoords, const TArray<int32>& PathIndices, TArray<FIntVector>& PathCoords) const;
+
+	/** check if line trace between local coords on grid hits any obstacles, doesn't validate coords! */
+	bool IsLineObstructed(const FIntVector& StartCoords, const FIntVector& EndCoords) const;
+
+	/** set unique Id of grid */
+	void SetGridId(int32 NewId);
+
+private:
+	int32 GridId;
+	FIntVector OriginWorldCoord;
+};
diff --git a/Engine/Source/Runtime/AIModule/Classes/Navigation/NavLocalGridManager.h b/Engine/Source/Runtime/AIModule/Classes/Navigation/NavLocalGridManager.h
new file mode 100644
--- /dev/null
+++ b/Engine/Source/Runtime/AIModule/Classes/Navigation/NavLocalGridManager.h
@@ -0,0 +1,100 @@
+// Copyright 1998-2016 Epic Games, Inc. All Rights Reserved.
+
+#pragma once
+
+#include "NavLocalGridData.h"
+#include "NavLocalGridManager.generated.h"
+
+struct FCombinedNavGridData : public FNavLocalGridData
+{
+	TArray<int32> SourceIds;
+
+	FCombinedNavGridData() {}
+	FCombinedNavGridData(const TArray<FNavLocalGridData>& SourceGrids) : FNavLocalGridData(SourceGrids) {}
+};
+
+/**
+ *  Manager for local navigation grids
+ * 
+ *  Builds non overlapping grid from multiple sources, that can be used later for pathfinding.
+ *  Check also: UGridPathFollowingComponent, FNavLocalGridData
+ */
+
+UCLASS(Experimental)
+class UNavLocalGridManager : public UObject
+{
+	GENERATED_UCLASS_BODY()
+public:
+
+	/** adds new grid */
+	int32 AddGridData(const FNavLocalGridData& GridData, bool bUpdate = true);
+
+	/** removes grid */
+	void RemoveGridData(int32 GridId, bool bUpdate = true);
+
+	/** rebuild overlapping grids if needed */
+	void RebuildGrids();
+
+	/** tries to find a path using grids, returns false when start and end locations are not on the same grid */
+	bool FindPath(const FVector& Start, const FVector& End, TArray<FVector>& PathPoints) const;
+
+	/** set shared size of grid cell, allowed only when there's no grid added */
+	bool SetCellSize(float CellSize);
+
+	/** get shared size of grid cell, static but there's only one active manager */
+	static float GetCellSize() { return GridCellSize; }
+
+	/** get number of known source grids */
+	int32 GetNumSources() const { return SourceGrids.Num(); }
+
+	/** get number of combined, non overlapping grids */
+	int32 GetNumGrids() const { return CombinedGrids.Num(); }
+
+	/** get source grid by index */
+	const FNavLocalGridData& GetSourceData(int32 SourceIdx) const { return SourceGrids[SourceIdx]; }
+
+	/** get combined, non overlapping grid by index */
+	const FNavLocalGridData& GetGridData(int32 GridIdx) const { return CombinedGrids[GridIdx]; }
+
+	/** find combined grid at location */
+	int32 GetGridIndex(const FVector& WorldLocation) const;
+
+	/** get version of grid data, incremented with each rebuild */
+	int32 GetVersion() const { return VersionNum; }
+
+	UFUNCTION(BlueprintCallable, Category = "AI|Navigation", meta = (WorldContext = "WorldContext"))
+	static bool SetLocalNavigationGridDensity(UObject* WorldContext, float CellSize);
+
+	UFUNCTION(BlueprintCallable, Category = "AI|Navigation", meta = (WorldContext = "WorldContext"))
+	static int32 AddLocalNavigationGridForPoint(UObject* WorldContext, const FVector& Location, const int32 Radius2D = 5, const float Height = 100.0f, bool bRebuildGrids = true);
+
+	UFUNCTION(BlueprintCallable, Category = "AI|Navigation", meta = (WorldContext = "WorldContext"))
+	static int32 AddLocalNavigationGridForBox(UObject* WorldContext, const FVector& Location, FVector Extent = FVector(1,1,1), FRotator Rotation = FRotator::ZeroRotator, const int32 Radius2D = 5, const float Height = 100.0f, bool bRebuildGrids = true);
+
+	UFUNCTION(BlueprintCallable, Category = "AI|Navigation", meta = (WorldContext = "WorldContext"))
+	static void RemoveLocalNavigationGrid(UObject* WorldContext, int32 GridId, bool bRebuildGrids = true);
+
+	UFUNCTION(BlueprintCallable, Category = "AI|Navigation", meta = (WorldContext = "WorldContext"))
+	static bool FindLocalNavigationGridPath(UObject* WorldContext, const FVector& Start, const FVector& End, TArray<FVector>& PathPoints);
+
+	static UNavLocalGridManager* GetCurrent(UWorld* World);
+	static UNavLocalGridManager* GetCurrent(const UObject* WorldContextObject);
+
+#if WITH_ENGINE
+	virtual UWorld* GetWorld() const override;
+#endif
+
+protected:
+
+	TArray<FCombinedNavGridData> CombinedGrids;
+	TArray<FNavLocalGridData> SourceGrids;
+	
+	static float GridCellSize;
+
+	int32 VersionNum;
+	int32 NextGridId;
+	uint32 bNeedsRebuilds : 1;
+
+	/** projects combined grids to navigation data */
+	virtual void ProjectGrids(const TArray<int32>& GridIndices);
+};
diff --git a/Engine/Source/Runtime/AIModule/Classes/Perception/AIPerceptionComponent.h b/Engine/Source/Runtime/AIModule/Classes/Perception/AIPerceptionComponent.h
--- a/Engine/Source/Runtime/AIModule/Classes/Perception/AIPerceptionComponent.h
+++ b/Engine/Source/Runtime/AIModule/Classes/Perception/AIPerceptionComponent.h
@@ -39,18 +39,31 @@ struct AIMODULE_API FActorPerceptionInfo
 		}
 	}
 
-	FORCEINLINE_DEBUGGABLE FVector GetLastStimulusLocation(float* OptionalAge = NULL) const 
+	/** Retrieves last known location. Active (last reported as "successful")
+	 *	stimuli are preferred. */
+	FVector GetLastStimulusLocation(float* OptionalAge = NULL) const 
 	{
 		FVector Location(FAISystem::InvalidLocation);
 		float BestAge = FLT_MAX;
+		bool bBestWasSuccessfullySensed = false;
 		for (int32 Sense = 0; Sense < LastSensedStimuli.Num(); ++Sense)
 		{
 			const float Age = LastSensedStimuli[Sense].GetAge();
+			const bool bWasSuccessfullySensed = LastSensedStimuli[Sense].WasSuccessfullySensed();
+
 			if (Age >= 0 && (Age < BestAge 
-				|| (Sense == DominantSense && LastSensedStimuli[Sense].WasSuccessfullySensed())))
+				|| (bBestWasSuccessfullySensed == false && bWasSuccessfullySensed)
+				|| (Sense == DominantSense && bWasSuccessfullySensed)))
 			{
 				BestAge = Age;
 				Location = LastSensedStimuli[Sense].StimulusLocation;
+				bBestWasSuccessfullySensed = bWasSuccessfullySensed;
+
+				if (Sense == DominantSense && bWasSuccessfullySensed)
+				{
+					// if dominant sense is active we don't want to look any further 
+					break;
+				}
 			}
 		}
 
@@ -249,6 +262,7 @@ public:
 
 	float GetYoungestStimulusAge(const AActor& Source) const;
 	bool HasAnyActiveStimulus(const AActor& Source) const;
+	bool HasAnyCurrentStimulus(const AActor& Source) const;
 	bool HasActiveStimulus(const AActor& Source, FAISenseID Sense) const;
 
 #if WITH_GAMEPLAY_DEBUGGER
diff --git a/Engine/Source/Runtime/AIModule/Private/AIModule.cpp b/Engine/Source/Runtime/AIModule/Private/AIModule.cpp
--- a/Engine/Source/Runtime/AIModule/Private/AIModule.cpp
+++ b/Engine/Source/Runtime/AIModule/Private/AIModule.cpp
@@ -19,6 +19,7 @@
 #include "GameplayDebugger/GameplayDebuggerCategory_EQS.h"
 #include "GameplayDebugger/GameplayDebuggerCategory_Navmesh.h"
 #include "GameplayDebugger/GameplayDebuggerCategory_Perception.h"
+#include "GameplayDebugger/GameplayDebuggerCategory_NavLocalGrid.h"
 #endif // WITH_GAMEPLAY_DEBUGGER
 
 #define LOCTEXT_NAMESPACE "AIModule"
@@ -71,6 +72,7 @@ void FAIModule::StartupModule()
 	GameplayDebuggerModule.RegisterCategory("EQS", IGameplayDebugger::FOnGetCategory::CreateStatic(&FGameplayDebuggerCategory_EQS::MakeInstance));
 	GameplayDebuggerModule.RegisterCategory("Navmesh", IGameplayDebugger::FOnGetCategory::CreateStatic(&FGameplayDebuggerCategory_Navmesh::MakeInstance), EGameplayDebuggerCategoryState::Disabled, 0);
 	GameplayDebuggerModule.RegisterCategory("Perception", IGameplayDebugger::FOnGetCategory::CreateStatic(&FGameplayDebuggerCategory_Perception::MakeInstance));
+	GameplayDebuggerModule.RegisterCategory("NavGrid", IGameplayDebugger::FOnGetCategory::CreateStatic(&FGameplayDebuggerCategory_NavLocalGrid::MakeInstance), EGameplayDebuggerCategoryState::Hidden);
 	GameplayDebuggerModule.NotifyCategoriesChanged();
 #endif
 }
@@ -92,6 +94,7 @@ void FAIModule::ShutdownModule()
 		GameplayDebuggerModule.UnregisterCategory("EQS");
 		GameplayDebuggerModule.UnregisterCategory("Navmesh");
 		GameplayDebuggerModule.UnregisterCategory("Perception");
+		GameplayDebuggerModule.UnregisterCategory("NavGrid");
 		GameplayDebuggerModule.NotifyCategoriesChanged();
 	}
 #endif
diff --git a/Engine/Source/Runtime/AIModule/Private/AISystem.cpp b/Engine/Source/Runtime/AIModule/Private/AISystem.cpp
--- a/Engine/Source/Runtime/AIModule/Private/AISystem.cpp
+++ b/Engine/Source/Runtime/AIModule/Private/AISystem.cpp
@@ -10,10 +10,13 @@
 #include "GameFramework/PlayerController.h"
 #include "HotSpots/AIHotSpotManager.h"
 #include "BehaviorTree/BlackboardData.h"
+#include "Navigation/NavLocalGridManager.h"
 #include "AISystem.h"
 
 DEFINE_STAT(STAT_AI_Overall);
 
+FRandomStream UAISystem::RandomStream;
+
 UAISystem::UAISystem(const FObjectInitializer& ObjectInitializer)
 	: Super(ObjectInitializer)
 {
@@ -30,8 +33,22 @@ UAISystem::UAISystem(const FObjectInitializer& ObjectInitializer)
 	{
 		UWorld* WorldOuter = Cast<UWorld>(GetOuter());
 		UObject* ManagersOuter = WorldOuter != NULL ? (UObject*)WorldOuter : (UObject*)this;
+		
 		BehaviorTreeManager = NewObject<UBehaviorTreeManager>(ManagersOuter);
 		EnvironmentQueryManager = NewObject<UEnvQueryManager>(ManagersOuter);
+		NavLocalGrids = NewObject<UNavLocalGridManager>(ManagersOuter);
+	}
+	else
+	{
+		// game-wise config
+		if (FParse::Param(FCommandLine::Get(), TEXT("FixedSeed")) == false)
+		{
+			// by default FRandomStream is initialized with 0
+			// only if not configured with commandline we should 
+			// initialize the random stream.
+			const int32 Seed = ((int32)(FDateTime::Now().GetTicks() % (int64)MAX_int32));
+			RandomStream.Initialize(Seed);
+		}
 	}
 }
 
diff --git a/Engine/Source/Runtime/AIModule/Private/Blueprint/AIBlueprintHelperLibrary.cpp b/Engine/Source/Runtime/AIModule/Private/Blueprint/AIBlueprintHelperLibrary.cpp
--- a/Engine/Source/Runtime/AIModule/Private/Blueprint/AIBlueprintHelperLibrary.cpp
+++ b/Engine/Source/Runtime/AIModule/Private/Blueprint/AIBlueprintHelperLibrary.cpp
@@ -8,6 +8,8 @@
 #include "TimerManager.h"
 #include "Blueprint/AIBlueprintHelperLibrary.h"
 
+DEFINE_LOG_CATEGORY_STATIC(LogAIBlueprint, Warning, All);
+
 //----------------------------------------------------------------------//
 // UAIAsyncTaskBlueprintProxy
 //----------------------------------------------------------------------//
@@ -91,7 +93,19 @@ UAIBlueprintHelperLibrary::UAIBlueprintHelperLibrary(const FObjectInitializer& O
 
 UAIAsyncTaskBlueprintProxy* UAIBlueprintHelperLibrary::CreateMoveToProxyObject(UObject* WorldContextObject, APawn* Pawn, FVector Destination, AActor* TargetActor, float AcceptanceRadius, bool bStopOnOverlap)
 {
-	check(WorldContextObject);
+	if (WorldContextObject == nullptr)
+	{
+		if (Pawn != nullptr)
+		{
+			WorldContextObject = Pawn;
+		}
+		else
+		{
+			UE_LOG(LogAIBlueprint, Warning, TEXT("Empty (None) world context as well as Pawn passed in while trying to create a MoveTo proxy"));
+			return nullptr;
+		}
+	}
+
 	if (Pawn == nullptr)
 	{
 		// maybe we can extract the pawn from the world context
@@ -106,11 +120,12 @@ UAIAsyncTaskBlueprintProxy* UAIBlueprintHelperLibrary::CreateMoveToProxyObject(U
 	{
 		return NULL;
 	}
+
 	UAIAsyncTaskBlueprintProxy* MyObj = NULL;
 	AAIController* AIController = Cast<AAIController>(Pawn->GetController());
 	if (AIController)
 	{
-		UWorld* World = GEngine->GetWorldFromContextObject( WorldContextObject );
+		UWorld* World = GEngine->GetWorldFromContextObject(WorldContextObject);
 		MyObj = NewObject<UAIAsyncTaskBlueprintProxy>(World);
 
 		FAIMoveRequest MoveReq;
diff --git a/Engine/Source/Runtime/AIModule/Private/EnvironmentQuery/EQSRenderingComponent.cpp b/Engine/Source/Runtime/AIModule/Private/EnvironmentQuery/EQSRenderingComponent.cpp
--- a/Engine/Source/Runtime/AIModule/Private/EnvironmentQuery/EQSRenderingComponent.cpp
+++ b/Engine/Source/Runtime/AIModule/Private/EnvironmentQuery/EQSRenderingComponent.cpp
@@ -135,7 +135,7 @@ void FEQSSceneProxy::CollectEQSData(const FEnvQueryResult* ResultItems, const FE
 		}
 	}
 	
-	const float ScoreNormalizer = bUseMidResults && (MaxScore != MinScore) ? 1.f / (MaxScore - MinScore) : 1.f;
+	const float ScoreNormalizer = bUseMidResults && (MaxScore != MinScore) ? (1.f / (MaxScore - MinScore)) : 1.f;
 	const float HighlightThreshold = (HighlightRangePct < 1.0f) ? (MaxScore * HighlightRangePct) : FLT_MAX;
 
 	if (bSingleItemResult == false)
@@ -144,17 +144,17 @@ void FEQSSceneProxy::CollectEQSData(const FEnvQueryResult* ResultItems, const FE
 		{
 			if (Items[ItemIndex].IsValid())
 			{
-				const float Score = bNoTestsPerformed ? 1 : Items[ItemIndex].Score * ScoreNormalizer;
+				const float NormalizedScore = bNoTestsPerformed ? 1 : (Items[ItemIndex].Score - MinScore) * ScoreNormalizer;
 				const bool bLowRadius = (HighlightThreshold < FLT_MAX) && (bNoTestsPerformed || (Items[ItemIndex].Score < HighlightThreshold));
 				const float Radius = ItemDrawRadius.X * (bLowRadius ? 0.2f : 1.0f);
 				const FVector Loc = FEQSRenderingHelper::ExtractLocation(ResultItems->ItemType, RawData, Items, ItemIndex);
 				Spheres.Add(FSphere(Radius, Loc, bNoTestsPerformed == false
-					? FLinearColor(FColor::MakeRedToGreenColorFromScalar(Score)) 
+					? FLinearColor(FColor::MakeRedToGreenColorFromScalar(NormalizedScore))
 					: FLinearColor(0.2, 1.0, 1.0, 1)));
 
 				DebugItems.Add(EQSDebug::FDebugHelper(Loc, Radius));
 
-				const FString Label = bNoTestsPerformed ? TEXT("") : FString::Printf(TEXT("%.2f"), Score);
+				const FString Label = bNoTestsPerformed ? TEXT("") : FString::Printf(TEXT("%.2f"), Items[ItemIndex].Score);
 				Texts.Add(FText3d(Label, Loc, FLinearColor::White));
 			}
 		}
@@ -163,7 +163,7 @@ void FEQSSceneProxy::CollectEQSData(const FEnvQueryResult* ResultItems, const FE
 	{
 		if (Items[0].IsValid())
 		{
-			const float Score = Items[0].Score * ScoreNormalizer;
+			const float Score = Items[0].Score;
 			const bool bLowRadius = false;
 			const float Radius = ItemDrawRadius.X * (bLowRadius ? 0.2f : 1.0f);
 			const FVector Loc = FEQSRenderingHelper::ExtractLocation(ResultItems->ItemType, RawData, Items, 0);
@@ -179,7 +179,7 @@ void FEQSSceneProxy::CollectEQSData(const FEnvQueryResult* ResultItems, const FE
 		{
 			if (Items[ItemIndex].IsValid())
 			{
-				const float Score = bNoTestsPerformed ? 1 : Items[ItemIndex].Score * ScoreNormalizer;
+				const float Score = bNoTestsPerformed ? 1 : Items[ItemIndex].Score;
 				const bool bLowRadius = (HighlightThreshold < FLT_MAX) && (bNoTestsPerformed || (Items[ItemIndex].Score < HighlightThreshold));
 				const float Radius = ItemDrawRadius.X * (bLowRadius ? 0.2f : 1.0f);
 				const FVector Loc = FEQSRenderingHelper::ExtractLocation(ResultItems->ItemType, RawData, Items, ItemIndex);
@@ -208,7 +208,7 @@ void FEQSSceneProxy::CollectEQSData(const FEnvQueryResult* ResultItems, const FE
 				continue;
 			}
 
-			const float Score = bNoTestsPerformed ? 1 : Items[ItemIndex].Score * ScoreNormalizer;
+			const float Score = bNoTestsPerformed ? 1 : Items[ItemIndex].Score;
 			const bool bLowRadius = (HighlightThreshold < FLT_MAX) && (bNoTestsPerformed || (Items[ItemIndex].Score < HighlightThreshold));
 			const float Radius = ItemDrawRadius.X * (bLowRadius ? 0.2f : 1.0f);
 			const FVector Loc = FEQSRenderingHelper::ExtractLocation(QueryInstance->ItemType, InstanceDebugData.RawData, DebugQueryItems, ItemIndex);
diff --git a/Engine/Source/Runtime/AIModule/Private/EnvironmentQuery/EnvQueryDebugHelpers.cpp b/Engine/Source/Runtime/AIModule/Private/EnvironmentQuery/EnvQueryDebugHelpers.cpp
--- a/Engine/Source/Runtime/AIModule/Private/EnvironmentQuery/EnvQueryDebugHelpers.cpp
+++ b/Engine/Source/Runtime/AIModule/Private/EnvironmentQuery/EnvQueryDebugHelpers.cpp
@@ -199,6 +199,13 @@ void UEnvQueryDebugHelpers::LogQueryInternal(FEnvQueryInstance& Query, const FLo
 		// draw test weights for best X items
 		const int32 NumItems = EQSLocalData.Items.Num();
 
+		// print sorted tests' descriptions, to be able to tie TestIdx with an actual test
+		const FEnvQueryOptionInstance& Option = Query.Options[Query.OptionIndex];
+		for (int32 TestIdx = 0; TestIdx < NumTests; TestIdx++)
+		{
+			Line.Line += FString::Printf(TEXT("%d: %s\n"), TestIdx, *Option.Tests[TestIdx]->GetDescriptionTitle().ToString());
+		}
+
 		// table header		
 		{
 			FString HeaderString;
diff --git a/Engine/Source/Runtime/AIModule/Private/EnvironmentQuery/EnvQueryInstance.cpp b/Engine/Source/Runtime/AIModule/Private/EnvironmentQuery/EnvQueryInstance.cpp
--- a/Engine/Source/Runtime/AIModule/Private/EnvironmentQuery/EnvQueryInstance.cpp
+++ b/Engine/Source/Runtime/AIModule/Private/EnvironmentQuery/EnvQueryInstance.cpp
@@ -421,6 +421,7 @@ void FEnvQueryInstance::ItemIterator::HandleFailedTestResult()
 void FEnvQueryInstance::ItemIterator::StoreTestResult()
 {
 	CheckItemPassed();
+	ensureAlways(FMath::IsNaN(ItemScore) == false);
 
 	if (Instance->IsInSingleItemFinalSearch())
 	{
@@ -449,6 +450,7 @@ void FEnvQueryInstance::ItemIterator::StoreTestResult()
 		}
 		else if (CachedScoreOp == EEnvTestScoreOperator::AverageScore && !bForced)
 		{
+			ensureAlways(NumPassedForItem != 0);
 			ItemScore /= NumPassedForItem;
 		}
 
@@ -763,4 +765,4 @@ FBox FEnvQueryInstance::GetBoundingBox() const
 	}
 
 	return BBox;
-}
\ No newline at end of file
+}
diff --git a/Engine/Source/Runtime/AIModule/Private/EnvironmentQuery/EnvQueryTest.cpp b/Engine/Source/Runtime/AIModule/Private/EnvironmentQuery/EnvQueryTest.cpp
--- a/Engine/Source/Runtime/AIModule/Private/EnvironmentQuery/EnvQueryTest.cpp
+++ b/Engine/Source/Runtime/AIModule/Private/EnvironmentQuery/EnvQueryTest.cpp
@@ -21,6 +21,7 @@ UEnvQueryTest::UEnvQueryTest(const FObjectInitializer& ObjectInitializer) : Supe
 	ClampMaxType = EEnvQueryTestClamping::None;
 	BoolValue.DefaultValue = true;
 	ScoringFactor.DefaultValue = 1.0f;
+	NormalizationType = EEQSNormalizationType::Absolute;
 
 	bWorkOnFloatValues = true;
 }
@@ -36,7 +37,7 @@ void UEnvQueryTest::NormalizeItemScores(FEnvQueryInstance& QueryInstance)
 	ScoringFactor.BindData(QueryOwner, QueryInstance.QueryID);
 	float ScoringFactorValue = ScoringFactor.GetValue();
 
-	float MinScore = 0;
+	float MinScore = (NormalizationType == EEQSNormalizationType::Absolute) ? 0 : BIG_NUMBER;
 	float MaxScore = -BIG_NUMBER;
 
 	if (ClampMinType == EEnvQueryTestClamping::FilterThreshold)
diff --git a/Engine/Source/Runtime/AIModule/Private/EnvironmentQuery/Generators/EnvQueryGenerator_OnCircle.cpp b/Engine/Source/Runtime/AIModule/Private/EnvironmentQuery/Generators/EnvQueryGenerator_OnCircle.cpp
--- a/Engine/Source/Runtime/AIModule/Private/EnvironmentQuery/Generators/EnvQueryGenerator_OnCircle.cpp
+++ b/Engine/Source/Runtime/AIModule/Private/EnvironmentQuery/Generators/EnvQueryGenerator_OnCircle.cpp
@@ -132,6 +132,7 @@ void UEnvQueryGenerator_OnCircle::GenerateItems(FEnvQueryInstance& QueryInstance
 	SpaceBetween.BindData(QueryOwner, QueryInstance.QueryID);
 	ArcAngle.BindData(QueryOwner, QueryInstance.QueryID);
 	CircleCenterZOffset.BindData(QueryOwner, QueryInstance.QueryID);
+	NumberOfPoints.BindData(QueryOwner, QueryInstance.QueryID);
 
 	float AngleDegree = ArcAngle.GetValue();
 	float RadiusValue = CircleRadius.GetValue();
diff --git a/Engine/Source/Runtime/AIModule/Private/GameplayDebugger/GameplayDebuggerCategory_NavLocalGrid.cpp b/Engine/Source/Runtime/AIModule/Private/GameplayDebugger/GameplayDebuggerCategory_NavLocalGrid.cpp
new file mode 100644
--- /dev/null
+++ b/Engine/Source/Runtime/AIModule/Private/GameplayDebugger/GameplayDebuggerCategory_NavLocalGrid.cpp
@@ -0,0 +1,423 @@
+// Copyright 1998-2016 Epic Games, Inc. All Rights Reserved.
+
+#include "AIModulePrivate.h"
+
+#if WITH_GAMEPLAY_DEBUGGER
+#include "Navigation/NavLocalGridManager.h"
+#include "Navigation/GridPathFollowingComponent.h"
+#include "GameplayDebuggerCategory_NavLocalGrid.h"
+#include "DebugRenderSceneProxy.h"
+#include "Materials/Material.h"
+
+//////////////////////////////////////////////////////////////////////////
+// Scene proxy
+
+class FNavLocalGridSceneProxy : public FDebugRenderSceneProxy
+{
+public:
+	FNavLocalGridSceneProxy(const UPrimitiveComponent* InComponent,
+		const FGameplayDebuggerCategory_NavLocalGrid::FRepData& RepData,
+		const FGameplayDebuggerCategory_NavLocalGrid::FRepAgentData& AgentData) 
+		: FDebugRenderSceneProxy(InComponent)
+	{
+		for (int32 Idx = 0; Idx < RepData.Grids.Num(); Idx++)
+		{
+			const FGameplayDebuggerCategory_NavLocalGrid::FRepData::FGridData& GridData = RepData.Grids[Idx];
+			StoreGridBounds(GridData.Bounds, BoundsMeshVerts, BoundsMeshIndices);
+
+			for (int32 IdxX = 0; IdxX < GridData.NumCols; IdxX++)
+			{
+				for (int32 IdxY = 0; IdxY < GridData.NumRows; IdxY++)
+				{
+					const int32 CellIndex = (IdxX * GridData.NumRows) + IdxY;
+					const bool bIsMarked = GridData.Cells[CellIndex] != 0;
+					const bool bIsOnPath = (AgentData.GridIdx == Idx) && AgentData.PathCells.Contains(CellIndex);
+
+					if (bIsMarked)
+					{
+						StoreGridCellFull(GridData.Bounds, FIntPoint(IdxX, IdxY), GridData.CellSize, MarkedCellMeshVerts, MarkedCellMeshIndices);
+					}
+					else if (bIsOnPath)
+					{
+						StoreGridCellLayer(GridData.Bounds, FIntPoint(IdxX, IdxY), GridData.CellSize, PathMeshVerts, PathMeshIndices);
+					}
+					else
+					{
+						StoreGridCellLayer(GridData.Bounds, FIntPoint(IdxX, IdxY), GridData.CellSize, FreeCellMeshVerts, FreeCellMeshIndices);
+					}
+				}
+			}
+		}
+
+		if (RepData.Grids.IsValidIndex(AgentData.GridIdx))
+		{
+			const FGameplayDebuggerCategory_NavLocalGrid::FRepData::FGridData& GridData = RepData.Grids[AgentData.GridIdx];
+			for (int32 Idx = 1; Idx < AgentData.PathCells.Num(); Idx++)
+			{
+				const FIntPoint P0(AgentData.PathCells[Idx - 1] / GridData.NumRows, AgentData.PathCells[Idx - 1] % GridData.NumRows);
+				const FIntPoint P1(AgentData.PathCells[Idx] / GridData.NumRows, AgentData.PathCells[Idx] % GridData.NumRows);
+
+				StoreGridConnector(GridData.Bounds, P0, P1, GridData.CellSize, PathMeshVerts, PathMeshIndices);
+			}
+		}
+	}
+
+	virtual void GetDynamicMeshElements(const TArray<const FSceneView*>& Views, const FSceneViewFamily& ViewFamily, uint32 VisibilityMap, FMeshElementCollector& Collector) const
+	{
+		static FColor BoundsColor(255, 255, 0, 128);
+		static FColor FreeCellColor(255, 255, 255, 16);
+		static FColor MarkedCellColor(255, 0, 0, 16);
+		static FColor PathColor(0, 255, 255, 128);
+
+		for (int32 ViewIndex = 0; ViewIndex < Views.Num(); ViewIndex++)
+		{
+			if (VisibilityMap & (1 << ViewIndex))
+			{
+				const FSceneView* View = Views[ViewIndex];
+
+				if (FreeCellMeshVerts.Num() > 0)
+				{
+					const FColoredMaterialRenderProxy* MeshColorInstance = new(FMemStack::Get()) FColoredMaterialRenderProxy(GEngine->DebugMeshMaterial->GetRenderProxy(false), FLinearColor(FreeCellColor));
+					FDynamicMeshBuilder	MeshBuilder;
+					MeshBuilder.AddVertices(FreeCellMeshVerts);
+					MeshBuilder.AddTriangles(FreeCellMeshIndices);
+					MeshBuilder.GetMesh(FMatrix::Identity, MeshColorInstance, GetDepthPriorityGroup(View), false, false, ViewIndex, Collector);
+				}
+
+				if (MarkedCellMeshVerts.Num() > 0)
+				{
+					const FColoredMaterialRenderProxy* MeshColorInstance = new(FMemStack::Get()) FColoredMaterialRenderProxy(GEngine->DebugMeshMaterial->GetRenderProxy(false), FLinearColor(MarkedCellColor));
+					FDynamicMeshBuilder	MeshBuilder;
+					MeshBuilder.AddVertices(MarkedCellMeshVerts);
+					MeshBuilder.AddTriangles(MarkedCellMeshIndices);
+					MeshBuilder.GetMesh(FMatrix::Identity, MeshColorInstance, GetDepthPriorityGroup(View), false, false, ViewIndex, Collector);
+				}
+
+				if (BoundsMeshVerts.Num() > 0)
+				{
+					const FColoredMaterialRenderProxy* MeshColorInstance = new(FMemStack::Get()) FColoredMaterialRenderProxy(GEngine->DebugMeshMaterial->GetRenderProxy(false), FLinearColor(BoundsColor));
+					FDynamicMeshBuilder	MeshBuilder;
+					MeshBuilder.AddVertices(BoundsMeshVerts);
+					MeshBuilder.AddTriangles(BoundsMeshIndices);
+					MeshBuilder.GetMesh(FMatrix::Identity, MeshColorInstance, GetDepthPriorityGroup(View), false, false, ViewIndex, Collector);
+				}
+
+				if (PathMeshVerts.Num() > 0)
+				{
+					const FColoredMaterialRenderProxy* MeshColorInstance = new(FMemStack::Get()) FColoredMaterialRenderProxy(GEngine->DebugMeshMaterial->GetRenderProxy(false), FLinearColor(PathColor));
+					FDynamicMeshBuilder	MeshBuilder;
+					MeshBuilder.AddVertices(PathMeshVerts);
+					MeshBuilder.AddTriangles(PathMeshIndices);
+					MeshBuilder.GetMesh(FMatrix::Identity, MeshColorInstance, GetDepthPriorityGroup(View), false, false, ViewIndex, Collector);
+				}
+			}
+		}
+	}
+
+protected:
+
+	inline void StoreTriIndices(int32 V0, int32 V1, int32 V2, int32 FirstVertIndex, TArray<int32>& Indices)
+	{
+		Indices.Add(FirstVertIndex + V0);
+		Indices.Add(FirstVertIndex + V1);
+		Indices.Add(FirstVertIndex + V2);
+	}
+
+	inline void StoreQuadIndices(int32 V0, int32 V1, int32 V2, int32 V3, int32 FirstVertIndex, TArray<int32>& Indices)
+	{
+		StoreTriIndices(V0, V1, V2, FirstVertIndex, Indices);
+		StoreTriIndices(V0, V2, V3, FirstVertIndex, Indices);
+	}
+
+	void StoreGridConnector(const FBox& GridBounds, const FIntPoint& P0, const FIntPoint& P1, const float CellSize, TArray<FDynamicMeshVertex>& Verts, TArray<int32>& Indices)
+	{
+		const float Width = 5.0f;
+
+		const FVector Location0 = FVector(GridBounds.Min.X, GridBounds.Min.Y, GridBounds.Max.Z) + FVector(CellSize * (P0.X + 0.5f), CellSize * (P0.Y + 0.5f), 0);
+		const FVector Location1 = FVector(GridBounds.Min.X, GridBounds.Min.Y, GridBounds.Max.Z) + FVector(CellSize * (P1.X + 0.5f), CellSize * (P1.Y + 0.5f), 0);
+		const FVector DirFwd = (Location1 - Location0).GetSafeNormal();
+		const FVector DirRight = FVector::CrossProduct(DirFwd, FVector::UpVector);
+
+		const int32 FirstVertIdx = Verts.Num();
+
+		Verts.Add(FDynamicMeshVertex(Location0 - (DirRight * Width)));
+		Verts.Add(FDynamicMeshVertex(Location0 + (DirRight * Width)));
+		Verts.Add(FDynamicMeshVertex(Location1 - (DirRight * Width)));
+		Verts.Add(FDynamicMeshVertex(Location1 + (DirRight * Width)));
+
+		StoreQuadIndices(0, 1, 3, 2, FirstVertIdx, Indices);
+	}
+
+	void StoreGridCellLayer(const FBox& GridBounds, const FIntPoint& CellIdx, const float CellSize, TArray<FDynamicMeshVertex>& Verts, TArray<int32>& Indices)
+	{
+		const float CellGapSize = 5.0f;
+		
+		const FVector CellMin = FVector(GridBounds.Min.X, GridBounds.Min.Y, GridBounds.Max.Z) + FVector(CellSize * CellIdx.X, CellSize * CellIdx.Y, 0);
+		const FBox CellBox(CellMin, CellMin + FVector(CellSize, CellSize, 0.0f));
+		const FVector Center = CellBox.GetCenter();
+		const FVector Extent = CellBox.GetExtent() - FVector(CellGapSize, CellGapSize, 0.f);
+
+		const int32 FirstVertIdx = Verts.Num();
+
+		Verts.Add(FDynamicMeshVertex(FVector(Center.X - Extent.X, Center.Y - Extent.Y, Center.Z - Extent.Z)));
+		Verts.Add(FDynamicMeshVertex(FVector(Center.X + Extent.X, Center.Y - Extent.Y, Center.Z - Extent.Z)));
+		Verts.Add(FDynamicMeshVertex(FVector(Center.X - Extent.X, Center.Y + Extent.Y, Center.Z - Extent.Z)));
+		Verts.Add(FDynamicMeshVertex(FVector(Center.X + Extent.X, Center.Y + Extent.Y, Center.Z - Extent.Z)));
+
+		StoreQuadIndices(0, 1, 3, 2, FirstVertIdx, Indices);
+	}
+
+	void StoreGridCellFull(const FBox& GridBounds, const FIntPoint& CellIdx, const float CellSize, TArray<FDynamicMeshVertex>& Verts, TArray<int32>& Indices)
+	{
+		const float CellGapSize = 5.0f;
+
+		const FVector CellMin = FVector(GridBounds.Min.X, GridBounds.Min.Y, GridBounds.Min.Z) + FVector(CellSize * CellIdx.X, CellSize * CellIdx.Y, 0);
+		const FBox CellBox(CellMin, CellMin + FVector(CellSize, CellSize, GridBounds.Max.Z - GridBounds.Min.Z));
+		const FVector Center = CellBox.GetCenter();
+		const FVector Extent = CellBox.GetExtent() - FVector(CellGapSize, CellGapSize, 0.f);
+
+		const int32 FirstVertIdx = Verts.Num();
+
+		Verts.Add(FDynamicMeshVertex(FVector(Center.X - Extent.X, Center.Y - Extent.Y, Center.Z - Extent.Z)));
+		Verts.Add(FDynamicMeshVertex(FVector(Center.X + Extent.X, Center.Y - Extent.Y, Center.Z - Extent.Z)));
+		Verts.Add(FDynamicMeshVertex(FVector(Center.X - Extent.X, Center.Y + Extent.Y, Center.Z - Extent.Z)));
+		Verts.Add(FDynamicMeshVertex(FVector(Center.X + Extent.X, Center.Y + Extent.Y, Center.Z - Extent.Z)));
+		Verts.Add(FDynamicMeshVertex(FVector(Center.X - Extent.X, Center.Y - Extent.Y, Center.Z + Extent.Z)));
+		Verts.Add(FDynamicMeshVertex(FVector(Center.X + Extent.X, Center.Y - Extent.Y, Center.Z + Extent.Z)));
+		Verts.Add(FDynamicMeshVertex(FVector(Center.X - Extent.X, Center.Y + Extent.Y, Center.Z + Extent.Z)));
+		Verts.Add(FDynamicMeshVertex(FVector(Center.X + Extent.X, Center.Y + Extent.Y, Center.Z + Extent.Z)));
+
+		StoreQuadIndices(4, 5, 7, 6, FirstVertIdx, Indices);
+		StoreQuadIndices(1, 3, 7, 5, FirstVertIdx, Indices);
+		StoreQuadIndices(0, 1, 5, 4, FirstVertIdx, Indices);
+
+		StoreQuadIndices(1, 0, 2, 3, FirstVertIdx, Indices);
+		StoreQuadIndices(0, 4, 6, 2, FirstVertIdx, Indices);
+		StoreQuadIndices(2, 6, 7, 3, FirstVertIdx, Indices);
+	}
+
+	void StoreGridBounds(const FBox& GridBounds, TArray<FDynamicMeshVertex>& Verts, TArray<int32>& Indices)
+	{
+		const float FaceWidth = 5.0f;
+
+		const int32 FirstVertIdx = Verts.Num();
+		const FVector Center = GridBounds.GetCenter();
+		const FVector Extent = GridBounds.GetExtent();
+
+		Verts.Add(FDynamicMeshVertex(FVector(Center.X - Extent.X, Center.Y - Extent.Y, Center.Z - Extent.Z)));
+		Verts.Add(FDynamicMeshVertex(FVector(Center.X + Extent.X, Center.Y - Extent.Y, Center.Z - Extent.Z)));
+		Verts.Add(FDynamicMeshVertex(FVector(Center.X - Extent.X, Center.Y + Extent.Y, Center.Z - Extent.Z)));
+		Verts.Add(FDynamicMeshVertex(FVector(Center.X + Extent.X, Center.Y + Extent.Y, Center.Z - Extent.Z)));
+		Verts.Add(FDynamicMeshVertex(FVector(Center.X - Extent.X, Center.Y - Extent.Y, Center.Z + Extent.Z)));
+		Verts.Add(FDynamicMeshVertex(FVector(Center.X + Extent.X, Center.Y - Extent.Y, Center.Z + Extent.Z)));
+		Verts.Add(FDynamicMeshVertex(FVector(Center.X - Extent.X, Center.Y + Extent.Y, Center.Z + Extent.Z)));
+		Verts.Add(FDynamicMeshVertex(FVector(Center.X + Extent.X, Center.Y + Extent.Y, Center.Z + Extent.Z)));
+
+		Verts.Add(FDynamicMeshVertex(FVector(Center.X - Extent.X + FaceWidth, Center.Y - Extent.Y, Center.Z - Extent.Z + FaceWidth)));
+		Verts.Add(FDynamicMeshVertex(FVector(Center.X + Extent.X - FaceWidth, Center.Y - Extent.Y, Center.Z - Extent.Z + FaceWidth)));
+		Verts.Add(FDynamicMeshVertex(FVector(Center.X - Extent.X + FaceWidth, Center.Y - Extent.Y, Center.Z + Extent.Z - FaceWidth)));
+		Verts.Add(FDynamicMeshVertex(FVector(Center.X + Extent.X - FaceWidth, Center.Y - Extent.Y, Center.Z + Extent.Z - FaceWidth)));
+
+		Verts.Add(FDynamicMeshVertex(FVector(Center.X + Extent.X, Center.Y - Extent.Y + FaceWidth, Center.Z - Extent.Z + FaceWidth)));
+		Verts.Add(FDynamicMeshVertex(FVector(Center.X + Extent.X, Center.Y + Extent.Y - FaceWidth, Center.Z - Extent.Z + FaceWidth)));
+		Verts.Add(FDynamicMeshVertex(FVector(Center.X + Extent.X, Center.Y - Extent.Y + FaceWidth, Center.Z + Extent.Z - FaceWidth)));
+		Verts.Add(FDynamicMeshVertex(FVector(Center.X + Extent.X, Center.Y + Extent.Y - FaceWidth, Center.Z + Extent.Z - FaceWidth)));
+
+		Verts.Add(FDynamicMeshVertex(FVector(Center.X - Extent.X + FaceWidth, Center.Y - Extent.Y + FaceWidth, Center.Z + Extent.Z)));
+		Verts.Add(FDynamicMeshVertex(FVector(Center.X + Extent.X - FaceWidth, Center.Y - Extent.Y + FaceWidth, Center.Z + Extent.Z)));
+		Verts.Add(FDynamicMeshVertex(FVector(Center.X - Extent.X + FaceWidth, Center.Y + Extent.Y - FaceWidth, Center.Z + Extent.Z)));
+		Verts.Add(FDynamicMeshVertex(FVector(Center.X + Extent.X - FaceWidth, Center.Y + Extent.Y - FaceWidth, Center.Z + Extent.Z)));
+
+		Verts.Add(FDynamicMeshVertex(FVector(Center.X - Extent.X + FaceWidth, Center.Y + Extent.Y, Center.Z - Extent.Z + FaceWidth)));
+		Verts.Add(FDynamicMeshVertex(FVector(Center.X + Extent.X - FaceWidth, Center.Y + Extent.Y, Center.Z - Extent.Z + FaceWidth)));
+		Verts.Add(FDynamicMeshVertex(FVector(Center.X - Extent.X + FaceWidth, Center.Y + Extent.Y, Center.Z + Extent.Z - FaceWidth)));
+		Verts.Add(FDynamicMeshVertex(FVector(Center.X + Extent.X - FaceWidth, Center.Y + Extent.Y, Center.Z + Extent.Z - FaceWidth)));
+
+		Verts.Add(FDynamicMeshVertex(FVector(Center.X - Extent.X, Center.Y - Extent.Y + FaceWidth, Center.Z - Extent.Z + FaceWidth)));
+		Verts.Add(FDynamicMeshVertex(FVector(Center.X - Extent.X, Center.Y + Extent.Y - FaceWidth, Center.Z - Extent.Z + FaceWidth)));
+		Verts.Add(FDynamicMeshVertex(FVector(Center.X - Extent.X, Center.Y - Extent.Y + FaceWidth, Center.Z + Extent.Z - FaceWidth)));
+		Verts.Add(FDynamicMeshVertex(FVector(Center.X - Extent.X, Center.Y + Extent.Y - FaceWidth, Center.Z + Extent.Z - FaceWidth)));
+
+		Verts.Add(FDynamicMeshVertex(FVector(Center.X - Extent.X + FaceWidth, Center.Y - Extent.Y + FaceWidth, Center.Z - Extent.Z)));
+		Verts.Add(FDynamicMeshVertex(FVector(Center.X + Extent.X - FaceWidth, Center.Y - Extent.Y + FaceWidth, Center.Z - Extent.Z)));
+		Verts.Add(FDynamicMeshVertex(FVector(Center.X - Extent.X + FaceWidth, Center.Y + Extent.Y - FaceWidth, Center.Z - Extent.Z)));
+		Verts.Add(FDynamicMeshVertex(FVector(Center.X + Extent.X - FaceWidth, Center.Y + Extent.Y - FaceWidth, Center.Z - Extent.Z)));
+
+		StoreQuadIndices(0, 1, 9, 8, FirstVertIdx, Indices);
+		StoreQuadIndices(1, 5, 11, 9, FirstVertIdx, Indices);
+		StoreQuadIndices(5, 4, 10, 11, FirstVertIdx, Indices);
+		StoreQuadIndices(4, 0, 8, 10, FirstVertIdx, Indices);
+
+		StoreQuadIndices(1, 3, 13, 12, FirstVertIdx, Indices);
+		StoreQuadIndices(3, 7, 15, 13, FirstVertIdx, Indices);
+		StoreQuadIndices(7, 5, 14, 15, FirstVertIdx, Indices);
+		StoreQuadIndices(5, 1, 12, 14, FirstVertIdx, Indices);
+
+		StoreQuadIndices(4, 5, 17, 16, FirstVertIdx, Indices);
+		StoreQuadIndices(5, 7, 19, 17, FirstVertIdx, Indices);
+		StoreQuadIndices(7, 6, 18, 19, FirstVertIdx, Indices);
+		StoreQuadIndices(6, 4, 16, 18, FirstVertIdx, Indices);
+
+		StoreQuadIndices(3, 2, 20, 21, FirstVertIdx, Indices);
+		StoreQuadIndices(7, 3, 21, 23, FirstVertIdx, Indices);
+		StoreQuadIndices(6, 7, 23, 22, FirstVertIdx, Indices);
+		StoreQuadIndices(2, 6, 22, 20, FirstVertIdx, Indices);
+
+		StoreQuadIndices(2, 0, 24, 25, FirstVertIdx, Indices);
+		StoreQuadIndices(6, 2, 25, 27, FirstVertIdx, Indices);
+		StoreQuadIndices(4, 6, 27, 26, FirstVertIdx, Indices);
+		StoreQuadIndices(0, 4, 26, 24, FirstVertIdx, Indices);
+
+		StoreQuadIndices(1, 0, 28, 29, FirstVertIdx, Indices);
+		StoreQuadIndices(3, 1, 29, 31, FirstVertIdx, Indices);
+		StoreQuadIndices(2, 3, 31, 30, FirstVertIdx, Indices);
+		StoreQuadIndices(0, 2, 30, 28, FirstVertIdx, Indices);
+	}
+
+	virtual FPrimitiveViewRelevance GetViewRelevance(const FSceneView* View) const override
+	{
+		FPrimitiveViewRelevance Result;
+		Result.bDrawRelevance = IsShown(View);
+		Result.bDynamicRelevance = true;
+		Result.bNormalTranslucencyRelevance = Result.bSeparateTranslucencyRelevance = IsShown(View);
+		return Result;
+	}
+
+	virtual uint32 GetMemoryFootprint(void) const override
+	{
+		return sizeof(*this) + GetAllocatedSize();
+	}
+
+	uint32 GetAllocatedSize(void) const
+	{
+		uint32 AllocatedSize = FDebugRenderSceneProxy::GetAllocatedSize() +
+			FreeCellMeshVerts.GetAllocatedSize() + FreeCellMeshIndices.GetAllocatedSize() +
+			MarkedCellMeshVerts.GetAllocatedSize() + MarkedCellMeshIndices.GetAllocatedSize() +
+			BoundsMeshVerts.GetAllocatedSize() + BoundsMeshIndices.GetAllocatedSize() +
+			PathMeshVerts.GetAllocatedSize() + PathMeshIndices.GetAllocatedSize();
+
+		return AllocatedSize;
+	}
+
+private:
+	TArray<FDynamicMeshVertex> FreeCellMeshVerts;
+	TArray<int32> FreeCellMeshIndices;
+	TArray<FDynamicMeshVertex> MarkedCellMeshVerts;
+	TArray<int32> MarkedCellMeshIndices;
+	TArray<FDynamicMeshVertex> BoundsMeshVerts;
+	TArray<int32> BoundsMeshIndices;
+	TArray<FDynamicMeshVertex> PathMeshVerts;
+	TArray<int32> PathMeshIndices;
+};
+
+
+//////////////////////////////////////////////////////////////////////////
+// Category
+
+FGameplayDebuggerCategory_NavLocalGrid::FGameplayDebuggerCategory_NavLocalGrid()
+{
+	SetDataPackReplication<FRepData>(&DataPack, EGameplayDebuggerDataPack::Persistent);
+	SetDataPackReplication<FRepAgentData>(&AgentDataPack);
+
+	bShowUpdateTimer = false;
+	bShowDataPackReplication = true;
+	bShowOnlyWithDebugActor = false;
+}
+
+TSharedRef<FGameplayDebuggerCategory> FGameplayDebuggerCategory_NavLocalGrid::MakeInstance()
+{
+	return MakeShareable(new FGameplayDebuggerCategory_NavLocalGrid());
+}
+
+void FGameplayDebuggerCategory_NavLocalGrid::FRepData::Serialize(FArchive& Ar)
+{
+	Ar << NumSources;
+
+	int32 NumGrids = Grids.Num();
+	Ar << NumGrids;
+	if (Ar.IsLoading())
+	{
+		Grids.SetNum(NumGrids);
+	}
+
+	for (int32 Idx = 0; Idx < NumGrids; Idx++)
+	{
+		FGridData& Data = Grids[Idx];
+		Ar << Data.Bounds;
+		Ar << Data.CellSize;
+		Ar << Data.NumRows;
+		Ar << Data.NumCols;
+
+		Data.Cells.BulkSerialize(Ar);
+	}
+}
+
+void FGameplayDebuggerCategory_NavLocalGrid::FRepAgentData::Serialize(FArchive& Ar)
+{
+	Ar << PathCells;
+	Ar << GridIdx;
+}
+
+void FGameplayDebuggerCategory_NavLocalGrid::CollectData(APlayerController* OwnerPC, AActor* DebugActor)
+{
+	UNavLocalGridManager* GridManager = UNavLocalGridManager::GetCurrent(OwnerPC);
+	if (GridManager && GridManager->GetVersion() != DataPack.VersionNum)
+	{
+		DataPack.VersionNum = GridManager->GetVersion();
+		DataPack.NumSources = GridManager->GetNumSources();
+		DataPack.Grids.SetNum(GridManager->GetNumGrids());
+
+		for (int32 Idx = 0; Idx < GridManager->GetNumGrids(); Idx++)
+		{
+			const FNavLocalGridData& GridData = GridManager->GetGridData(Idx);
+			FRepData::FGridData& DebugGrid = DataPack.Grids[Idx];
+
+			DebugGrid.Bounds = GridData.WorldBounds;
+			DebugGrid.CellSize = GridData.GridCellSize;
+			DebugGrid.NumCols = (int32)GridData.GridSize.Width;
+			DebugGrid.NumRows = (int32)GridData.GridSize.Height;
+
+			const uint32 NumGridCells = GridData.GetCellsCount();
+			DebugGrid.Cells.SetNum(NumGridCells);
+			for (uint32 CellIdx = 0; CellIdx < NumGridCells; CellIdx++)
+			{
+				DebugGrid.Cells[CellIdx] = GridData.GetCellAtIndexUnsafe(CellIdx);
+			}
+		}
+	}
+
+	const APawn* DebugPawn = Cast<APawn>(DebugActor);
+	const AAIController* DebugAI = DebugPawn ? Cast<AAIController>(DebugPawn->GetController()) : nullptr;
+	const UGridPathFollowingComponent* GridPathComp = DebugAI ? Cast<UGridPathFollowingComponent>(DebugAI->GetPathFollowingComponent()) : nullptr;
+	if (GridPathComp && GridPathComp->HasActiveGrid() && GridManager)
+	{
+		AgentDataPack.GridIdx = GridPathComp->GetActiveGridIdx();
+
+		TArray<FVector> PathPoints = GridPathComp->GetGridPathPoints();
+		const FNavLocalGridData& GridData = GridManager->GetGridData(AgentDataPack.GridIdx);
+	
+		AgentDataPack.PathCells.SetNum(PathPoints.Num());
+		for (int32 Idx = 0; Idx < PathPoints.Num(); Idx++)
+		{
+			AgentDataPack.PathCells[Idx] = GridData.GetCellIndex(PathPoints[Idx]);
+		}
+	}
+}
+
+void FGameplayDebuggerCategory_NavLocalGrid::DrawData(APlayerController* OwnerPC, FGameplayDebuggerCanvasContext& CanvasContext)
+{
+	CanvasContext.Printf(TEXT("Num sources: {%s}%d"), DataPack.NumSources ? TEXT("yellow") : TEXT("red"), DataPack.NumSources);
+}
+
+FDebugRenderSceneProxy* FGameplayDebuggerCategory_NavLocalGrid::CreateDebugSceneProxy(const UPrimitiveComponent* InComponent, FDebugDrawDelegateHelper*& OutDelegateHelper)
+{
+	OutDelegateHelper = nullptr;
+	return new FNavLocalGridSceneProxy(InComponent, DataPack, AgentDataPack);
+}
+
+void FGameplayDebuggerCategory_NavLocalGrid::OnDataPackReplicated(int32 DataPackId)
+{
+	MarkRenderStateDirty();
+}
+
+#endif // WITH_GAMEPLAY_DEBUGGER
diff --git a/Engine/Source/Runtime/AIModule/Private/GameplayDebugger/GameplayDebuggerCategory_Navmesh.cpp b/Engine/Source/Runtime/AIModule/Private/GameplayDebugger/GameplayDebuggerCategory_Navmesh.cpp
--- a/Engine/Source/Runtime/AIModule/Private/GameplayDebugger/GameplayDebuggerCategory_Navmesh.cpp
+++ b/Engine/Source/Runtime/AIModule/Private/GameplayDebugger/GameplayDebuggerCategory_Navmesh.cpp
@@ -30,8 +30,11 @@ void FGameplayDebuggerCategory_Navmesh::CollectData(APlayerController* OwnerPC,
 	if (OwnerPC && DestPawn)
 	{
 		UNavigationSystem* NavSys = UNavigationSystem::GetCurrent(OwnerPC->GetWorld());
-		const FNavAgentProperties& NavAgentProperties = DestPawn->GetNavAgentPropertiesRef();
-		NavData = Cast<ARecastNavMesh>(NavSys->GetNavDataForProps(NavAgentProperties));
+		if (NavSys) 
+		{
+			const FNavAgentProperties& NavAgentProperties = DestPawn->GetNavAgentPropertiesRef();
+			NavData = Cast<ARecastNavMesh>(NavSys->GetNavDataForProps(NavAgentProperties));
+		}
 	}
 
 	if (NavData)
diff --git a/Engine/Source/Runtime/AIModule/Private/Perception/AIPerceptionComponent.cpp b/Engine/Source/Runtime/AIModule/Private/Perception/AIPerceptionComponent.cpp
--- a/Engine/Source/Runtime/AIModule/Private/Perception/AIPerceptionComponent.cpp
+++ b/Engine/Source/Runtime/AIModule/Private/Perception/AIPerceptionComponent.cpp
@@ -95,6 +95,11 @@ UAISenseConfig* UAIPerceptionComponent::GetSenseConfig(const FAISenseID& SenseID
 void UAIPerceptionComponent::PostInitProperties() 
 {
 	Super::PostInitProperties();
+
+	if (DominantSense)
+	{
+		DominantSenseID = UAISense::GetSenseID(DominantSense);
+	}
 }
 
 void UAIPerceptionComponent::ConfigureSense(UAISenseConfig& Config)
@@ -586,6 +591,17 @@ bool UAIPerceptionComponent::HasAnyActiveStimulus(const AActor& Source) const
 	return Info->HasAnyKnownStimulus();
 }
 
+bool UAIPerceptionComponent::HasAnyCurrentStimulus(const AActor& Source) const
+{
+	const FActorPerceptionInfo* Info = GetActorInfo(Source);
+	if (Info == NULL)
+	{
+		return false;
+	}
+
+	return Info->HasAnyCurrentStimulus();
+}
+
 bool UAIPerceptionComponent::HasActiveStimulus(const AActor& Source, FAISenseID Sense) const
 {
 	const FActorPerceptionInfo* Info = GetActorInfo(Source);
diff --git a/Engine/Source/Runtime/AIModule/Private/Perception/AISense_Sight.cpp b/Engine/Source/Runtime/AIModule/Private/Perception/AISense_Sight.cpp
--- a/Engine/Source/Runtime/AIModule/Private/Perception/AISense_Sight.cpp
+++ b/Engine/Source/Runtime/AIModule/Private/Perception/AISense_Sight.cpp
@@ -349,9 +349,34 @@ void UAISense_Sight::UnregisterSource(AActor& SourceActor)
 {
 	const FAISightTarget::FTargetId AsTargetId = SourceActor.GetFName();
 	FAISightTarget* AsTarget = ObservedTargets.Find(AsTargetId);
-	if (AsTarget != nullptr)
+	
+	if (AsTarget != nullptr && SightQueryQueue.Num() > 0)
 	{
-		RemoveAllQueriesToTarget(AsTargetId, DontSort);
+		AActor* TargetActor = AsTarget->Target.Get();
+
+		if (TargetActor)
+		{
+			// notify all interested observers that this source is no longer
+			// visible		
+			AIPerception::FListenerMap& ListenersMap = *GetListeners();
+			const FAISightQuery* SightQuery = &SightQueryQueue[SightQueryQueue.Num() - 1];
+			for (int32 QueryIndex = SightQueryQueue.Num() - 1; QueryIndex >= 0; --QueryIndex, --SightQuery)
+			{
+				if (SightQuery->TargetId == AsTargetId)
+				{
+					if (SightQuery->bLastResult == true)
+					{
+						FPerceptionListener& Listener = ListenersMap[SightQuery->ObserverId];
+						ensure(Listener.Listener.IsValid());
+
+						Listener.RegisterStimulus(TargetActor, FAIStimulus(*this, 0.f, SightQuery->LastSeenLocation, Listener.CachedLocation, FAIStimulus::SensingFailed));
+					}
+
+					SightQueryQueue.RemoveAt(QueryIndex, 1, /*bAllowShrinking=*/false);
+				}
+			}
+			// no point in sorting, we haven't change the order of other queries
+		}
 	}
 }
 
diff --git a/Engine/Source/Runtime/AIModule/Public/GraphAStar.h b/Engine/Source/Runtime/AIModule/Public/GraphAStar.h
--- a/Engine/Source/Runtime/AIModule/Public/GraphAStar.h
+++ b/Engine/Source/Runtime/AIModule/Public/GraphAStar.h
@@ -21,13 +21,14 @@ enum EGraphAStarResult
 /**
  *	Generic graph A* implementation
  *
- *	TGraph holds graph representation. Needs to implement two functions:
- *		int32 GetNeighbourCount(FNodeRef NodeRef) const;	- returns number of neighbours that the graph node identified with NodeRef has
- *		bool IsValidRef(FNodeRef NodeRef) const;			- returns whether given node identyfication is correct
+ *	TGraph holds graph representation. Needs to implement functions:
+ *		int32 GetNeighbourCount(FNodeRef NodeRef) const;										- returns number of neighbours that the graph node identified with NodeRef has
+ *		bool IsValidRef(FNodeRef NodeRef) const;												- returns whether given node identyfication is correct
+ *      FNodeRef GetNeighbour(const FNodeRef NodeRef, const int32 NeighbourIndex) const;		- returns neighbour ref
  *	
  *	it also needs to specify node type 
  *		FNodeRef		- type used as identification of nodes in the graph
-  *	
+ *	
  *	TQueryFilter (FindPath's parameter) filter class is what decides which graph edges can be used and at what cost. It needs to implement following functions:
  *		float GetHeuristicScale() const;														- used as GetHeuristicCost's multiplier
  *		float GetHeuristicCost(const FNodeRef StartNodeRef, const FNodeRef EndNodeRef) const;	- estimate of cost from StartNodeRef to EndNodeRef
diff --git a/Engine/Source/Runtime/AIModule/Public/SimpleCellGrid.h b/Engine/Source/Runtime/AIModule/Public/SimpleCellGrid.h
--- a/Engine/Source/Runtime/AIModule/Public/SimpleCellGrid.h
+++ b/Engine/Source/Runtime/AIModule/Public/SimpleCellGrid.h
@@ -3,7 +3,6 @@
 #pragma once
 
 #include "Serialization/Archive.h"
-#include "HAL/UnrealMemory.h"
 
 struct FGridSize2D
 {
@@ -22,212 +21,339 @@ struct TSimpleCellGrid
 {
 	typedef CellType FCellType;
 
-	/** *Expresses in whole unreal units */
+	// DEPRECATED, use GridCellSize instead of this
 	uint32 CellSize;
+
+	float GridCellSize;
 	FBox WorldBounds;
 	FVector Origin;
 	FVector BoundsSize;
 	FGridSize2D GridSize;
 
 protected:
-	FCellType* Cells;
+	TArray<FCellType> Cells;
 	static FCellType InvalidCell;
 
 public:
 	TSimpleCellGrid()
 		: CellSize(0)
+		, GridCellSize(0.0f)
 		, WorldBounds(EForceInit::ForceInitToZero)
 		, Origin(FLT_MAX)
 		, BoundsSize(0)
-		, Cells(nullptr)
 	{
 	}
 
-	TSimpleCellGrid(const TSimpleCellGrid& OtherMap)
+	bool Init(float InCellSize, const FBox& Bounds)
 	{
-		*this = OtherMap;
+		if (InCellSize <= 0.0f || !Bounds.IsValid)
+		{
+			return false;
+		}
+			
+		GridCellSize = InCellSize;
+		CellSize = FMath::TruncToInt(InCellSize);
+
+		const FVector RealBoundsSize = Bounds.GetSize();
+		GridSize = FGridSize2D(FMath::CeilToInt(RealBoundsSize.X / InCellSize), FMath::CeilToInt(RealBoundsSize.Y / InCellSize));
+		BoundsSize = FVector(GridSize.Width * InCellSize, GridSize.Height * InCellSize, RealBoundsSize.Z);
+		Origin = FVector(Bounds.Min.X, Bounds.Min.Y, (Bounds.Min.Z + Bounds.Max.Z) * 0.5f);
+		UpdateWorldBounds();
+
+		Cells.AddDefaulted(GridSize.Width * GridSize.Height);
+
+		return true;
 	}
 
-	~TSimpleCellGrid()
+	void UpdateWorldBounds()
 	{
-		FreeMemory();
+		WorldBounds = FBox(Origin - FVector(0, 0, BoundsSize.Z / 2), Origin + FVector(BoundsSize.X, BoundsSize.Y, BoundsSize.Z / 2));
 	}
 
-	TSimpleCellGrid& operator=(const TSimpleCellGrid& OtherMap)
+	FORCEINLINE bool IsValid() const
 	{
-		CellSize = OtherMap.CellSize;
-		WorldBounds = OtherMap.WorldBounds;
-		Origin = OtherMap.Origin;
-		BoundsSize = OtherMap.BoundsSize;
-		GridSize = OtherMap.GridSize;
+		return Cells.Num() && GridCellSize > 0;
+	}
 
-		FreeMemory();
+	FORCEINLINE bool IsValidIndex(const int32 CellIndex) const
+	{
+		return Cells.IsValidIndex(CellIndex);
+	}
 
-		if (OtherMap.IsValid())
-		{
-			AllocateMemory();
-			const uint32 MemSize = GetValuesMemorySize();
-			FMemory::Memcpy(Cells, OtherMap.Cells, MemSize);
-		}
+	FORCEINLINE bool IsValidCoord(int32 LocationX, int32 LocationY) const
+	{
+		return (LocationX >= 0) && (LocationX < (int32)GridSize.Width) && (LocationY >= 0) && (LocationY < (int32)GridSize.Height);
+	}
 
-		return *this;
+	FORCEINLINE bool IsValidCoord(const FIntVector& CellCoords) const
+	{
+		return IsValidCoord(CellCoords.X, CellCoords.Y);
 	}
 
-	void UpdateWorldBounds()
+	FORCEINLINE uint32 GetAllocatedSize() const
 	{
-		WorldBounds = FBox(Origin - FVector(0,0,BoundsSize.Z / 2), Origin + FVector(BoundsSize.X, BoundsSize.Y, BoundsSize.Z / 2));
+		return Cells.GetAllocatedSize();
 	}
 
-	void Init(uint32 InCellSize, const FBox& Bounds)
+	/** Convert world location to (X,Y) coords on grid, result can be outside grid */
+	FORCEINLINE FIntVector GetCellCoordsUnsafe(const FVector& WorldLocation) const
 	{
-		CellSize = InCellSize;
-		const FVector RealBoundsSize = Bounds.GetSize();
-		GridSize = FGridSize2D(static_cast<uint32>(FMath::CeilToInt(RealBoundsSize.X / InCellSize))
-			, static_cast<uint32>(FMath::CeilToInt(RealBoundsSize.Y / InCellSize)));
-		BoundsSize = FVector(GridSize.Width * InCellSize
-			, GridSize.Height * InCellSize
-			, RealBoundsSize.Z);
-		Origin = FVector(Bounds.Min.X, Bounds.Min.Y, Bounds.Min.Z + (Bounds.Max.Z - Bounds.Min.Z) / 2);
-		UpdateWorldBounds();
+		return FIntVector(
+			FMath::TruncToInt((WorldLocation.X - Origin.X) / GridCellSize),
+			FMath::TruncToInt((WorldLocation.Y - Origin.Y) / GridCellSize),
+			0);
+	}
 
-		AllocateMemory();
-		Zero();
+	/** Convert world location to (X,Y) coords on grid, result is clamped to grid */
+	FIntVector GetCellCoords(const FVector& WorldLocation) const
+	{
+		const FIntVector UnsafeCoords = GetCellCoordsUnsafe(WorldLocation);
+		return FIntVector(FMath::Clamp(UnsafeCoords.X, 0, (int32)GridSize.Width - 1), FMath::Clamp(UnsafeCoords.Y, 0, (int32)GridSize.Height - 1), 0);
 	}
 
-	void Zero()
+	/** Convert cell index to coord X on grid, result can be invalid */
+	FORCEINLINE int32 GetCellCoordX(int32 CellIndex) const
 	{
-		const uint32 MemSize = GetValuesMemorySize();
-		FMemory::Memzero(Cells, MemSize);
+		return CellIndex / GridSize.Height;
 	}
 
-	bool IsValid() const
+	/** Convert cell index to coord Y on grid, result can be invalid */
+	FORCEINLINE int32 GetCellCoordY(int32 CellIndex) const
 	{
-		return Cells != nullptr && CellSize > 0 && GridSize.Width > 0 && GridSize.Height > 0;
+		return CellIndex % GridSize.Height;
 	}
 
-	bool IsValidCellIndex(const uint32 CellIndex) const
+	/** Convert cell index to (X,Y) coords on grid */
+	FORCEINLINE FIntVector GetCellCoords(int32 CellIndex) const
 	{
-		return CellIndex < (GridSize.Height * GridSize.Width);
+		return FIntVector(GetCellCoordX(CellIndex), GetCellCoordY(CellIndex), 0);
 	}
 
-	uint32 GetValuesMemorySize() const { return uint32(GridSize.Width * GridSize.Height) * sizeof(FCellType); }
-	 
-	FORCEINLINE FIntVector WorldToCellCoords(const FVector& WorldLocation) const
+	/** Convert world location to cell index, result can be invalid */
+	int32 GetCellIndexUnsafe(const FVector& WorldLocation) const
 	{
-		const uint32 LocationX = uint32(WorldLocation.X - Origin.X) / CellSize;
-		const uint32 LocationY = uint32(WorldLocation.Y - Origin.Y) / CellSize;
-		return FIntVector(LocationX, LocationY, 0);
+		const FIntVector CellCoords = GetCellCoordsUnsafe(WorldLocation);
+		return GetCellIndexUnsafe(CellCoords.X, CellCoords.Y);
 	}
-	
-	FORCEINLINE void WorldToCellCoords(const FVector& WorldLocation, uint32& LocationX, uint32& LocationY) const
+
+	/** Convert (X,Y) coords on grid to cell index, result can be invalid */
+	FORCEINLINE int32 GetCellIndexUnsafe(const FIntVector& CellCoords) const
 	{
-		LocationX = uint32(WorldLocation.X - Origin.X) / CellSize;
-		LocationY = uint32(WorldLocation.Y - Origin.Y) / CellSize;
+		return GetCellIndexUnsafe(CellCoords.X, CellCoords.Y);
 	}
 
-	FORCEINLINE uint32 WorldToCellIndex(const FVector& WorldLocation) const
+	/** Convert (X,Y) coords on grid to cell index, result can be invalid */
+	FORCEINLINE int32 GetCellIndexUnsafe(int32 LocationX, int32 LocationY) const
 	{
-		const uint32 LocationX = uint32(WorldLocation.X - Origin.X) / CellSize;
-		const uint32 LocationY = uint32(WorldLocation.Y - Origin.Y) / CellSize;
-		return CellCoordsToCellIndex(LocationX, LocationY);
+		return (LocationX * GridSize.Height) + LocationY;
 	}
 
-	FORCEINLINE FIntVector CellIndexToCoords(uint32 CellIndex, uint32& LocationX, uint32& LocationY) const
+	/** Convert (X,Y) coords on grid to cell index, returns -1 for position outside grid */
+	FORCEINLINE int32 GetCellIndex(int32 LocationX, int32 LocationY) const
 	{
-		LocationX = static_cast<uint32>(FMath::FloorToInt(CellIndex / GridSize.Height));
-		LocationY = CellIndex - (LocationX * GridSize.Height);
-		return FIntVector(LocationX, LocationY, 0);
+		return IsValidCoord(LocationX, LocationY) ? GetCellIndexUnsafe(LocationX, LocationY) : INDEX_NONE;
 	}
 
-	FORCEINLINE uint32 CellCoordsToCellIndex(const int32 LocationX, const int32 LocationY) const
+	/** Convert world location to cell index, returns -1 for position outside grid */
+	int32 GetCellIndex(const FVector& WorldLocation) const
 	{
-		return LocationX * GridSize.Height + LocationY;
+		const FIntVector CellCoords = GetCellCoordsUnsafe(WorldLocation);
+		return GetCellIndex(CellCoords.X, CellCoords.Y);
 	}
 	
-	FBox GetWorldCellBox(uint32 CellIndex) const
+	FORCEINLINE FBox GetWorldCellBox(int32 CellIndex) const
 	{
-		uint32 LocationX = 0, LocationY = 0;
-		CellIndexToCoords(CellIndex, LocationX, LocationY);
-		return GetWorldCellBox(LocationX, LocationY);
+		return GetWorldCellBox(GetCellCoordX(CellIndex), GetCellCoordY(CellIndex));
 	}
 
-	FBox GetWorldCellBox(uint32 LocationX, uint32 LocationY) const
+	FORCEINLINE FBox GetWorldCellBox(int32 LocationX, int32 LocationY) const
 	{
-		const FBox Box(FVector(LocationX * CellSize, LocationY * CellSize, -BoundsSize.Z / 2) + Origin
-			, FVector((LocationX + 1)*CellSize, (LocationY + 1)*CellSize, BoundsSize.Z / 2) + Origin);
-		return Box;
+		return FBox(
+			Origin + FVector(LocationX * GridCellSize, LocationY * GridCellSize, -BoundsSize.Z * 0.5f), 
+			Origin + FVector((LocationX + 1) * GridCellSize, (LocationY + 1) * GridCellSize, BoundsSize.Z * 0.5f)
+			);
+	}
+
+	FORCEINLINE FVector GetWorldCellCenter(int32 CellIndex) const
+	{
+		return GetWorldCellCenter(GetCellCoordX(CellIndex), GetCellCoordY(CellIndex));
+	}
+
+	FORCEINLINE FVector GetWorldCellCenter(int32 LocationX, int32 LocationY) const
+	{
+		return Origin + FVector((LocationX + 0.5f) * GridCellSize, (LocationY + 0.5f) * GridCellSize, 0);
 	}
 
 	const FCellType& GetCellAtWorldLocationUnsafe(const FVector& WorldLocation) const
 	{
-		const uint32 CellIndex = WorldToCellIndex(WorldLocation);
+		const int32 CellIndex = GetCellIndexUnsafe(WorldLocation);
 		return Cells[CellIndex];
 	}
 
-	const FCellType& GetCellAtWorldLocationSafe(const FVector& WorldLocation) const
+	const FCellType& GetCellAtWorldLocation(const FVector& WorldLocation) const
 	{
-		const uint32 CellIndex = WorldToCellIndex(WorldLocation);
-		return CellIndex < (GridSize.Height * GridSize.Width) ? Cells[CellIndex] : InvalidCell;
+		const int32 CellIndex = GetCellIndex(WorldLocation);
+		return (CellIndex == INDEX_NONE) ? InvalidCell : Cells[CellIndex];
 	}
 
-	// gets a "column" (in practice a row from cache point of view)
-	FCellType* operator[](uint32 CellX) { return &Cells[CellX * GridSize.Height]; }
-
-	/** this function puts the burden of making sure CellIndex is valid on the caller */
-	FCellType& GetCellAtIndexUnsafe(int32 CellIndex) { return Cells[CellIndex]; }
+	FORCEINLINE FCellType& operator[](int32 CellIndex) { return Cells[CellIndex]; }
+	FORCEINLINE const FCellType& operator[](int32 CellIndex) const { return Cells[CellIndex]; }
 
-	/** this function puts the burden of making sure CellIndex is valid on the caller */
-	const FCellType& GetCellAtIndexUnsafe(int32 CellIndex) const { return Cells[CellIndex]; }
+	FORCEINLINE FCellType& GetCellAtIndexUnsafe(int32 CellIndex) { return Cells.GetData()[CellIndex]; }
+	FORCEINLINE const FCellType& GetCellAtIndexUnsafe(int32 CellIndex) const { return Cells.GetData()[CellIndex]; }
 	
-	FORCEINLINE uint32 GetCellsCount() const
+	FORCEINLINE int32 GetCellsCount() const
+	{
+		return Cells.Num();
+	}
+
+	FORCEINLINE int32 Num() const
 	{
-		return static_cast<uint32>(GridSize.Width * GridSize.Height);
+		return Cells.Num();
 	}
 
 	void Serialize(FArchive& Ar)
-	{		
-		Ar << CellSize;
+	{	
+		// CellSize must act as version checking: MAX_uint32, float CellSize is used
+		uint32 VersionNum = MAX_uint32;
+		Ar << VersionNum;
+
+		if (Ar.IsLoading())
+		{
+			if (VersionNum == MAX_uint32)
+			{
+				Ar << GridCellSize;
+				CellSize = FMath::TruncToInt(GridCellSize);
+			}
+			else
+			{
+				CellSize = VersionNum;
+				GridCellSize = VersionNum * 1.0f;
+			}
+		}
+		else
+		{
+			Ar << GridCellSize;
+		}
+
 		Ar << Origin;
 		Ar << BoundsSize;
 		Ar << GridSize.Width << GridSize.Height;
 	
 		UpdateWorldBounds();
 
-		int32 DataBytesCount = sizeof(FCellType) * GridSize.Width * GridSize.Height;
-		Ar << DataBytesCount;
-
-		if (DataBytesCount > 0)
+		if (VersionNum == MAX_uint32)
 		{
-			if (Cells == nullptr)
+			Ar << Cells;
+		}
+		else
+		{
+			int32 DataBytesCount = GetAllocatedSize();
+			Ar << DataBytesCount;
+
+			if (DataBytesCount > 0)
 			{
-				// this makes sense only during loading
-				ensure(Ar.IsLoading());
-				AllocateMemory();
+				if (Ar.IsLoading())
+				{
+					Cells.SetNum(GridSize.Width * GridSize.Height);
+				}
+
+				Ar.Serialize(Cells.GetData(), DataBytesCount);
 			}
-			Ar.Serialize(Cells, DataBytesCount);
 		}
 	}
 
+	void AllocateMemory()
+	{
+		Cells.SetNum(GridSize.Width * GridSize.Height);
+	}
+
+	void FreeMemory()
+	{
+		Cells.Empty();
+	}
+
+	void Zero()
+	{
+		Cells.Reset();
+		Cells.AddZeroed(GridSize.Width * GridSize.Height);
+	}
+
 	void CleanUp()
 	{
-		FreeMemory();
-		CellSize = 0;
+		Cells.Empty();
+		GridCellSize = 0;
 		Origin = FVector(FLT_MAX);
 	}
 
-	void AllocateMemory()
+	//////////////////////////////////////////////////////////////////////////
+	// DEPRECATED SUPPORT
+
+	DEPRECATED(4.14, "This function is now deprecated, please use the one with float CellSize argument")
+	void Init(uint32 InCellSize, const FBox& Bounds)
 	{
-		FreeMemory();
-		const uint32 MemSize = GetValuesMemorySize();
-		Cells = static_cast<FCellType*>(FMemory::Malloc(MemSize));
+		Init(1.0f * InCellSize, Bounds);
 	}
 
-	void FreeMemory()
+	DEPRECATED(4.14, "This function is now deprecated, please use GetCellCoords instead.")
+	FIntVector WorldToCellCoords(const FVector& WorldLocation) const
 	{
-		if (Cells)
-		{
-			FMemory::Free(Cells);
-			Cells = nullptr;
-		}
+		const uint32 LocationX = FMath::TruncToInt((WorldLocation.X - Origin.X) / CellSize);
+		const uint32 LocationY = FMath::TruncToInt((WorldLocation.Y - Origin.Y) / CellSize);
+		return FIntVector(LocationX, LocationY, 0);
+	}
+
+	DEPRECATED(4.14, "This function is now deprecated, please use GetCellCoords instead.")
+	void WorldToCellCoords(const FVector& WorldLocation, uint32& LocationX, uint32& LocationY) const
+	{
+		LocationX = FMath::TruncToInt((WorldLocation.X - Origin.X) / CellSize);
+		LocationY = FMath::TruncToInt((WorldLocation.Y - Origin.Y) / CellSize);
 	}
-};
\ No newline at end of file
+
+	DEPRECATED(4.14, "This function is now deprecated, please use GetCellIndex instead.")
+	uint32 WorldToCellIndex(const FVector& WorldLocation) const
+	{
+		const uint32 LocationX = FMath::TruncToInt((WorldLocation.X - Origin.X) / CellSize);
+		const uint32 LocationY = FMath::TruncToInt((WorldLocation.Y - Origin.Y) / CellSize);
+		return CellCoordsToCellIndex(LocationX, LocationY);
+	}
+
+	DEPRECATED(4.14, "This function is now deprecated, please use GetCellCoords instead.")
+	FIntVector CellIndexToCoords(uint32 CellIndex) const
+	{
+		return FIntVector(CellIndex / GridSize.Height, CellIndex % GridSize.Height, 0);
+	}
+
+	DEPRECATED(4.14, "This function is now deprecated, please use GetCellCoords instead.")
+	FIntVector CellIndexToCoords(uint32 CellIndex, uint32& LocationX, uint32& LocationY) const
+	{
+		LocationX = CellIndex / GridSize.Height;
+		LocationY = CellIndex % GridSize.Height;
+		return FIntVector(LocationX, LocationY, 0);
+	}
+
+	DEPRECATED(4.14, "This function is now deprecated, please use GetCellIndex instead.")
+	uint32 CellCoordsToCellIndex(const int32 LocationX, const int32 LocationY) const
+	{
+		return (LocationX * GridSize.Height) + LocationY;
+	}
+
+	DEPRECATED(4.14, "This function is now deprecated, please use GetCellAtWorldLocation instead.")
+	const FCellType& GetCellAtWorldLocationSafe(const FVector& WorldLocation) const
+	{
+		const uint32 CellIndex = GetCellIndex(WorldLocation);
+		return CellIndex < (GridSize.Height * GridSize.Width) ? Cells[CellIndex] : InvalidCell;
+	}
+
+	DEPRECATED(4.14, "This function is now deprecated, please use GetAllocatedSize instead.")
+	uint32 GetValuesMemorySize() const
+	{
+		return GridSize.Height * GridSize.Width * sizeof(FCellType);
+	}
+
+	DEPRECATED(4.14, "This function is now deprecated, please use IsValidIndex instead.")
+	bool IsValidCellIndex(const int32 CellIndex) const
+	{
+		return IsValidIndex(CellIndex);
+	}
+};
diff --git a/Engine/Source/Runtime/AnimGraphRuntime/Private/AnimNodes/AnimNode_Slot.cpp b/Engine/Source/Runtime/AnimGraphRuntime/Private/AnimNodes/AnimNode_Slot.cpp
--- a/Engine/Source/Runtime/AnimGraphRuntime/Private/AnimNodes/AnimNode_Slot.cpp
+++ b/Engine/Source/Runtime/AnimGraphRuntime/Private/AnimNodes/AnimNode_Slot.cpp
@@ -84,7 +84,7 @@ void FAnimNode_Slot::GatherDebugData(FNodeDebugData& DebugData)
 				float CurrentAnimPos;
 				if (UAnimSequenceBase* Anim = Segment.GetAnimationData(MontageInstance->GetPosition(), CurrentAnimPos, Weight))
 				{
-					FString MontageLine = FString::Printf(TEXT("Montage: '%s' Anim: '%s' Play Time:%.2f W:%.2f%%"), *MontageInstance->Montage->GetName(), *Anim->GetName(), CurrentAnimPos, Weight*100.f);
+					FString MontageLine = FString::Printf(TEXT("Montage('%s') Anim('%s') P(%.2f) W(%.0f%%)"), *MontageInstance->Montage->GetName(), *Anim->GetName(), CurrentAnimPos, WeightData.SlotNodeWeight*100.f);
 					DebugData.BranchFlow(1.0f).AddDebugItem(MontageLine, true);
 					break;
 				}
