Non-unity fixes from streams that copied up yesterday.

diff --git a/Engine/Source/Runtime/AIModule/Private/EnvironmentQuery/EnvQueryInstance.cpp b/Engine/Source/Runtime/AIModule/Private/EnvironmentQuery/EnvQueryInstance.cpp
--- a/Engine/Source/Runtime/AIModule/Private/EnvironmentQuery/EnvQueryInstance.cpp
+++ b/Engine/Source/Runtime/AIModule/Private/EnvironmentQuery/EnvQueryInstance.cpp
@@ -12,6 +12,7 @@
 #include "EnvironmentQuery/EnvQueryGenerator.h"
 #include "EnvironmentQuery/EnvQueryManager.h"
 #include "EnvironmentQuery/Contexts/EnvQueryContext_Item.h"
+#include "AISystem.h"
 
 //----------------------------------------------------------------------//
 // FEQSQueryDebugData
diff --git a/Engine/Source/Runtime/Slate/Public/SlateSettings.h b/Engine/Source/Runtime/Slate/Public/SlateSettings.h
--- a/Engine/Source/Runtime/Slate/Public/SlateSettings.h
+++ b/Engine/Source/Runtime/Slate/Public/SlateSettings.h
@@ -1,6 +1,9 @@
 // Copyright 1998-2016 Epic Games, Inc. All Rights Reserved.
 
 #pragma once
+
+#include "CoreMinimal.h"
+#include "UObject/Object.h"
 #include "SlateSettings.generated.h"
 
 /**
diff --git a/Engine/Source/Runtime/UMG/Public/Blueprint/SlateBlueprintLibrary.h b/Engine/Source/Runtime/UMG/Public/Blueprint/SlateBlueprintLibrary.h
--- a/Engine/Source/Runtime/UMG/Public/Blueprint/SlateBlueprintLibrary.h
+++ b/Engine/Source/Runtime/UMG/Public/Blueprint/SlateBlueprintLibrary.h
@@ -4,7 +4,9 @@
 
 #include "CoreMinimal.h"
 #include "UObject/ObjectMacros.h"
+#include "UObject/Script.h"
 #include "Layout/Geometry.h"
+#include "Styling/SlateBrush.h"
 #include "Kismet/BlueprintFunctionLibrary.h"
 #include "SlateBlueprintLibrary.generated.h"
 
