Fixing some localization key conflicts that crept in from stream merges

diff --git a/Engine/Source/Runtime/AIModule/Private/EnvironmentQuery/Generators/EnvQueryGenerator_ActorsOfClass.cpp b/Engine/Source/Runtime/AIModule/Private/EnvironmentQuery/Generators/EnvQueryGenerator_ActorsOfClass.cpp
--- a/Engine/Source/Runtime/AIModule/Private/EnvironmentQuery/Generators/EnvQueryGenerator_ActorsOfClass.cpp
+++ b/Engine/Source/Runtime/AIModule/Private/EnvironmentQuery/Generators/EnvQueryGenerator_ActorsOfClass.cpp
@@ -80,7 +80,7 @@ FText UEnvQueryGenerator_ActorsOfClass::GetDescriptionTitle() const
 
 	if (!GenerateOnlyActorsInRadius.IsDynamic() && !GenerateOnlyActorsInRadius.GetValue())
 	{
-		return FText::Format(LOCTEXT("DescriptionGenerateActorsAroundContext", "{DescriptionTitle}: generate set of actors of {ActorsClass}"), Args);
+		return FText::Format(LOCTEXT("DescriptionGenerateActors", "{DescriptionTitle}: generate set of actors of {ActorsClass}"), Args);
 	}
 
 	Args.Add(TEXT("DescribeContext"), UEnvQueryTypes::DescribeContext(SearchCenter));
