Copying //UE4/Dev-Framework to //UE4/Dev-Main (Source: //UE4/Dev-Framework @ 3510040)

diff --git a/Engine/Source/Editor/UnrealEd/Public/EditorCategoryUtils.h b/Engine/Source/Editor/UnrealEd/Public/EditorCategoryUtils.h
--- a/Engine/Source/Editor/UnrealEd/Public/EditorCategoryUtils.h
+++ b/Engine/Source/Editor/UnrealEd/Public/EditorCategoryUtils.h
@@ -130,7 +130,7 @@ namespace FEditorCategoryUtils
 	 * @param  CategoriesOut	An array that will be filled with a list of hidden categories.
 	 * @param  bHomogenize		Determines if the categories should be ran through expansion and display sanitation (useful even when not being displayed, for comparisons)
 	 */
-	UNREALED_API void GetClassHideCategories(const UClass* Class, TArray<FString>& CategoriesOut, bool bHomogenize = true);
+	UNREALED_API void GetClassHideCategories(const UStruct* Class, TArray<FString>& CategoriesOut, bool bHomogenize = true);
 
 	/**
 	 * Parses out the class's "ShowCategories" metadata, and returns it 
@@ -139,7 +139,7 @@ namespace FEditorCategoryUtils
 	 * @param  Class			The class you want to pull data from.
 	 * @param  CategoriesOut	An array that will be filled with a list of shown categories.
 	 */
-	UNREALED_API void GetClassShowCategories(const UClass* Class, TArray<FString>& CategoriesOut);
+	UNREALED_API void GetClassShowCategories(const UStruct* Class, TArray<FString>& CategoriesOut);
 
 	/**
 	 * Checks to see if the category associated with the supplied common 
@@ -149,7 +149,7 @@ namespace FEditorCategoryUtils
 	 * @param  CategoryId	An id associated with a category that you want to check.
 	 * @return True if the common category is hidden, false if not.
 	 */
-	UNREALED_API bool IsCategoryHiddenFromClass(const UClass* Class, const FCommonEditorCategory::EValue CategoryId);
+	UNREALED_API bool IsCategoryHiddenFromClass(const UStruct* Class, const FCommonEditorCategory::EValue CategoryId);
 
 	/**
 	 * Checks to see if the specified category is hidden from the supplied class.
@@ -158,7 +158,7 @@ namespace FEditorCategoryUtils
 	 * @param  Category A category path that you want to check.
 	 * @return True if the category is hidden, false if not.
 	 */
-	UNREALED_API bool IsCategoryHiddenFromClass(const UClass* Class, const FText& Category);
+	UNREALED_API bool IsCategoryHiddenFromClass(const UStruct* Class, const FText& Category);
 
 	/**
 	 * Checks to see if the specified category is hidden from the supplied class.
@@ -167,7 +167,7 @@ namespace FEditorCategoryUtils
 	 * @param  Category A category path that you want to check.
 	 * @return True if the category is hidden, false if not.
 	 */
-	UNREALED_API bool IsCategoryHiddenFromClass(const UClass* Class, const FString& Category);
+	UNREALED_API bool IsCategoryHiddenFromClass(const UStruct* Class, const FString& Category);
 
 	/**
 	 * Checks to see if the specified category is hidden from the supplied Class, avoids recalculation of ClassHideCategories.
@@ -178,7 +178,7 @@ namespace FEditorCategoryUtils
 	 * @param  Category A category path that you want to check.
 	 * @return True if the category is hidden, false if not.
 	 */
-	UNREALED_API bool IsCategoryHiddenFromClass(const TArray<FString>& ClassHideCategories, const UClass* Class, const FString& Category);
+	UNREALED_API bool IsCategoryHiddenFromClass(const TArray<FString>& ClassHideCategories, const UStruct* Class, const FString& Category);
 	/**
 	 * Returns tooltip information for the specified category
 	 * 
@@ -198,5 +198,5 @@ namespace FEditorCategoryUtils
 	 * @param  Class		The class you want to query.
 	 * @return  The set of categories that should be hidden.
 	 */
-	UNREALED_API TSet<FString> GetHiddenCategories(const UClass* Class);
+	UNREALED_API TSet<FString> GetHiddenCategories(const UStruct* Class);
 };
diff --git a/Engine/Source/Editor/UnrealEd/Public/GameModeInfoCustomizer.h b/Engine/Source/Editor/UnrealEd/Public/GameModeInfoCustomizer.h
--- a/Engine/Source/Editor/UnrealEd/Public/GameModeInfoCustomizer.h
+++ b/Engine/Source/Editor/UnrealEd/Public/GameModeInfoCustomizer.h
@@ -23,6 +23,7 @@
 #include "PropertyCustomizationHelpers.h"
 #include "IDocumentation.h"
 #include "Kismet2/KismetEditorUtilities.h"
+#include "Kismet2/BlueprintEditorUtils.h"
 #include "EditorClassUtils.h"
 
 #define LOCTEXT_NAMESPACE "FGameModeInfoCustomizer"
@@ -245,8 +246,12 @@ public:
 			const UClass** DefaultClassPtr = ClassProp->ContainerPtrToValuePtr<const UClass*>(GetCurrentGameModeCDO());
 			*DefaultClassPtr = NewDefaultClass;
 
-			// Indicate that the BP has changed and would need to be saved.
-			GameModeClass->MarkPackageDirty();
+			UBlueprint* Blueprint = Cast<UBlueprint>(GameModeClass->ClassGeneratedBy);
+			if (Blueprint)
+			{
+				// Indicate that the BP has changed and would need to be saved.
+				FBlueprintEditorUtils::MarkBlueprintAsModified(Blueprint);
+			}
 		}
 	}
 
diff --git a/Engine/Source/Editor/UnrealEd/Public/Kismet2/BlueprintEditorUtils.h b/Engine/Source/Editor/UnrealEd/Public/Kismet2/BlueprintEditorUtils.h
--- a/Engine/Source/Editor/UnrealEd/Public/Kismet2/BlueprintEditorUtils.h
+++ b/Engine/Source/Editor/UnrealEd/Public/Kismet2/BlueprintEditorUtils.h
@@ -186,6 +186,9 @@ public:
 	 */
 	static void RefreshVariables(UBlueprint* Blueprint);
 
+	/** Helper function to punch through and honor UAnimGraphNode_Base::PreloadRequiredAssets, which formerly relied on loading assets during compile */
+	static void PreloadBlueprintSpecificData(UBlueprint* Blueprint);
+
 	/**
 	 * Regenerates the class at class load time, and refreshes the blueprint
 	 */
@@ -247,6 +250,12 @@ public:
 	 */
 	static UBlueprint* FindBlueprintForGraphChecked(const UEdGraph* Graph);
 
+	static UClass* GetSkeletonClass(UClass* FromClass);
+	static UClass* GetMostUpToDateClass(UClass* FromClass);
+
+	/** Looks at the most up to data class and returns whether the given property exists in it as well */
+	static bool PropertyStillExists(UProperty* Property);
+
 	/**
 	 * Updates sources of delegates.
 	 */
@@ -1099,8 +1108,32 @@ public:
 	static void PostEditChangeBlueprintActors(UBlueprint* Blueprint, bool bComponentEditChange = false);
 
 	/** Checks if the property can be modified in given blueprint */
+	DEPRECATED(4.17, "Use IsPropertyWritableInBlueprint instead.")
 	static bool IsPropertyReadOnlyInCurrentBlueprint(const UBlueprint* Blueprint, const UProperty* Property);
 
+	/** Enumeration of whether a property is writable or if not, why. */
+	enum class EPropertyWritableState : uint8
+	{
+		Writable,
+		Private,
+		NotBlueprintVisible,
+		BlueprintReadOnly
+	};
+
+	/** Returns an enumeration indicating if the property can be written to by the given Blueprint */
+	static EPropertyWritableState IsPropertyWritableInBlueprint(const UBlueprint* Blueprint, const UProperty* Property);
+
+	/** Enumeration of whether a property is readable or if not, why. */
+	enum class EPropertyReadableState : uint8
+	{
+		Readable,
+		Private,
+		NotBlueprintVisible
+	};
+
+	/** Returns an enumeration indicating if the property can be read by the given Blueprint */
+	static EPropertyReadableState IsPropertyReadableInBlueprint(const UBlueprint* Blueprint, const UProperty* Property);
+
 	/** Ensures that the CDO root component reference is valid for Actor-based Blueprints */
 	static void UpdateRootComponentReference(UBlueprint* Blueprint);
 
diff --git a/Engine/Source/Editor/UnrealEd/Public/Kismet2/CompilerResultsLog.h b/Engine/Source/Editor/UnrealEd/Public/Kismet2/CompilerResultsLog.h
--- a/Engine/Source/Editor/UnrealEd/Public/Kismet2/CompilerResultsLog.h
+++ b/Engine/Source/Editor/UnrealEd/Public/Kismet2/CompilerResultsLog.h
@@ -7,12 +7,14 @@
 #include "UObject/Object.h"
 #include "EdGraph/EdGraphNode.h"
 #include "EdGraph/EdGraphPin.h"
+#include "ObjectKey.h"
 #if WITH_EDITOR
 #include "Logging/TokenizedMessage.h"
 #include "Misc/CompilationResult.h"
 #include "EdGraphToken.h"
 #endif
 
+class FTokenizedMessage;
 class IMessageLogListing;
 
 #if WITH_EDITOR
@@ -79,7 +81,7 @@ class UNREALED_API FCompilerResultsLog
 
 public:
 	// List of all tokenized messages
-	TArray< TSharedRef<class FTokenizedMessage> > Messages;
+	TArray< TSharedRef<FTokenizedMessage> > Messages;
 
 	// Number of error messages
 	int32 NumErrors;
@@ -152,7 +154,10 @@ public:
 		SourcePath = InSourcePath;
 	}
 
-	// Note: @@ will re replaced by FEdGraphToken::Create
+	/**
+	 * Write an error in to the compiler log.
+	 * Note: @@ will be replaced by node or pin links for nodes/pins passed via varargs 
+	 */
 	template<typename... Args>
 	TSharedRef<FTokenizedMessage> Error(const TCHAR* Format, Args... args)
 	{
@@ -162,6 +167,10 @@ public:
 		return Line;
 	}
 
+	/**
+	 * Write a warning in to the compiler log.
+	 * Note: @@ will be replaced by node or pin links for nodes/pins passed via varargs 
+	 */
 	template<typename... Args>
 	TSharedRef<FTokenizedMessage> Warning(const TCHAR* Format, Args... args)
 	{
@@ -171,6 +180,10 @@ public:
 		return Line;
 	}
 
+	/**
+	 * Write a note in to the compiler log.
+	 * Note: @@ will be replaced by node or pin links for nodes/pins passed via varargs 
+	 */
 	template<typename... Args>
 	TSharedRef<FTokenizedMessage> Note(const TCHAR* Format, Args... args)
 	{
@@ -179,6 +192,79 @@ public:
 		return Line;
 	}
 
+	/**
+	 * Store a potential error for a given node in the compiler log. All messages for the node can be committed to the log later by calling CommitPotentialMessages
+	 * Note: @@ will be replaced by node or pin links for nodes/pins passed via varargs 
+	 */
+	template<typename... Args>
+	TSharedRef<FTokenizedMessage> StorePotentialError(const UEdGraphNode* Source, const TCHAR* Format, Args... args)
+	{
+		TSharedRef<FTokenizedMessage> Line = FTokenizedMessage::Create(EMessageSeverity::Error);
+		UEdGraphNode* SourceNode = nullptr;
+		Tokenize(Format, *Line, SourceNode, args...);
+		PotentialMessages.FindOrAdd(Source).Add(Line);
+		return Line;
+	}
+
+	/**
+	 * Store a potential warning for a given node in the compiler log. All messages for the node can be committed to the log later by calling CommitPotentialMessages
+	 * Note: @@ will be replaced by node or pin links for nodes/pins passed via varargs 
+	 */
+	template<typename... Args>
+	TSharedRef<FTokenizedMessage> StorePotentialWarning(const UEdGraphNode* Source, const TCHAR* Format, Args... args)
+	{
+		TSharedRef<FTokenizedMessage> Line = FTokenizedMessage::Create(EMessageSeverity::Warning);
+		UEdGraphNode* SourceNode = nullptr;
+		Tokenize(Format, *Line, SourceNode, args...);
+		PotentialMessages.FindOrAdd(Source).Add(Line);
+		return Line;
+	}
+
+	/**
+	 * Store a potential note for a given node in the compiler log. All messages for the node can be committed to the log later by calling CommitPotentialMessages
+	 * Note: @@ will be replaced by node or pin links for nodes/pins passed via varargs 
+	 */
+	template<typename... Args>
+	TSharedRef<FTokenizedMessage> StorePotentialNote(const UEdGraphNode* Source, const TCHAR* Format, Args... args)
+	{
+		TSharedRef<FTokenizedMessage> Line = FTokenizedMessage::Create(EMessageSeverity::Info);
+		UEdGraphNode* SourceNode = nullptr;
+		Tokenize(Format, *Line, SourceNode, args...);
+		PotentialMessages.FindOrAdd(Source).Add(Line);
+		return Line;
+	}
+
+	/**
+	 * Commit all stored potential messages for a given node. Returns true if any messages were written.
+	 */
+	bool CommitPotentialMessages(UEdGraphNode* Source)
+	{
+		TArray<TSharedRef<FTokenizedMessage>> FoundMessages;
+		if (PotentialMessages.RemoveAndCopyValue(Source, FoundMessages))
+		{
+			for (const TSharedRef<FTokenizedMessage>& Message : FoundMessages)
+			{
+				switch (Message->GetSeverity())
+				{
+				case EMessageSeverity::Error:
+					++NumErrors;
+					break;
+
+				case EMessageSeverity::Warning:
+					++NumWarnings;
+					break;
+
+				default:
+					break;
+				}
+
+				InternalLogMessage(Message, Source);
+			}
+			return true;
+		}
+		return false;
+	}
+
 	/** Update the source backtrack map to note that NewObject was most closely generated/caused by the SourceObject */
 	void NotifyIntermediateObjectCreation(UObject* NewObject, UObject* SourceObject);
 	void NotifyIntermediatePinCreation(UEdGraphPin* NewObject, UEdGraphPin* SourceObject);
@@ -304,8 +390,12 @@ protected:
 	void InternalLogEvent(const FCompilerEvent& InEvent, int32 InDepth = 0);
 
 private:
+
+	/** Map of stored potential messages indexed by a node. Can be committed to the results log by calling CommitPotentialMessages for that node. */
+	TMap< FObjectKey, TArray< TSharedRef<FTokenizedMessage> > > PotentialMessages;
+
 	/** Parses a compiler log dump to generate tokenized output */
-	static TArray< TSharedRef<class FTokenizedMessage> > ParseCompilerLogDump(const FString& LogDump);
+	static TArray< TSharedRef<FTokenizedMessage> > ParseCompilerLogDump(const FString& LogDump);
 
 	/** Goes to an error given a Message Token */
 	static void OnGotoError(const class TSharedRef<IMessageToken>& Token);
diff --git a/Engine/Source/Editor/UnrealEd/Public/Kismet2/EnumEditorUtils.h b/Engine/Source/Editor/UnrealEd/Public/Kismet2/EnumEditorUtils.h
--- a/Engine/Source/Editor/UnrealEd/Public/Kismet2/EnumEditorUtils.h
+++ b/Engine/Source/Editor/UnrealEd/Public/Kismet2/EnumEditorUtils.h
@@ -8,7 +8,7 @@
 
 class UNREALED_API FEnumEditorUtils
 {
-	static void PrepareForChange(const UUserDefinedEnum* Enum);
+	static void PrepareForChange(UUserDefinedEnum* Enum);
 	static void BroadcastChanges(const UUserDefinedEnum* Enum, const TArray<TPair<FName, int64>>& OldNames, bool bResolveData = true);
 
 	/** copy full enumeratos names from given enum to OutEnumNames, the last '_MAX' enumerator is skipped */
diff --git a/Engine/Source/Editor/UnrealEd/Public/Kismet2/Kismet2NameValidators.h b/Engine/Source/Editor/UnrealEd/Public/Kismet2/Kismet2NameValidators.h
--- a/Engine/Source/Editor/UnrealEd/Public/Kismet2/Kismet2NameValidators.h
+++ b/Engine/Source/Editor/UnrealEd/Public/Kismet2/Kismet2NameValidators.h
@@ -8,6 +8,11 @@
 class INameValidatorInterface;
 class UBlueprint;
 
+// Eventually we want this to be the same as INVALID_OBJECTNAME_CHARACTERS, except we might allow spaces
+// but for now it only includes "." as that has known failure modes (FindObject<T> will try to interpret
+// that as a packagename.objectname)
+#define UE_BLUEPRINT_INVALID_NAME_CHARACTERS TEXT(".")
+
 //////////////////////////////////////////////////////////////////////////
 // FNameValidtorFactory
 class UNREALED_API FNameValidatorFactory
@@ -16,7 +21,7 @@ public:
 	static TSharedPtr<class INameValidatorInterface> MakeValidator(class UEdGraphNode* Node);
 };
 
-enum EValidatorResult
+enum class EValidatorResult
 {
 	/** Name is valid for this object */
 	Ok,
@@ -28,6 +33,8 @@ enum EValidatorResult
 	ExistingName,
 	/* The entered name is too long */
 	TooLong,
+	/* The entered name contains invalid characters (see INVALID_OBJECTNAME_CHARACTERS, except for space) */
+	ContainsInvalidCharacters,
 	/** The entered is in use locally */
 	LocallyInUse
 };
@@ -47,7 +54,14 @@ public:
 	/** @return true if FString is valid, false otherwise */
 	virtual EValidatorResult IsValid (const FString& Name, bool bOriginal = false) = 0;
 
-	static FString GetErrorString(const FString& Name, EValidatorResult ErrorCode);
+	/** @return A text string describing the type of error in ErrorCode for Name */
+	static FText GetErrorText(const FString& Name, EValidatorResult ErrorCode);
+
+	/** @return A string describing the type of error in ErrorCode for Name */
+	static FString GetErrorString(const FString& Name, EValidatorResult ErrorCode)
+	{
+		return GetErrorText(Name, ErrorCode).ToString();
+	}
 
 	/** @return Ok if was valid and doesn't require anything, AlreadyInUse if was already in use and is replaced with new one */
 	EValidatorResult FindValidString(FString& InOutName);
diff --git a/Engine/Source/Editor/UnrealEd/Public/Kismet2/StructureEditorUtils.h b/Engine/Source/Editor/UnrealEd/Public/Kismet2/StructureEditorUtils.h
--- a/Engine/Source/Editor/UnrealEd/Public/Kismet2/StructureEditorUtils.h
+++ b/Engine/Source/Editor/UnrealEd/Public/Kismet2/StructureEditorUtils.h
@@ -5,9 +5,11 @@
 #include "Misc/Guid.h"
 #include "Engine/UserDefinedStruct.h"
 #include "Kismet2/ListenerManager.h"
-#include "UserDefinedStructure/UserDefinedStructEditorData.h"
 
+struct FEdGraphPinType;
+struct FStructVariableDescription;
 class UBlueprint;
+class UUserDefinedStruct;
 
 class UNREALED_API FStructureEditorUtils
 {
@@ -125,25 +127,9 @@ public:
 
 	static const TArray<FStructVariableDescription>* GetVarDescPtr(const UUserDefinedStruct* Struct);
 
-	static FStructVariableDescription* GetVarDescByGuid(UUserDefinedStruct* Struct, FGuid VarGuid)
-	{
-		if (Struct)
-		{
-			TArray<FStructVariableDescription>* VarDescArray = GetVarDescPtr(Struct);
-			return VarDescArray ? VarDescArray->FindByPredicate(FFindByGuidHelper<FStructVariableDescription>(VarGuid)) : nullptr;
-		}
-		return nullptr;
-	}
+	static FStructVariableDescription* GetVarDescByGuid(UUserDefinedStruct* Struct, FGuid VarGuid);
 
-	static const FStructVariableDescription* GetVarDescByGuid(const UUserDefinedStruct* Struct, FGuid VarGuid)
-	{
-		if (Struct)
-		{
-			const TArray<FStructVariableDescription>* VarDescArray = GetVarDescPtr(Struct);
-			return VarDescArray ? VarDescArray->FindByPredicate(FFindByGuidHelper<FStructVariableDescription>(VarGuid)) : nullptr;
-		}
-		return nullptr;
-	}
+	static const FStructVariableDescription* GetVarDescByGuid(const UUserDefinedStruct* Struct, FGuid VarGuid);
 
 	static FGuid GetGuidForProperty(const UProperty* Property);
 
diff --git a/Engine/Source/Editor/UnrealEd/UnrealEd.Build.cs b/Engine/Source/Editor/UnrealEd/UnrealEd.Build.cs
--- a/Engine/Source/Editor/UnrealEd/UnrealEd.Build.cs
+++ b/Engine/Source/Editor/UnrealEd/UnrealEd.Build.cs
@@ -197,7 +197,6 @@ public class UnrealEd : ModuleRules
 				"GeometryMode",
 				"TextureAlignMode",
 				"FoliageEdit",
-				"PackageDependencyInfo",
 				"ImageWrapper",
 				"Blutility",
 				"IntroTutorials",
diff --git a/Engine/Source/Programs/UnrealHeaderTool/Private/HeaderParser.cpp b/Engine/Source/Programs/UnrealHeaderTool/Private/HeaderParser.cpp
--- a/Engine/Source/Programs/UnrealHeaderTool/Private/HeaderParser.cpp
+++ b/Engine/Source/Programs/UnrealHeaderTool/Private/HeaderParser.cpp
@@ -6326,16 +6326,10 @@ void FHeaderParser::CompileFunctionDeclaration(FClasses& AllClasses)
 
 		if (!bHasMenuCategory && !bInternalOnly && !bDeprecated && !bBlueprintAccessor) 
 		{ 
-			const bool bModuleIsGame = CurrentlyParsedModule && (
-				CurrentlyParsedModule->ModuleType == EBuildModuleType::GameDeveloper ||
-				CurrentlyParsedModule->ModuleType == EBuildModuleType::GameEditor ||
-				CurrentlyParsedModule->ModuleType == EBuildModuleType::GameRuntime ||
-				CurrentlyParsedModule->ModuleType == EBuildModuleType::GameThirdParty);
-
 			// To allow for quick iteration, don't enforce the requirement that game functions have to be categorized
-			if (!bModuleIsGame)
+			if (bIsCurrentModulePartOfEngine)
 			{
-				UE_LOG_ERROR_UHT(TEXT("Blueprint accessible functions must have a category specified"));
+				UE_LOG_ERROR_UHT(TEXT("An explicit Category specifier is required for Blueprint accessible functions in an Engine module."));
 			}
 		} 
 	}
@@ -6996,17 +6990,14 @@ void FHeaderParser::CompileVariableDeclaration(FClasses& AllClasses, UStruct* St
 	// First check if the category was specified at all and if the property was exposed to the editor.
 	if (!Category && (OriginalProperty.PropertyFlags & (CPF_Edit|CPF_BlueprintVisible)))
 	{
-		static const FString AbsoluteEngineDir = FPaths::ConvertRelativePathToFull(FPaths::EngineDir());
-		FString SourceFilename = GetCurrentSourceFile()->GetFilename();
-		FPaths::NormalizeFilename(SourceFilename);
-		if (Struct->GetOutermost() != nullptr && !SourceFilename.StartsWith(AbsoluteEngineDir))
+		if ((Struct->GetOutermost() != nullptr) && !bIsCurrentModulePartOfEngine)
 		{
 			OriginalProperty.MetaData.Add("Category", Struct->GetFName().ToString());
 			Category = OriginalProperty.MetaData.Find("Category");
 		}
 		else
 		{
-			UE_LOG_ERROR_UHT(TEXT("Property is exposed to the editor or blueprints but has no Category specified."));
+			UE_LOG_ERROR_UHT(TEXT("An explicit Category specifier is required for any property exposed to the editor or Blueprints in an Engine module."));
 		}
 	}
 
@@ -7692,14 +7683,42 @@ void FHeaderParser::ExportNativeHeaders(
 	}
 }
 
-FHeaderParser::FHeaderParser(FFeedbackContext* InWarn, FUHTMakefile& InUHTMakefile)
-: FBaseParser                       ()
-, Warn                              (InWarn)
-, UHTMakefile(InUHTMakefile)
-, bSpottedAutogeneratedHeaderInclude(false)
-, NestLevel							(0)
-, TopNest                           (nullptr)
+FHeaderParser::FHeaderParser(FFeedbackContext* InWarn, FUHTMakefile& InUHTMakefile, const FManifestModule& InModule)
+	: FBaseParser()
+	, Warn(InWarn)
+	, UHTMakefile(InUHTMakefile)
+	, bSpottedAutogeneratedHeaderInclude(false)
+	, NestLevel(0)
+	, TopNest(nullptr)
+	, CurrentlyParsedModule(&InModule)
 {
+	// Determine if the current module is part of the engine or a game (we are more strict about things for Engine modules)
+	switch (InModule.ModuleType)
+	{
+	case EBuildModuleType::Program:
+		{
+			const FString AbsoluteEngineDir = FPaths::ConvertRelativePathToFull(FPaths::EngineDir());
+			const FString ModuleDir = FPaths::ConvertRelativePathToFull(InModule.BaseDirectory);
+			bIsCurrentModulePartOfEngine = ModuleDir.StartsWith(AbsoluteEngineDir);
+		}
+		break;
+	case EBuildModuleType::EngineRuntime:
+	case EBuildModuleType::EngineDeveloper:
+	case EBuildModuleType::EngineEditor:
+	case EBuildModuleType::EngineThirdParty:
+		bIsCurrentModulePartOfEngine = true;
+		break;
+	case EBuildModuleType::GameRuntime:
+	case EBuildModuleType::GameDeveloper:
+	case EBuildModuleType::GameEditor:
+	case EBuildModuleType::GameThirdParty:
+		bIsCurrentModulePartOfEngine = false;
+		break;
+	default:
+		bIsCurrentModulePartOfEngine = true;
+		check(false);
+	}
+
 	FScriptLocation::Compiler = this;
 
 	// This should be moved to some sort of config
@@ -7800,8 +7819,7 @@ ECompilationResult::Type FHeaderParser::ParseAllHeadersInside(
 	UHTMakefile.SetCurrentModuleName(FName(*Module.Name));
 	// Create the header parser and register it as the warning context.
 	// Note: This must be declared outside the try block, since the catch block will log into it.
-	FHeaderParser HeaderParser(Warn, UHTMakefile);
-	HeaderParser.CurrentlyParsedModule = &Module;
+	FHeaderParser HeaderParser(Warn, UHTMakefile, Module);
 	Warn->SetContext(&HeaderParser);
 
 
diff --git a/Engine/Source/Programs/UnrealHeaderTool/Private/HeaderParser.h b/Engine/Source/Programs/UnrealHeaderTool/Private/HeaderParser.h
--- a/Engine/Source/Programs/UnrealHeaderTool/Private/HeaderParser.h
+++ b/Engine/Source/Programs/UnrealHeaderTool/Private/HeaderParser.h
@@ -410,7 +410,7 @@ protected:
 
 protected:
 	// Constructor.
-	FHeaderParser(FFeedbackContext* InWarn, FUHTMakefile& InUHTMakefile);
+	FHeaderParser(FFeedbackContext* InWarn, FUHTMakefile& InUHTMakefile, const FManifestModule& InModule);
 
 	virtual ~FHeaderParser()
 	{
@@ -713,6 +713,9 @@ private:
 	// Module currently parsed by UHT.
 	const FManifestModule* CurrentlyParsedModule;
 
+	// True if the module currently being parsed is part of the engine, as opposed to being part of a game
+	bool bIsCurrentModulePartOfEngine;
+
 	/**
 	 * Tries to match constructor parameter list. Assumes that constructor
 	 * name is already matched.
diff --git a/Engine/Source/Runtime/AIModule/Classes/Blueprint/AIBlueprintHelperLibrary.h b/Engine/Source/Runtime/AIModule/Classes/Blueprint/AIBlueprintHelperLibrary.h
--- a/Engine/Source/Runtime/AIModule/Classes/Blueprint/AIBlueprintHelperLibrary.h
+++ b/Engine/Source/Runtime/AIModule/Classes/Blueprint/AIBlueprintHelperLibrary.h
@@ -58,4 +58,11 @@ class AIMODULE_API UAIBlueprintHelperLibrary : public UBlueprintFunctionLibrary
 
 	UFUNCTION(BlueprintPure, Category = "AI")
 	static bool IsValidAIRotation(FRotator Rotation);
+
+	/** Returns a copy of navigation path given controller is currently using. 
+	 *	The result being a copy means you won't be able to influence agent's pathfollowing 
+	 *	by manipulating received path */
+	UFUNCTION(BlueprintPure, Category = "AI", meta = (UnsafeDuringActorConstruction = "true"))
+	static UNavigationPath* GetCurrentPath(AController* Controller);
+
 };
diff --git a/Engine/Source/Runtime/AIModule/Classes/Navigation/CrowdManager.h b/Engine/Source/Runtime/AIModule/Classes/Navigation/CrowdManager.h
--- a/Engine/Source/Runtime/AIModule/Classes/Navigation/CrowdManager.h
+++ b/Engine/Source/Runtime/AIModule/Classes/Navigation/CrowdManager.h
@@ -2,6 +2,7 @@
 
 #pragma once
 
+#include "AI/Navigation/CrowdManagerBase.h"
 #include "CoreMinimal.h"
 #include "Stats/Stats.h"
 #include "UObject/ObjectMacros.h"
@@ -20,7 +21,6 @@ class UCrowdManager;
 struct dtCrowdAgent;
 struct dtCrowdAgentDebugInfo;
 struct dtCrowdAgentParams;
-struct dtQuerySpecialLinkFilter;
 struct FNavMeshPath;
 
 #if WITH_RECAST
@@ -164,11 +164,14 @@ struct FCrowdTickHelper : FTickableGameObject
 };
 
 UCLASS(config = Engine, defaultconfig)
-class AIMODULE_API UCrowdManager : public UObject
+class AIMODULE_API UCrowdManager : public UCrowdManagerBase
 {
-	GENERATED_UCLASS_BODY()
+	GENERATED_BODY()
 
-	virtual void Tick(float DeltaTime);
+public:
+	UCrowdManager(const FObjectInitializer& ObjectInitializer = FObjectInitializer::Get());
+
+	virtual void Tick(float DeltaTime) override;
 	virtual void BeginDestroy() override;
 
 	/** adds new agent to crowd */
diff --git a/Engine/Source/Runtime/AIModule/Private/Actions/PawnAction.cpp b/Engine/Source/Runtime/AIModule/Private/Actions/PawnAction.cpp
--- a/Engine/Source/Runtime/AIModule/Private/Actions/PawnAction.cpp
+++ b/Engine/Source/Runtime/AIModule/Private/Actions/PawnAction.cpp
@@ -389,7 +389,7 @@ void UPawnAction::WaitForMessage(FName MessageType, FAIRequestID InRequestID)
 
 UPawnAction* UPawnAction::CreateActionInstance(UObject* WorldContextObject, TSubclassOf<UPawnAction> ActionClass)
 {
-	UWorld* World = GEngine->GetWorldFromContextObject(WorldContextObject);
+	UWorld* World = GEngine->GetWorldFromContextObject(WorldContextObject, EGetWorldErrorMode::LogAndReturnNull);
 	if (World && ActionClass)
 	{
 		return NewObject<UPawnAction>(World, ActionClass);
diff --git a/Engine/Source/Runtime/AIModule/Private/BehaviorTree/BTFunctionLibrary.cpp b/Engine/Source/Runtime/AIModule/Private/BehaviorTree/BTFunctionLibrary.cpp
--- a/Engine/Source/Runtime/AIModule/Private/BehaviorTree/BTFunctionLibrary.cpp
+++ b/Engine/Source/Runtime/AIModule/Private/BehaviorTree/BTFunctionLibrary.cpp
@@ -16,6 +16,7 @@
 #include "BehaviorTree/Blackboard/BlackboardKeyType_Int.h"
 #include "BehaviorTree/Blackboard/BlackboardKeyType_Name.h"
 #include "BehaviorTree/Blackboard/BlackboardKeyType_String.h"
+#include "Misc/RuntimeErrors.h"
 
 //----------------------------------------------------------------------//
 // 
@@ -39,118 +40,95 @@ namespace FBTNodeBPImplementationHelper
 //----------------------------------------------------------------------//
 // 
 //----------------------------------------------------------------------//
-namespace
-{
-	FORCEINLINE UBlackboardComponent* GetBlackboard(UBTNode& BTNode)
-	{
-		check(BTNode.GetOuter());
-		checkf(BTNode.GetClass()->HasAnyClassFlags(CLASS_CompiledFromBlueprint) || Cast<UDynamicClass>(BTNode.GetClass()), TEXT("This function call is valid only for BP-implemented BT nodes"));
-		check(Cast<UBehaviorTreeComponent>(BTNode.GetOuter()));
-
-		return ((UBehaviorTreeComponent*)BTNode.GetOuter())->GetBlackboardComponent();
-	}
-}
 
 UBehaviorTreeComponent* UBTFunctionLibrary::GetOwnerComponent(UBTNode* NodeOwner)
 {
-	check(NodeOwner);
-	checkf(NodeOwner->GetClass()->HasAnyClassFlags(CLASS_CompiledFromBlueprint), TEXT("This function call is valid only for BP-implemented BT nodes"));
-	check(NodeOwner->GetOuter());
+	ensureAsRuntimeWarning((NodeOwner != nullptr) && (NodeOwner->GetClass()->HasAnyClassFlags(CLASS_CompiledFromBlueprint) || Cast<UDynamicClass>(NodeOwner->GetClass())));
 
-	UBehaviorTreeComponent* OwnerComp = Cast<UBehaviorTreeComponent>(NodeOwner->GetOuter());
-	check(OwnerComp);
+	UBehaviorTreeComponent* OwnerComp = NodeOwner ? Cast<UBehaviorTreeComponent>(NodeOwner->GetOuter()) : nullptr;
+	ensureAsRuntimeWarning(OwnerComp != nullptr);
 
 	return OwnerComp;
 }
 
 UBlackboardComponent* UBTFunctionLibrary::GetOwnersBlackboard(UBTNode* NodeOwner)
 {
-	check(NodeOwner != NULL);
-	return GetBlackboard(*NodeOwner);
+	UBehaviorTreeComponent* BTComponent = GetOwnerComponent(NodeOwner);
+	if (ensureAsRuntimeWarning(BTComponent != nullptr))
+	{
+		return BTComponent->GetBlackboardComponent();
+	}
+	return nullptr;
 }
 
 UObject* UBTFunctionLibrary::GetBlackboardValueAsObject(UBTNode* NodeOwner, const FBlackboardKeySelector& Key)
 {
-	check(NodeOwner != NULL);
-	const UBlackboardComponent* BlackboardComp = GetBlackboard(*NodeOwner);
-	return BlackboardComp ? BlackboardComp->GetValueAsObject(Key.SelectedKeyName) : NULL;
+	UBlackboardComponent* BlackboardComp = GetOwnersBlackboard(NodeOwner);
+	return BlackboardComp ? BlackboardComp->GetValueAsObject(Key.SelectedKeyName) : nullptr;
 }
 
 AActor* UBTFunctionLibrary::GetBlackboardValueAsActor(UBTNode* NodeOwner, const FBlackboardKeySelector& Key)
 {
-	check(NodeOwner != NULL);
-	const UBlackboardComponent* BlackboardComp = GetBlackboard(*NodeOwner);
-	return BlackboardComp ? Cast<AActor>(BlackboardComp->GetValueAsObject(Key.SelectedKeyName)) : NULL;
+	return Cast<AActor>(GetBlackboardValueAsObject(NodeOwner, Key));
 }
 
 UClass* UBTFunctionLibrary::GetBlackboardValueAsClass(UBTNode* NodeOwner, const FBlackboardKeySelector& Key)
 {
-	check(NodeOwner != NULL);
-	const UBlackboardComponent* BlackboardComp = GetBlackboard(*NodeOwner);
-	return BlackboardComp ? BlackboardComp->GetValueAsClass(Key.SelectedKeyName) : NULL;
+	UBlackboardComponent* BlackboardComp = GetOwnersBlackboard(NodeOwner);
+	return BlackboardComp ? BlackboardComp->GetValueAsClass(Key.SelectedKeyName) : nullptr;
 }
 
 uint8 UBTFunctionLibrary::GetBlackboardValueAsEnum(UBTNode* NodeOwner, const FBlackboardKeySelector& Key)
 {
-	check(NodeOwner != NULL);
-	const UBlackboardComponent* BlackboardComp = GetBlackboard(*NodeOwner);
+	UBlackboardComponent* BlackboardComp = GetOwnersBlackboard(NodeOwner);
 	return BlackboardComp ? BlackboardComp->GetValueAsEnum(Key.SelectedKeyName) : 0;
 }
 
 int32 UBTFunctionLibrary::GetBlackboardValueAsInt(UBTNode* NodeOwner, const FBlackboardKeySelector& Key)
 {
-	check(NodeOwner != NULL);
-	const UBlackboardComponent* BlackboardComp = GetBlackboard(*NodeOwner);
+	UBlackboardComponent* BlackboardComp = GetOwnersBlackboard(NodeOwner);
 	return BlackboardComp ? BlackboardComp->GetValueAsInt(Key.SelectedKeyName) : 0;
 }
 
 float UBTFunctionLibrary::GetBlackboardValueAsFloat(UBTNode* NodeOwner, const FBlackboardKeySelector& Key)
 {
-	check(NodeOwner != NULL);
-	const UBlackboardComponent* BlackboardComp = GetBlackboard(*NodeOwner);
+	UBlackboardComponent* BlackboardComp = GetOwnersBlackboard(NodeOwner);
 	return BlackboardComp ? BlackboardComp->GetValueAsFloat(Key.SelectedKeyName) : 0.0f;
 }
 
 bool UBTFunctionLibrary::GetBlackboardValueAsBool(UBTNode* NodeOwner, const FBlackboardKeySelector& Key)
 {
-	check(NodeOwner != NULL);
-	const UBlackboardComponent* BlackboardComp = GetBlackboard(*NodeOwner);
+	UBlackboardComponent* BlackboardComp = GetOwnersBlackboard(NodeOwner);
 	return BlackboardComp ? BlackboardComp->GetValueAsBool(Key.SelectedKeyName) : false;
 }
 
 FString UBTFunctionLibrary::GetBlackboardValueAsString(UBTNode* NodeOwner, const FBlackboardKeySelector& Key)
 {
-	check(NodeOwner != NULL);
-	const UBlackboardComponent* BlackboardComp = GetBlackboard(*NodeOwner);
+	UBlackboardComponent* BlackboardComp = GetOwnersBlackboard(NodeOwner);
 	return BlackboardComp ? BlackboardComp->GetValueAsString(Key.SelectedKeyName) : FString();
 }
 
 FName UBTFunctionLibrary::GetBlackboardValueAsName(UBTNode* NodeOwner, const FBlackboardKeySelector& Key)
 {
-	check(NodeOwner != NULL);
-	const UBlackboardComponent* BlackboardComp = GetBlackboard(*NodeOwner);
+	UBlackboardComponent* BlackboardComp = GetOwnersBlackboard(NodeOwner);
 	return BlackboardComp ? BlackboardComp->GetValueAsName(Key.SelectedKeyName) : NAME_None;
 }
 
 FVector UBTFunctionLibrary::GetBlackboardValueAsVector(UBTNode* NodeOwner, const FBlackboardKeySelector& Key)
 {
-	check(NodeOwner != NULL);
-	const UBlackboardComponent* BlackboardComp = GetBlackboard(*NodeOwner);
+	UBlackboardComponent* BlackboardComp = GetOwnersBlackboard(NodeOwner);
 	return BlackboardComp ? BlackboardComp->GetValueAsVector(Key.SelectedKeyName) : FVector::ZeroVector;
 }
 
 FRotator UBTFunctionLibrary::GetBlackboardValueAsRotator(UBTNode* NodeOwner, const FBlackboardKeySelector& Key)
 {
-	check(NodeOwner != NULL);
-	const UBlackboardComponent* BlackboardComp = GetBlackboard(*NodeOwner);
+	UBlackboardComponent* BlackboardComp = GetOwnersBlackboard(NodeOwner);
 	return BlackboardComp ? BlackboardComp->GetValueAsRotator(Key.SelectedKeyName) : FRotator::ZeroRotator;
 }
 
 void UBTFunctionLibrary::SetBlackboardValueAsObject(UBTNode* NodeOwner, const FBlackboardKeySelector& Key, UObject* Value)
 {
-	check(NodeOwner != NULL);
-	UBlackboardComponent* BlackboardComp = GetBlackboard(*NodeOwner);
-	if (BlackboardComp != NULL)
+	if (UBlackboardComponent* BlackboardComp = GetOwnersBlackboard(NodeOwner))
 	{
 		BlackboardComp->SetValue<UBlackboardKeyType_Object>(Key.SelectedKeyName, Value);
 	}
@@ -158,9 +136,7 @@ void UBTFunctionLibrary::SetBlackboardValueAsObject(UBTNode* NodeOwner, const FB
 
 void UBTFunctionLibrary::SetBlackboardValueAsClass(UBTNode* NodeOwner, const FBlackboardKeySelector& Key, UClass* Value)
 {
-	check(NodeOwner != NULL);
-	UBlackboardComponent* BlackboardComp = GetBlackboard(*NodeOwner);
-	if (BlackboardComp != NULL)
+	if (UBlackboardComponent* BlackboardComp = GetOwnersBlackboard(NodeOwner))
 	{
 		BlackboardComp->SetValue<UBlackboardKeyType_Class>(Key.SelectedKeyName, Value);
 	}
@@ -168,9 +144,7 @@ void UBTFunctionLibrary::SetBlackboardValueAsClass(UBTNode* NodeOwner, const FBl
 
 void UBTFunctionLibrary::SetBlackboardValueAsEnum(UBTNode* NodeOwner, const FBlackboardKeySelector& Key, uint8 Value)
 {
-	check(NodeOwner != NULL);
-	UBlackboardComponent* BlackboardComp = GetBlackboard(*NodeOwner);
-	if (BlackboardComp != NULL)
+	if (UBlackboardComponent* BlackboardComp = GetOwnersBlackboard(NodeOwner))
 	{
 		BlackboardComp->SetValue<UBlackboardKeyType_Enum>(Key.SelectedKeyName, Value);
 	}
@@ -178,9 +152,7 @@ void UBTFunctionLibrary::SetBlackboardValueAsEnum(UBTNode* NodeOwner, const FBla
 
 void UBTFunctionLibrary::SetBlackboardValueAsInt(UBTNode* NodeOwner, const FBlackboardKeySelector& Key, int32 Value)
 {
-	check(NodeOwner != NULL);
-	UBlackboardComponent* BlackboardComp = GetBlackboard(*NodeOwner);
-	if (BlackboardComp != NULL)
+	if (UBlackboardComponent* BlackboardComp = GetOwnersBlackboard(NodeOwner))
 	{
 		BlackboardComp->SetValue<UBlackboardKeyType_Int>(Key.SelectedKeyName, Value);
 	}
@@ -188,9 +160,7 @@ void UBTFunctionLibrary::SetBlackboardValueAsInt(UBTNode* NodeOwner, const FBlac
 
 void UBTFunctionLibrary::SetBlackboardValueAsFloat(UBTNode* NodeOwner, const FBlackboardKeySelector& Key, float Value)
 {
-	check(NodeOwner != NULL);
-	UBlackboardComponent* BlackboardComp = GetBlackboard(*NodeOwner);
-	if (BlackboardComp != NULL)
+	if (UBlackboardComponent* BlackboardComp = GetOwnersBlackboard(NodeOwner))
 	{
 		BlackboardComp->SetValue<UBlackboardKeyType_Float>(Key.SelectedKeyName, Value);
 	}
@@ -198,9 +168,7 @@ void UBTFunctionLibrary::SetBlackboardValueAsFloat(UBTNode* NodeOwner, const FBl
 
 void UBTFunctionLibrary::SetBlackboardValueAsBool(UBTNode* NodeOwner, const FBlackboardKeySelector& Key, bool Value)
 {
-	check(NodeOwner != NULL);
-	UBlackboardComponent* BlackboardComp = GetBlackboard(*NodeOwner);
-	if (BlackboardComp != NULL)
+	if (UBlackboardComponent* BlackboardComp = GetOwnersBlackboard(NodeOwner))
 	{
 		BlackboardComp->SetValue<UBlackboardKeyType_Bool>(Key.SelectedKeyName, Value);
 	}
@@ -208,9 +176,7 @@ void UBTFunctionLibrary::SetBlackboardValueAsBool(UBTNode* NodeOwner, const FBla
 
 void UBTFunctionLibrary::SetBlackboardValueAsString(UBTNode* NodeOwner, const FBlackboardKeySelector& Key, FString Value)
 {
-	check(NodeOwner != NULL);
-	UBlackboardComponent* BlackboardComp = GetBlackboard(*NodeOwner);
-	if (BlackboardComp != NULL)
+	if (UBlackboardComponent* BlackboardComp = GetOwnersBlackboard(NodeOwner))
 	{
 		BlackboardComp->SetValue<UBlackboardKeyType_String>(Key.SelectedKeyName, Value);
 	}
@@ -218,9 +184,7 @@ void UBTFunctionLibrary::SetBlackboardValueAsString(UBTNode* NodeOwner, const FB
 
 void UBTFunctionLibrary::SetBlackboardValueAsName(UBTNode* NodeOwner, const FBlackboardKeySelector& Key, FName Value)
 {
-	check(NodeOwner != NULL);
-	UBlackboardComponent* BlackboardComp = GetBlackboard(*NodeOwner);
-	if (BlackboardComp != NULL)
+	if (UBlackboardComponent* BlackboardComp = GetOwnersBlackboard(NodeOwner))
 	{
 		BlackboardComp->SetValue<UBlackboardKeyType_Name>(Key.SelectedKeyName, Value);
 	}
@@ -228,9 +192,7 @@ void UBTFunctionLibrary::SetBlackboardValueAsName(UBTNode* NodeOwner, const FBla
 
 void UBTFunctionLibrary::SetBlackboardValueAsVector(UBTNode* NodeOwner, const FBlackboardKeySelector& Key, FVector Value)
 {
-	check(NodeOwner != NULL);
-	UBlackboardComponent* BlackboardComp = GetBlackboard(*NodeOwner);
-	if (BlackboardComp != NULL)
+	if (UBlackboardComponent* BlackboardComp = GetOwnersBlackboard(NodeOwner))
 	{
 		BlackboardComp->SetValue<UBlackboardKeyType_Vector>(Key.SelectedKeyName, Value);
 	}
@@ -238,9 +200,7 @@ void UBTFunctionLibrary::SetBlackboardValueAsVector(UBTNode* NodeOwner, const FB
 
 void UBTFunctionLibrary::ClearBlackboardValueAsVector(UBTNode* NodeOwner, const FBlackboardKeySelector& Key)
 {
-	check(NodeOwner != NULL);
-	UBlackboardComponent* BlackboardComp = GetBlackboard(*NodeOwner);
-	if (BlackboardComp != NULL)
+	if (UBlackboardComponent* BlackboardComp = GetOwnersBlackboard(NodeOwner))
 	{
 		BlackboardComp->ClearValue(Key.SelectedKeyName);
 	}
@@ -248,9 +208,7 @@ void UBTFunctionLibrary::ClearBlackboardValueAsVector(UBTNode* NodeOwner, const
 
 void UBTFunctionLibrary::SetBlackboardValueAsRotator(UBTNode* NodeOwner, const FBlackboardKeySelector& Key, FRotator Value)
 {
-	check(NodeOwner != NULL);
-	UBlackboardComponent* BlackboardComp = GetBlackboard(*NodeOwner);
-	if (BlackboardComp != NULL)
+	if (UBlackboardComponent* BlackboardComp = GetOwnersBlackboard(NodeOwner))
 	{
 		BlackboardComp->SetValue<UBlackboardKeyType_Rotator>(Key.SelectedKeyName, Value);
 	}
@@ -258,9 +216,7 @@ void UBTFunctionLibrary::SetBlackboardValueAsRotator(UBTNode* NodeOwner, const F
 
 void UBTFunctionLibrary::ClearBlackboardValue(UBTNode* NodeOwner, const FBlackboardKeySelector& Key)
 {
-	check(NodeOwner != NULL);
-	UBlackboardComponent* BlackboardComp = GetBlackboard(*NodeOwner);
-	if (BlackboardComp != NULL)
+	if (UBlackboardComponent* BlackboardComp = GetOwnersBlackboard(NodeOwner))
 	{
 		BlackboardComp->ClearValue(Key.SelectedKeyName);
 	}
diff --git a/Engine/Source/Runtime/AIModule/Private/BehaviorTree/BehaviorTreeManager.cpp b/Engine/Source/Runtime/AIModule/Private/BehaviorTree/BehaviorTreeManager.cpp
--- a/Engine/Source/Runtime/AIModule/Private/BehaviorTree/BehaviorTreeManager.cpp
+++ b/Engine/Source/Runtime/AIModule/Private/BehaviorTree/BehaviorTreeManager.cpp
@@ -508,12 +508,12 @@ void UBehaviorTreeManager::RemoveActiveComponent(UBehaviorTreeComponent& Compone
 UBehaviorTreeManager* UBehaviorTreeManager::GetCurrent(UWorld* World)
 {
 	UAISystem* AISys = UAISystem::GetCurrentSafe(World);
-	return AISys ? AISys->GetBehaviorTreeManager() : NULL;
+	return AISys ? AISys->GetBehaviorTreeManager() : nullptr;
 }
 
 UBehaviorTreeManager* UBehaviorTreeManager::GetCurrent(UObject* WorldContextObject)
 {
-	UWorld* World = GEngine->GetWorldFromContextObject(WorldContextObject, false);
+	UWorld* World = GEngine->GetWorldFromContextObject(WorldContextObject, EGetWorldErrorMode::ReturnNull);
 	UAISystem* AISys = UAISystem::GetCurrentSafe(World);
-	return AISys ? AISys->GetBehaviorTreeManager() : NULL;
+	return AISys ? AISys->GetBehaviorTreeManager() : nullptr;
 }
diff --git a/Engine/Source/Runtime/AIModule/Private/Blueprint/AIBlueprintHelperLibrary.cpp b/Engine/Source/Runtime/AIModule/Private/Blueprint/AIBlueprintHelperLibrary.cpp
--- a/Engine/Source/Runtime/AIModule/Private/Blueprint/AIBlueprintHelperLibrary.cpp
+++ b/Engine/Source/Runtime/AIModule/Private/Blueprint/AIBlueprintHelperLibrary.cpp
@@ -13,6 +13,7 @@
 #include "BehaviorTree/BlackboardComponent.h"
 #include "Blueprint/AIAsyncTaskBlueprintProxy.h"
 #include "Animation/AnimInstance.h"
+#include "AI/Navigation/NavigationPath.h"
 
 DEFINE_LOG_CATEGORY_STATIC(LogAIBlueprint, Warning, All);
 
@@ -131,7 +132,7 @@ UAIAsyncTaskBlueprintProxy* UAIBlueprintHelperLibrary::CreateMoveToProxyObject(U
 	AAIController* AIController = Cast<AAIController>(Pawn->GetController());
 	if (AIController)
 	{
-		UWorld* World = GEngine->GetWorldFromContextObject(WorldContextObject);
+		UWorld* World = GEngine->GetWorldFromContextObjectChecked(WorldContextObject);
 		MyObj = NewObject<UAIAsyncTaskBlueprintProxy>(World);
 
 		FAIMoveRequest MoveReq;
@@ -179,7 +180,7 @@ APawn* UAIBlueprintHelperLibrary::SpawnAIFromClass(UObject* WorldContextObject,
 {
 	APawn* NewPawn = NULL;
 
-	UWorld* World = GEngine->GetWorldFromContextObject(WorldContextObject);
+	UWorld* World = GEngine->GetWorldFromContextObject(WorldContextObject, EGetWorldErrorMode::LogAndReturnNull);
 	if (World && *PawnClass)
 	{
 		FActorSpawnParameters ActorSpawnParams;
@@ -304,3 +305,36 @@ bool UAIBlueprintHelperLibrary::IsValidAIRotation(FRotator Rotation)
 {
 	return FAISystem::IsValidRotation(Rotation);
 }
+
+UNavigationPath* UAIBlueprintHelperLibrary::GetCurrentPath(AController* Controller)
+{
+	UNavigationPath* ResultPath = nullptr;
+	if (Controller)
+	{
+		AAIController* AIController = Cast<AAIController>(Controller);
+		UPathFollowingComponent* PFComp = nullptr;
+		if (AIController)
+		{
+			PFComp = AIController->GetPathFollowingComponent();
+		}
+		else
+		{
+			// player controller, most probably using SimpleMove
+			PFComp = Controller->FindComponentByClass<UPathFollowingComponent>();
+		}
+
+		if (PFComp != nullptr)
+		{
+			const FNavPathSharedPtr Path = PFComp->GetPath();
+			if (Path.IsValid())
+			{
+				// we don't care if Path->IsValid(), we're going to retrieve whatever's there
+				ResultPath = NewObject<UNavigationPath>();
+				ResultPath->SetPath(Path);
+			}
+		}
+	}
+
+	return ResultPath;
+}
+
diff --git a/Engine/Source/Runtime/AIModule/Private/BrainComponent.cpp b/Engine/Source/Runtime/AIModule/Private/BrainComponent.cpp
--- a/Engine/Source/Runtime/AIModule/Private/BrainComponent.cpp
+++ b/Engine/Source/Runtime/AIModule/Private/BrainComponent.cpp
@@ -67,7 +67,7 @@ void FAIMessage::Send(UBrainComponent* BrainComp, const FAIMessage& Message)
 
 void FAIMessage::Broadcast(UObject* WorldContextObject, const FAIMessage& Message)
 {
-	UWorld* MyWorld = GEngine->GetWorldFromContextObject(WorldContextObject);
+	UWorld* MyWorld = GEngine->GetWorldFromContextObject(WorldContextObject, EGetWorldErrorMode::LogAndReturnNull);
 	if (MyWorld)
 	{
 		for (FConstControllerIterator It = MyWorld->GetControllerIterator(); It; ++It)
diff --git a/Engine/Source/Runtime/AIModule/Private/EnvironmentQuery/EnvQueryManager.cpp b/Engine/Source/Runtime/AIModule/Private/EnvironmentQuery/EnvQueryManager.cpp
--- a/Engine/Source/Runtime/AIModule/Private/EnvironmentQuery/EnvQueryManager.cpp
+++ b/Engine/Source/Runtime/AIModule/Private/EnvironmentQuery/EnvQueryManager.cpp
@@ -81,7 +81,7 @@ int32 FEnvQueryRequest::Execute(EEnvQueryRunMode::Type RunMode, FQueryFinishedSi
 
 	if (World == NULL)
 	{
-		World = GEngine->GetWorldFromContextObject(Owner);
+		World = GEngine->GetWorldFromContextObject(Owner, EGetWorldErrorMode::ReturnNull);
 		if (World == NULL)
 		{
 			UE_LOG(LogEQS, Warning, TEXT("Unable to access world with owner: %s"), *GetNameSafe(Owner));
@@ -140,14 +140,14 @@ void UEnvQueryManager::FinishDestroy()
 UEnvQueryManager* UEnvQueryManager::GetCurrent(UWorld* World)
 {
 	UAISystem* AISys = UAISystem::GetCurrentSafe(World);
-	return AISys ? AISys->GetEnvironmentQueryManager() : NULL;
+	return AISys ? AISys->GetEnvironmentQueryManager() : nullptr;
 }
 
 UEnvQueryManager* UEnvQueryManager::GetCurrent(const UObject* WorldContextObject)
 {
-	UWorld* World = GEngine->GetWorldFromContextObject(WorldContextObject, false);
+	UWorld* World = GEngine->GetWorldFromContextObject(WorldContextObject, EGetWorldErrorMode::LogAndReturnNull);
 	UAISystem* AISys = UAISystem::GetCurrentSafe(World);
-	return AISys ? AISys->GetEnvironmentQueryManager() : NULL;
+	return AISys ? AISys->GetEnvironmentQueryManager() : nullptr;
 }
 
 #if USE_EQS_DEBUGGER
diff --git a/Engine/Source/Runtime/AIModule/Private/EnvironmentQuery/Generators/EnvQueryGenerator_ActorsOfClass.cpp b/Engine/Source/Runtime/AIModule/Private/EnvironmentQuery/Generators/EnvQueryGenerator_ActorsOfClass.cpp
--- a/Engine/Source/Runtime/AIModule/Private/EnvironmentQuery/Generators/EnvQueryGenerator_ActorsOfClass.cpp
+++ b/Engine/Source/Runtime/AIModule/Private/EnvironmentQuery/Generators/EnvQueryGenerator_ActorsOfClass.cpp
@@ -33,7 +33,7 @@ void UEnvQueryGenerator_ActorsOfClass::GenerateItems(FEnvQueryInstance& QueryIns
 		return;
 	}
 	
-	UWorld* World = GEngine->GetWorldFromContextObject(QueryOwner);
+	UWorld* World = GEngine->GetWorldFromContextObject(QueryOwner, EGetWorldErrorMode::LogAndReturnNull);
 	if (World == nullptr)
 	{
 		return;
diff --git a/Engine/Source/Runtime/AIModule/Private/EnvironmentQuery/Tests/EnvQueryTest_PathfindingBatch.cpp b/Engine/Source/Runtime/AIModule/Private/EnvironmentQuery/Tests/EnvQueryTest_PathfindingBatch.cpp
--- a/Engine/Source/Runtime/AIModule/Private/EnvironmentQuery/Tests/EnvQueryTest_PathfindingBatch.cpp
+++ b/Engine/Source/Runtime/AIModule/Private/EnvironmentQuery/Tests/EnvQueryTest_PathfindingBatch.cpp
@@ -3,10 +3,7 @@
 #include "EnvironmentQuery/Tests/EnvQueryTest_PathfindingBatch.h"
 #include "Engine/World.h"
 #include "AI/Navigation/RecastNavMesh.h"
-
-#if WITH_RECAST
-#include "AI/Navigation/PImplRecastNavMesh.h"
-#endif
+#include "AI/Navigation/RecastQueryFilter.h"
 
 #define LOCTEXT_NAMESPACE "EnvQueryGenerator"
 
diff --git a/Engine/Source/Runtime/AIModule/Private/Navigation/CrowdFollowingComponent.cpp b/Engine/Source/Runtime/AIModule/Private/Navigation/CrowdFollowingComponent.cpp
--- a/Engine/Source/Runtime/AIModule/Private/Navigation/CrowdFollowingComponent.cpp
+++ b/Engine/Source/Runtime/AIModule/Private/Navigation/CrowdFollowingComponent.cpp
@@ -916,9 +916,16 @@ void UCrowdFollowingComponent::UpdatePathSegment()
 	{
 		if (!Path->IsWaitingForRepath())
 		{
+			UE_VLOG(this, LogPathFollowing, Log, TEXT("Aborting move due to path being invalid and not waiting for repath"));
 			OnPathFinished(FPathFollowingResult(EPathFollowingResult::Aborted, FPathFollowingResultFlags::InvalidPath));
+			return;
+		}
+		else
+		{
+			// continue with execution, if navigation is being rebuild constantly AI will get stuck with current waypoint
+			// path updates should be still coming in, even though they get invalidated right away
+			UE_VLOG(this, LogPathFollowing, Log, TEXT("Updating path points in invalid & pending path!"));
 		}
-		return;
 	}
 
 	// if agent has control over its movement, check finish conditions
diff --git a/Engine/Source/Runtime/AIModule/Private/Navigation/CrowdManager.cpp b/Engine/Source/Runtime/AIModule/Private/Navigation/CrowdManager.cpp
--- a/Engine/Source/Runtime/AIModule/Private/Navigation/CrowdManager.cpp
+++ b/Engine/Source/Runtime/AIModule/Private/Navigation/CrowdManager.cpp
@@ -10,10 +10,11 @@
 #include "AIModuleLog.h"
 
 #if WITH_RECAST
-#include "AI/Navigation/PImplRecastNavMesh.h"
 #include "AI/Navigation/RecastHelpers.h"
 #include "DetourCrowd/DetourObstacleAvoidance.h"
 #include "DetourCrowd/DetourCrowd.h"
+#include "Detour/DetourNavMesh.h"
+#include "AI/Navigation/RecastQueryFilter.h"
 #endif
 
 #include "Navigation/CrowdFollowingComponent.h"
@@ -806,8 +807,7 @@ void UCrowdManager::UpdateAgentPaths()
 {
 	UNavigationSystem* NavSys = Cast<UNavigationSystem>(GetOuter());
 	ARecastNavMesh* RecastNavData = Cast<ARecastNavMesh>(MyNavData);
-	FPImplRecastNavMesh* PImplNavMesh = RecastNavData ? RecastNavData->RecastNavMeshImpl : NULL;
-	if (PImplNavMesh == NULL)
+	if (RecastNavData == NULL)
 	{
 		return;
 	}
@@ -831,7 +831,7 @@ void UCrowdManager::UpdateAgentPaths()
 
 				if (AnimInfo.t == 0)
 				{
-					const uint32 NavLinkId = PImplNavMesh->GetLinkUserId(AnimInfo.polyRef);
+					const uint32 NavLinkId = RecastNavData->GetLinkUserId(AnimInfo.polyRef);
 					INavLinkCustomInterface* CustomLink = NavSys->GetCustomLink(NavLinkId);
 
 					if (CustomLink)
@@ -883,9 +883,18 @@ void UCrowdManager::UpdateSelectedDebug(const ICrowdAgentInterface* Agent, int32
 
 void UCrowdManager::CreateCrowdManager()
 {
+	if (MyNavData == nullptr)
+	{
+		// run update and quit since UpdateNavData will call CreateCrowdManager
+		// if navigation mesh is found
+		UpdateNavData();
+
+		UE_CLOG(MyNavData == nullptr, LogCrowdFollowing, Error, TEXT("Unable to find RecastNavMesh instance while trying to create UCrowdManager instance"));
+		return;
+	}
+
 	ARecastNavMesh* RecastNavData = Cast<ARecastNavMesh>(MyNavData);
-	FPImplRecastNavMesh* PImplNavMesh = RecastNavData ? RecastNavData->RecastNavMeshImpl : NULL;
-	dtNavMesh* NavMeshPtr = PImplNavMesh ? PImplNavMesh->GetRecastMesh() : NULL;
+	dtNavMesh* NavMeshPtr = RecastNavData->GetRecastMesh();
 
 	if (NavMeshPtr)
 	{
@@ -1370,11 +1379,11 @@ UWorld* UCrowdManager::GetWorld() const
 UCrowdManager* UCrowdManager::GetCurrent(UObject* WorldContextObject)
 {
 	UNavigationSystem* NavSys = UNavigationSystem::GetCurrent(WorldContextObject);
-	return NavSys ? NavSys->GetCrowdManager() : NULL;
+	return NavSys ? Cast<UCrowdManager>(NavSys->GetCrowdManager()) : NULL;
 }
 
 UCrowdManager* UCrowdManager::GetCurrent(UWorld* World)
 {
 	UNavigationSystem* NavSys = UNavigationSystem::GetCurrent(World);
-	return NavSys ? NavSys->GetCrowdManager() : NULL;
+	return NavSys ? Cast<UCrowdManager>(NavSys->GetCrowdManager()) : NULL;
 }
diff --git a/Engine/Source/Runtime/AIModule/Private/Navigation/NavLocalGridManager.cpp b/Engine/Source/Runtime/AIModule/Private/Navigation/NavLocalGridManager.cpp
--- a/Engine/Source/Runtime/AIModule/Private/Navigation/NavLocalGridManager.cpp
+++ b/Engine/Source/Runtime/AIModule/Private/Navigation/NavLocalGridManager.cpp
@@ -430,7 +430,7 @@ UNavLocalGridManager* UNavLocalGridManager::GetCurrent(UWorld* World)
 
 UNavLocalGridManager* UNavLocalGridManager::GetCurrent(const UObject* WorldContextObject)
 {
-	UWorld* World = GEngine->GetWorldFromContextObject(WorldContextObject);
+	UWorld* World = GEngine->GetWorldFromContextObject(WorldContextObject, EGetWorldErrorMode::LogAndReturnNull);
 	UAISystem* AISys = World ? UAISystem::GetCurrent(*World) : nullptr;
 	return AISys ? AISys->GetNavLocalGridManager() : nullptr;
 }
diff --git a/Engine/Source/Runtime/AIModule/Private/Navigation/PathFollowingComponent.cpp b/Engine/Source/Runtime/AIModule/Private/Navigation/PathFollowingComponent.cpp
--- a/Engine/Source/Runtime/AIModule/Private/Navigation/PathFollowingComponent.cpp
+++ b/Engine/Source/Runtime/AIModule/Private/Navigation/PathFollowingComponent.cpp
@@ -877,11 +877,16 @@ void UPathFollowingComponent::UpdatePathSegment()
 	{
 		if (!Path->IsWaitingForRepath())
 		{
-			UE_VLOG(this, LogPathFollowing, Log, TEXT("Aborting move due to path being invelid and not waiting for repath"));
+			UE_VLOG(this, LogPathFollowing, Log, TEXT("Aborting move due to path being invalid and not waiting for repath"));
 			OnPathFinished(EPathFollowingResult::Aborted, FPathFollowingResultFlags::InvalidPath);
+			return;
+		}
+		else
+		{
+			// continue with execution, if navigation is being rebuild constantly AI will get stuck with current waypoint
+			// path updates should be still coming in, even though they get invalidated right away
+			UE_VLOG(this, LogPathFollowing, Log, TEXT("Updating path points in invalid & pending path!"));
 		}
-
-		return;
 	}
 
 	FMetaNavMeshPath* MetaNavPath = bIsUsingMetaPath ? Path->CastPath<FMetaNavMeshPath>() : nullptr;
diff --git a/Engine/Source/Runtime/AIModule/Private/Perception/AIPerceptionComponent.cpp b/Engine/Source/Runtime/AIModule/Private/Perception/AIPerceptionComponent.cpp
--- a/Engine/Source/Runtime/AIModule/Private/Perception/AIPerceptionComponent.cpp
+++ b/Engine/Source/Runtime/AIModule/Private/Perception/AIPerceptionComponent.cpp
@@ -244,7 +244,10 @@ void UAIPerceptionComponent::CleanUp()
 
 void UAIPerceptionComponent::BeginDestroy()
 {
-	CleanUp();
+	if (!HasAnyFlags(RF_ClassDefaultObject))
+	{
+		CleanUp();
+	}
 	Super::BeginDestroy();
 }
 
diff --git a/Engine/Source/Runtime/AIModule/Private/Perception/AIPerceptionSystem.cpp b/Engine/Source/Runtime/AIModule/Private/Perception/AIPerceptionSystem.cpp
--- a/Engine/Source/Runtime/AIModule/Private/Perception/AIPerceptionSystem.cpp
+++ b/Engine/Source/Runtime/AIModule/Private/Perception/AIPerceptionSystem.cpp
@@ -161,7 +161,7 @@ void UAIPerceptionSystem::Tick(float DeltaSeconds)
 	// if no new stimuli
 	// and it's not time to remove stimuli from "know events"
 
-	UWorld* World = GEngine->GetWorldFromContextObject(GetOuter());
+	UWorld* World = GEngine->GetWorldFromContextObjectChecked(GetOuter());
 	check(World);
 
 	if (World->bPlayersOnly == false)
@@ -280,16 +280,14 @@ void UAIPerceptionSystem::AgeStimuli()
 UAIPerceptionSystem* UAIPerceptionSystem::GetCurrent(UObject* WorldContextObject)
 {
 	UWorld* World = Cast<UWorld>(WorldContextObject);
-
-	if (World == nullptr && WorldContextObject != nullptr)
+	if (World == nullptr)
 	{
-		World = GEngine->GetWorldFromContextObject(WorldContextObject);
+		World = GEngine->GetWorldFromContextObject(WorldContextObject, EGetWorldErrorMode::LogAndReturnNull);
 	}
 
 	if (World && World->GetAISystem())
 	{
-		check(Cast<UAISystem>(World->GetAISystem()));
-		UAISystem* AISys = (UAISystem*)(World->GetAISystem());
+		UAISystem* AISys = CastChecked<UAISystem>(World->GetAISystem());
 
 		return AISys->GetPerceptionSystem();
 	}
@@ -579,7 +577,7 @@ bool UAIPerceptionSystem::RegisterPerceptionStimuliSource(UObject* WorldContextO
 	bool bResult = false;
 	if (Sense && Target)
 	{
-		UWorld* World = GEngine->GetWorldFromContextObject(WorldContextObject);
+		UWorld* World = GEngine->GetWorldFromContextObject(WorldContextObject, EGetWorldErrorMode::LogAndReturnNull);
 		if (World && World->GetAISystem())
 		{
 			UAISystem* AISys = Cast<UAISystem>(World->GetAISystem());
diff --git a/Engine/Source/Runtime/AIModule/Private/Perception/AISense_Sight.cpp b/Engine/Source/Runtime/AIModule/Private/Perception/AISense_Sight.cpp
--- a/Engine/Source/Runtime/AIModule/Private/Perception/AISense_Sight.cpp
+++ b/Engine/Source/Runtime/AIModule/Private/Perception/AISense_Sight.cpp
@@ -141,7 +141,7 @@ float UAISense_Sight::Update()
 {
 	SCOPE_CYCLE_COUNTER(STAT_AI_Sense_Sight);
 
-	const UWorld* World = GEngine->GetWorldFromContextObject(GetPerceptionSystem()->GetOuter());
+	const UWorld* World = GEngine->GetWorldFromContextObject(GetPerceptionSystem()->GetOuter(), EGetWorldErrorMode::LogAndReturnNull);
 
 	if (World == NULL)
 	{
diff --git a/Engine/Source/Runtime/AssetRegistry/Private/AssetDataGatherer.cpp b/Engine/Source/Runtime/AssetRegistry/Private/AssetDataGatherer.cpp
--- a/Engine/Source/Runtime/AssetRegistry/Private/AssetDataGatherer.cpp
+++ b/Engine/Source/Runtime/AssetRegistry/Private/AssetDataGatherer.cpp
@@ -13,7 +13,7 @@
 
 namespace AssetDataGathererConstants
 {
-	static const int32 CacheSerializationVersion = 10;
+	static const int32 CacheSerializationVersion = 11;
 	static const int32 MaxFilesToDiscoverBeforeFlush = 2500;
 	static const int32 MaxFilesToGatherBeforeFlush = 250;
 	static const int32 MaxFilesToProcessBeforeCacheWrite = 50000;
diff --git a/Engine/Source/Runtime/AssetRegistry/Private/AssetRegistry.cpp b/Engine/Source/Runtime/AssetRegistry/Private/AssetRegistry.cpp
--- a/Engine/Source/Runtime/AssetRegistry/Private/AssetRegistry.cpp
+++ b/Engine/Source/Runtime/AssetRegistry/Private/AssetRegistry.cpp
@@ -23,6 +23,8 @@
 #if WITH_EDITOR
 #include "IDirectoryWatcher.h"
 #include "DirectoryWatcherModule.h"
+#include "HAL/ThreadHeartBeat.h"
+#include "HAL/PlatformProcess.h"
 #endif // WITH_EDITOR
 
 
@@ -362,8 +364,27 @@ void UAssetRegistryImpl::SearchAllAssets(bool bSynchronousSearch)
 	// Start the asset search (synchronous in commandlets)
 	if ( bSynchronousSearch )
 	{
-		const bool bForceRescan = false;
-		ScanPathsAndFilesSynchronous(PathsToSearch, TArray<FString>(), bForceRescan, EAssetDataCacheMode::UseMonolithicCache);
+#if WITH_EDITOR
+		if (IsLoadingAssets())
+		{
+			// Force a flush of the current gatherer instead
+			UE_LOG(LogAssetRegistry, Log, TEXT("Flushing asset discovery search because of synchronous request, this can take several seconds..."));
+
+			while (IsLoadingAssets())
+			{
+				Tick(-1.0f);
+
+				FThreadHeartBeat::Get().HeartBeat();
+				FPlatformProcess::SleepNoStats(0.0001f);
+			}
+		}
+		else
+#endif
+		{
+			const bool bForceRescan = false;
+			ScanPathsAndFilesSynchronous(PathsToSearch, TArray<FString>(), bForceRescan, EAssetDataCacheMode::UseMonolithicCache);
+		}
+
 
 #if WITH_EDITOR
 		if (IsRunningCommandlet())
@@ -469,13 +490,6 @@ bool UAssetRegistryImpl::GetAssets(const FARFilter& InFilter, TArray<FAssetData>
 			{
 				UPackage* InMemoryPackage = Obj->GetOutermost();
 
-				static const bool bUsingWorldAssets = UAssetRegistryImpl::IsUsingWorldAssets();
-				// Skip assets in map packages... unless we are showing world assets
-				if ( InMemoryPackage->ContainsMap() && !bUsingWorldAssets )
-				{
-					return;
-				}
-
 				// Skip assets that were loaded for diffing
 				if (InMemoryPackage->HasAnyPackageFlags(PKG_ForDiffing))
 				{
@@ -1392,11 +1406,6 @@ void UAssetRegistryImpl::Tick(float DeltaTime)
 	}
 }
 
-bool UAssetRegistryImpl::IsUsingWorldAssets()
-{
-	return !FParse::Param(FCommandLine::Get(), TEXT("DisableWorldAssets"));
-}
-
 void UAssetRegistryImpl::Serialize(FArchive& Ar)
 {
 	State.Serialize(Ar, SerializationOptions);
@@ -1411,6 +1420,19 @@ void UAssetRegistryImpl::Serialize(FArchive& Ar)
 			if (AssetData != nullptr)
 			{
 				AddAssetPath(AssetData->PackagePath);
+
+				// Populate the class map if adding blueprint
+				if (ClassGeneratorNames.Contains(AssetData->AssetClass))
+				{
+					const FString GeneratedClass = AssetData->GetTagValueRef<FString>("GeneratedClass");
+					const FString ParentClass = AssetData->GetTagValueRef<FString>("ParentClass");
+					if (!GeneratedClass.IsEmpty() && !ParentClass.IsEmpty())
+					{
+						const FName GeneratedClassFName = *ExportTextPathToObjectName(GeneratedClass);
+						const FName ParentClassFName = *ExportTextPathToObjectName(ParentClass);
+						CachedInheritanceMap.Add(GeneratedClassFName, ParentClassFName);
+					}
+				}
 			}
 		}
 	}
@@ -1802,6 +1824,25 @@ void UAssetRegistryImpl::DependencyDataGathered(const double TickStartTime, TBac
 			
 			if (DependsNode != nullptr)
 			{
+				const FAssetIdentifier& Identifier = DependsNode->GetIdentifier();
+				if (DependsNode->GetConnectionCount() == 0 && Identifier.IsPackage())
+				{
+					// This was newly created, see if we need to read the script package Guid
+					FString PackageName = Identifier.PackageName.ToString();
+
+					if (FPackageName::IsScriptPackage(PackageName))
+					{
+						// Get the guid off the script package, this is updated when script is changed
+						UPackage* Package = FindPackage(nullptr, *PackageName);
+
+						if (Package)
+						{
+							FAssetPackageData* ScriptPackageData = State.CreateOrGetAssetPackageData(Identifier.PackageName);
+							ScriptPackageData->PackageGuid = Package->GetGuid();
+						}
+					}
+				}
+
 				Node->AddDependency(DependsNode, NewDependsIt.Value);
 				DependsNode->AddReferencer(Node);
 			}
diff --git a/Engine/Source/Runtime/AssetRegistry/Private/AssetRegistry.h b/Engine/Source/Runtime/AssetRegistry/Private/AssetRegistry.h
--- a/Engine/Source/Runtime/AssetRegistry/Private/AssetRegistry.h
+++ b/Engine/Source/Runtime/AssetRegistry/Private/AssetRegistry.h
@@ -112,8 +112,8 @@ public:
 
 	virtual void Tick (float DeltaTime) override;
 
-	/** True if world assets are enabled */
-	static bool IsUsingWorldAssets();
+	DEPRECATED(4.17, "IsUsingWorldAssets is now always true, remove any code that assumes it could be false")
+	static bool IsUsingWorldAssets() { return true; }
 
 protected:
 	virtual void SetManageReferences(const TMultiMap<FAssetIdentifier, FAssetIdentifier>& ManagerMap, bool bClearExisting, EAssetRegistryDependencyType::Type RecurseType, ShouldSetManagerPredicate ShouldSetManager) override;
diff --git a/Engine/Source/Runtime/AssetRegistry/Private/AssetRegistryState.cpp b/Engine/Source/Runtime/AssetRegistry/Private/AssetRegistryState.cpp
--- a/Engine/Source/Runtime/AssetRegistry/Private/AssetRegistryState.cpp
+++ b/Engine/Source/Runtime/AssetRegistry/Private/AssetRegistryState.cpp
@@ -192,15 +192,29 @@ void FAssetRegistryState::InitializeFromExisting(const TMap<FName, FAssetData*>&
 		}
 	}
 
+	TSet<FAssetIdentifier> ScriptPackages;
+
 	if (!bRefreshExisting)
 	{
 		for (const TPair<FName, FAssetPackageData*>& Pair : AssetPackageDataMap)
 		{
-			// Only add if also in asset data map
-			if (Pair.Value && CachedAssetsByPackageName.Find(Pair.Key))
+			bool bIsScriptPackage = FPackageName::IsScriptPackage(Pair.Key.ToString());
+
+			if (Pair.Value)
 			{
-				FAssetPackageData* NewData = CreateOrGetAssetPackageData(Pair.Key);
-				*NewData = *Pair.Value;
+				// Only add if also in asset data map, or script package
+				if (bIsScriptPackage)
+				{
+					ScriptPackages.Add(Pair.Key);
+
+					FAssetPackageData* NewData = CreateOrGetAssetPackageData(Pair.Key);
+					*NewData = *Pair.Value;
+				}
+				else if (CachedAssetsByPackageName.Find(Pair.Key) || bIsScriptPackage)
+				{
+					FAssetPackageData* NewData = CreateOrGetAssetPackageData(Pair.Key);
+					*NewData = *Pair.Value;
+				}
 			}
 		}
 
@@ -209,11 +223,12 @@ void FAssetRegistryState::InitializeFromExisting(const TMap<FName, FAssetData*>&
 			FDependsNode* OldNode = Pair.Value;
 			FDependsNode* NewNode = CreateOrFindDependsNode(Pair.Key);
 
-			Pair.Value->IterateOverDependencies([this, &DependsNodeMap, OldNode, NewNode](FDependsNode* InDependency, EAssetRegistryDependencyType::Type InDependencyType) {
-				if (DependsNodeMap.Find(InDependency->GetIdentifier()))
+			Pair.Value->IterateOverDependencies([this, &DependsNodeMap, &ScriptPackages, OldNode, NewNode](FDependsNode* InDependency, EAssetRegistryDependencyType::Type InDependencyType) {
+				const FAssetIdentifier& Identifier = InDependency->GetIdentifier();
+				if (DependsNodeMap.Find(Identifier) || ScriptPackages.Contains(Identifier))
 				{
 					// Only add if this node is in the incoming map
-					FDependsNode* NewDependency = CreateOrFindDependsNode(InDependency->GetIdentifier());
+					FDependsNode* NewDependency = CreateOrFindDependsNode(Identifier);
 					NewNode->AddDependency(NewDependency, InDependencyType, true);
 					NewDependency->AddReferencer(NewNode);
 				}
@@ -668,7 +683,7 @@ bool FAssetRegistryState::Serialize(FArchive& OriginalAr, FAssetRegistrySerializ
 		FAssetRegistryVersion::Type Version = FAssetRegistryVersion::LatestVersion;
 		FAssetRegistryVersion::SerializeVersion(OriginalAr, Version);
 
-		if (Version < FAssetRegistryVersion::ChangedAssetData)
+		if (Version < FAssetRegistryVersion::RemovedMD5Hash)
 		{
 			// Cannot read states before this version
 			return false;
diff --git a/Engine/Source/Runtime/AssetRegistry/Private/PackageReader.cpp b/Engine/Source/Runtime/AssetRegistry/Private/PackageReader.cpp
--- a/Engine/Source/Runtime/AssetRegistry/Private/PackageReader.cpp
+++ b/Engine/Source/Runtime/AssetRegistry/Private/PackageReader.cpp
@@ -139,20 +139,13 @@ bool FPackageReader::ReadAssetRegistryData (TArray<FAssetData*>& AssetDataList)
 
 	const bool bIsMapPackage = (PackageFileSummary.PackageFlags & PKG_ContainsMap) != 0;
 
-	// Assets do not show up in map packages unless we launch with -WorldAssets
-	static const bool bUsingWorldAssets = UAssetRegistryImpl::IsUsingWorldAssets();
-	if ( bIsMapPackage && !bUsingWorldAssets )
-	{
-		return true;
-	}
-
 	// Load the object count
 	int32 ObjectCount = 0;
 	*this << ObjectCount;
 
 	// Worlds that were saved before they were marked public do not have asset data so we will synthesize it here to make sure we see all legacy umaps
 	// We will also do this for maps saved after they were marked public but no asset data was saved for some reason. A bug caused this to happen for some maps.
-	if (bUsingWorldAssets && bIsMapPackage)
+	if (bIsMapPackage)
 	{
 		const bool bLegacyPackage = PackageFileSummary.GetFileVersionUE4() < VER_UE4_PUBLIC_WORLDS;
 		const bool bNoMapAsset = (ObjectCount == 0);
diff --git a/Engine/Source/Runtime/AssetRegistry/Public/AssetData.h b/Engine/Source/Runtime/AssetRegistry/Public/AssetData.h
--- a/Engine/Source/Runtime/AssetRegistry/Public/AssetData.h
+++ b/Engine/Source/Runtime/AssetRegistry/Public/AssetData.h
@@ -10,7 +10,6 @@
 #include "UObject/Package.h"
 #include "UObject/ObjectRedirector.h"
 #include "Misc/PackageName.h"
-#include "Misc/SecureHash.h"
 #include "UObject/LinkerLoad.h"
 #include "AssetDataTagMap.h"
 #include "PrimaryAssetId.h"
@@ -28,6 +27,7 @@ struct ASSETREGISTRY_API FAssetRegistryVersion
 		HardSoftDependencies,	// The first version of the runtime asset registry to include file versioning.
 		AddAssetRegistryState,	// Added FAssetRegistryState and support for piecemeal serialization
 		ChangedAssetData,		// AssetData serialization format changed, versions before this are not readable
+		RemovedMD5Hash,			// Removed MD5 hash from package data
 
 		// -----<new versions can be added above this line>-------------------------------------------------
 		VersionPlusOne,
@@ -519,12 +519,9 @@ public:
 	/** Total size of this asset on disk */
 	int64 DiskSize;
 
-	/** Guid of the guid, uniquely identifies an asset package */
+	/** Guid of the source package, uniquely identifies an asset package */
 	FGuid PackageGuid;
 
-	/** Hash of the package and any dependencies that created it */
-	FMD5Hash PackageSourceHash;
-
 	FAssetPackageData()
 		: DiskSize(0)
 	{
@@ -538,7 +535,6 @@ public:
 	{
 		Ar << DiskSize;
 		Ar << PackageGuid;
-		Ar << PackageSourceHash;
 	}
 };
 
diff --git a/Engine/Source/Runtime/AudioMixer/Private/AudioMixerBlueprintLibrary.cpp b/Engine/Source/Runtime/AudioMixer/Private/AudioMixerBlueprintLibrary.cpp
--- a/Engine/Source/Runtime/AudioMixer/Private/AudioMixerBlueprintLibrary.cpp
+++ b/Engine/Source/Runtime/AudioMixer/Private/AudioMixerBlueprintLibrary.cpp
@@ -9,7 +9,7 @@
 
 static FAudioDevice* GetAudioDeviceFromWorldContext(const UObject* WorldContextObject)
 {
-	UWorld* ThisWorld = GEngine->GetWorldFromContextObject(WorldContextObject);
+	UWorld* ThisWorld = GEngine->GetWorldFromContextObject(WorldContextObject, EGetWorldErrorMode::LogAndReturnNull);
 	if (!ThisWorld || !ThisWorld->bAllowAudioPlayback || ThisWorld->GetNetMode() == NM_DedicatedServer)
 	{
 		return nullptr;
diff --git a/Engine/Source/Runtime/Core/Public/Misc/AssertionMacros.h b/Engine/Source/Runtime/Core/Public/Misc/AssertionMacros.h
--- a/Engine/Source/Runtime/Core/Public/Misc/AssertionMacros.h
+++ b/Engine/Source/Runtime/Core/Public/Misc/AssertionMacros.h
@@ -145,17 +145,20 @@ struct CORE_API FDebug
  *
  * Example:
  *
- *		if( ensure( InObject != NULL ) )
+ *		if (ensure(InObject != nullptr))
  *		{
  *			InObject->Modify();
  *		}
  *
- * This code is safe to execute as the pointer dereference is wrapped in a non-NULL conditional block, but
+ * This code is safe to execute as the pointer dereference is wrapped in a non-nullptr conditional block, but
  * you still want to find out if this ever happens so you can avoid side effects.  Using ensure() here will
  * force a crash report to be generated without crashing the application (and potentially causing editor
  * users to lose unsaved work.)
  *
- * ensure() resolves to a regular assertion (crash) in shipping or test builds.
+ * ensure() resolves to just evaluate the expression when DO_CHECK is 0 (typically shipping or test builds).
+ *
+ * By default a given call site will only print the callstack and submit the 'crash report' the first time an
+ * ensure is hit in a session; ensureAlways can be used instead if you want to handle every failure
  */
 
 #if DO_CHECK && !USING_CODE_ANALYSIS // The Visual Studio 2013 analyzer doesn't understand these complex conditionals
diff --git a/Engine/Source/Runtime/Core/Public/UObject/PropertyPortFlags.h b/Engine/Source/Runtime/Core/Public/UObject/PropertyPortFlags.h
--- a/Engine/Source/Runtime/Core/Public/UObject/PropertyPortFlags.h
+++ b/Engine/Source/Runtime/Core/Public/UObject/PropertyPortFlags.h
@@ -29,7 +29,9 @@ enum EPropertyPortFlags
 	/** Indicates that importing values for config properties is disallowed */
 	PPF_RestrictImportTypes			= 0x00000020,
 
-	//								= 0x00000040,
+	/** Indicates that this is a blueprint pin or something else that is saved to disk as import text */
+	PPF_SerializedAsImportText		= 0x00000040,
+
 	//								= 0x00000080,
 
 	/** only include properties which are marked CPF_InstancedReference */
@@ -72,7 +74,7 @@ enum EPropertyPortFlags
 	/** indicates that we're exporting properties for display in the property window. - used to hide EditHide items in collapsed structs */
 	PPF_PropertyWindow				= 0x00080000,
 
-	PPF_NoInternalArcheType			= 0x00100000,
+	//								= 0x00100000,
 
 	/** Force fully qualified object names (for debug dumping) */
 	PPF_DebugDump					= 0x00200000,
diff --git a/Engine/Source/Runtime/CoreUObject/Private/Misc/StringAssetReference.cpp b/Engine/Source/Runtime/CoreUObject/Private/Misc/StringAssetReference.cpp
--- a/Engine/Source/Runtime/CoreUObject/Private/Misc/StringAssetReference.cpp
+++ b/Engine/Source/Runtime/CoreUObject/Private/Misc/StringAssetReference.cpp
@@ -41,7 +41,7 @@ void FStringAssetReference::SetPath(FString Path)
 	}
 }
 
-void FStringAssetReference::PreSavePath()
+bool FStringAssetReference::PreSavePath()
 {
 #if WITH_EDITOR
 	FString* FoundRedirection = GRedirectCollector.GetAssetPathRedirection(ToString());
@@ -49,8 +49,10 @@ void FStringAssetReference::PreSavePath()
 	if (FoundRedirection)
 	{
 		SetPath(*FoundRedirection);
+		return true;
 	}
 #endif // WITH_EDITOR
+	return false;
 }
 
 void FStringAssetReference::PostLoadPath() const
diff --git a/Engine/Source/Runtime/CoreUObject/Private/UObject/Class.cpp b/Engine/Source/Runtime/CoreUObject/Private/UObject/Class.cpp
--- a/Engine/Source/Runtime/CoreUObject/Private/UObject/Class.cpp
+++ b/Engine/Source/Runtime/CoreUObject/Private/UObject/Class.cpp
@@ -1430,10 +1430,6 @@ const UStruct* UStruct::HasMetaDataHierarchical(const FName& Key) const
 	}
 #endif // #if USE_CIRCULAR_DEPENDENCY_LOAD_DEFERRING
 
-//
-// Serialize an expression to an archive.
-// Returns expression token.
-//
 EExprToken UStruct::SerializeExpr( int32& iCode, FArchive& Ar )
 {
 #define SERIALIZEEXPR_INC
diff --git a/Engine/Source/Runtime/CoreUObject/Private/UObject/LinkerLoad.cpp b/Engine/Source/Runtime/CoreUObject/Private/UObject/LinkerLoad.cpp
--- a/Engine/Source/Runtime/CoreUObject/Private/UObject/LinkerLoad.cpp
+++ b/Engine/Source/Runtime/CoreUObject/Private/UObject/LinkerLoad.cpp
@@ -3845,8 +3845,10 @@ UObject* FLinkerLoad::CreateExport( int32 Index )
 
 				for (UObject* SubObject : SuperSubObjects)
 				{
-					if (SubObject->HasAnyFlags(RF_NeedLoad))
+					// Matching behavior in UBlueprint::ForceLoad to ensure that the subobject is actually loaded:
+					if (SubObject->HasAnyFlags(RF_NeedLoad) || !SubObject->HasAnyFlags(RF_LoadCompleted))
 					{
+						SubObject->SetFlags(RF_NeedLoad);
 						Preload(SubObject);
 					}
 				}
diff --git a/Engine/Source/Runtime/CoreUObject/Private/UObject/Obj.cpp b/Engine/Source/Runtime/CoreUObject/Private/UObject/Obj.cpp
--- a/Engine/Source/Runtime/CoreUObject/Private/UObject/Obj.cpp
+++ b/Engine/Source/Runtime/CoreUObject/Private/UObject/Obj.cpp
@@ -762,22 +762,38 @@ FString UObject::GetDetailedInfo() const
 }
 
 #if WITH_ENGINE
+
+#if DO_CHECK
+// Used to check to see if a derived class actually implemented GetWorld() or not, but only on the game thread
 bool bGetWorldOverridden = false;
+#endif
 
 class UWorld* UObject::GetWorld() const
 {
-	bGetWorldOverridden = false;
-	return NULL;
+#if DO_CHECK
+	if (IsInGameThread())
+	{
+		bGetWorldOverridden = false;
+	}
+#endif
+	return nullptr;
 }
 
 class UWorld* UObject::GetWorldChecked(bool& bSupported) const
 {
-	bGetWorldOverridden = true;
+#if DO_CHECK
+	const bool bGameThread = IsInGameThread();
+	if (bGameThread)
+	{
+		bGetWorldOverridden = true;
+	}
+#endif
+
 	UWorld* World = GetWorld();
 
-	if (!bGetWorldOverridden)
-	{
 #if DO_CHECK
+	if (bGameThread && !bGetWorldOverridden)
+	{
 		static TSet<UClass*> ReportedClasses;
 
 		UClass* UnsupportedClass = GetClass();
@@ -791,23 +807,33 @@ class UWorld* UObject::GetWorldChecked(bool& bSupported) const
 				ParentHierarchy += FString::Printf(TEXT(", %s"), *SuperClass->GetName());
 			}
 
-			ensureMsgf(false, TEXT("Unsupported context object of class %s (SuperClass(es) - %s). You must add a way to retrieve a UWorld context for this class."), *UnsupportedClass->GetName(), *ParentHierarchy);
+			ensureAlwaysMsgf(false, TEXT("Unsupported context object of class %s (SuperClass(es) - %s). You must add a way to retrieve a UWorld context for this class."), *UnsupportedClass->GetName(), *ParentHierarchy);
 
 			ReportedClasses.Add(UnsupportedClass);
 		}
-#endif
 	}
 
-	bSupported = bGetWorldOverridden;
+	bSupported = bGameThread ? bGetWorldOverridden : (World != nullptr);
+	check(World && bSupported);
+#else
+	bSupported = World != nullptr;
+#endif
+
 	return World;
 }
 
 bool UObject::ImplementsGetWorld() const
 {
+#if DO_CHECK
+	check(IsInGameThread());
 	bGetWorldOverridden = true;
 	GetWorld();
 	return bGetWorldOverridden;
+#else
+	return true;
+#endif	
 }
+
 #endif
 
 #define PROFILE_ConditionalBeginDestroy (0)
@@ -1575,9 +1601,10 @@ void UObject::FAssetRegistryTag::GetAssetRegistryTagsFromSearchableProperties(co
 				// enums are alphabetical
 				TagType = FAssetRegistryTag::TT_Alphabetical;
 			}
-			else if ( Class->IsChildOf(UArrayProperty::StaticClass()) || Class->IsChildOf(UMapProperty::StaticClass()) || Class->IsChildOf(USetProperty::StaticClass()) )
+			else if (Class->IsChildOf(UArrayProperty::StaticClass()) || Class->IsChildOf(UMapProperty::StaticClass()) || Class->IsChildOf(USetProperty::StaticClass())
+				|| Class->IsChildOf(UStructProperty::StaticClass()) || Class->IsChildOf(UObjectPropertyBase::StaticClass()))
 			{
-				// arrays/maps/sets are hidden, it is often too much information to display and sort
+				// Arrays/maps/sets/structs/objects are hidden, it is often too much information to display and sort
 				TagType = FAssetRegistryTag::TT_Hidden;
 			}
 			else
@@ -1656,7 +1683,7 @@ void UObject::GetAssetRegistryTagMetadata(TMap<FName, FAssetRegistryTagMetadata>
 bool UObject::IsAsset() const
 {
 	// Assets are not transient or CDOs. They must be public.
-	const bool bHasValidObjectFlags = !HasAnyFlags(RF_Transient | RF_ClassDefaultObject) && HasAnyFlags(RF_Public);
+	const bool bHasValidObjectFlags = !HasAnyFlags(RF_Transient | RF_ClassDefaultObject) && HasAnyFlags(RF_Public) && !IsPendingKill();
 
 	if ( bHasValidObjectFlags )
 	{
diff --git a/Engine/Source/Runtime/CoreUObject/Private/UObject/PropertyBaseObject.cpp b/Engine/Source/Runtime/CoreUObject/Private/UObject/PropertyBaseObject.cpp
--- a/Engine/Source/Runtime/CoreUObject/Private/UObject/PropertyBaseObject.cpp
+++ b/Engine/Source/Runtime/CoreUObject/Private/UObject/PropertyBaseObject.cpp
@@ -32,7 +32,7 @@ void UObjectPropertyBase::InstanceSubobjects(void* Data, void const* DefaultData
 		UObject* CurrentValue = GetObjectPropertyValue((uint8*)Data + ArrayIndex * ElementSize);
 		if ( CurrentValue )
 		{
-			UObject *SubobjectTemplate = DefaultData ? GetObjectPropertyValue((uint8*)DefaultData + ArrayIndex * ElementSize): NULL;
+			UObject *SubobjectTemplate = DefaultData ? GetObjectPropertyValue((uint8*)DefaultData + ArrayIndex * ElementSize): nullptr;
 			UObject* NewValue = InstanceGraph->InstancePropertyValue(SubobjectTemplate, CurrentValue, Owner, HasAnyPropertyFlags(CPF_Transient), HasAnyPropertyFlags(CPF_InstancedReference));
 			SetObjectPropertyValue((uint8*)Data + ArrayIndex * ElementSize, NewValue);
 		}
@@ -41,8 +41,8 @@ void UObjectPropertyBase::InstanceSubobjects(void* Data, void const* DefaultData
 
 bool UObjectPropertyBase::Identical( const void* A, const void* B, uint32 PortFlags ) const
 {
-	UObject* ObjectA = A ? GetObjectPropertyValue(A) : NULL;
-	UObject* ObjectB = B ? GetObjectPropertyValue(B) : NULL;
+	UObject* ObjectA = A ? GetObjectPropertyValue(A) : nullptr;
+	UObject* ObjectB = B ? GetObjectPropertyValue(B) : nullptr;
 	if (!ObjectA && !ObjectB)
 	{
 		return true;
@@ -55,7 +55,7 @@ bool UObjectPropertyBase::Identical( const void* A, const void* B, uint32 PortFl
 	// which contains actor references. We want to serialize those references so they are fixed up.
 	const bool bDuplicatingForPIE = (PortFlags&PPF_DuplicateForPIE) != 0;
 	bool bResult = !bDuplicatingForPIE ? (ObjectA == ObjectB) : false;
-	// always serialize the cross level references, because they could be NULL
+	// always serialize the cross level references, because they could be nullptr
 	// @todo: okay, this is pretty hacky overall - we should have a PortFlag or something
 	// that is set during SavePackage. Other times, we don't want to immediately return false
 	// (instead of just this ExportDefProps case)
@@ -156,7 +156,7 @@ FString UObjectPropertyBase::GetExportPath(const UObject* Object, const UObject*
 			StopOuter = Parent->GetOutermost();
 		}
 	}
-	else if (Parent != NULL && Object->IsIn(Parent))
+	else if (Parent != nullptr && Object->IsIn(Parent))
 	{
 		StopOuter = Parent;
 	}
@@ -190,7 +190,7 @@ void UObjectPropertyBase::ExportTextItem( FString& ValueStr, const void* Propert
 		return;
 	}
 
-	if( Temp != NULL )
+	if( Temp != nullptr )
 	{
 		if (PortFlags & PPF_DebugDump)
 		{
@@ -223,7 +223,7 @@ void UObjectPropertyBase::ExportTextItem( FString& ValueStr, const void* Propert
  *
  * @param	Property			the property that the value is being importing to
  * @param	OwnerObject			the object that is importing the value; used for determining search scope.
- * @param	RequiredMetaClass	the meta-class for the object to find; if the object that is resolved is not of this class type, the result is NULL.
+ * @param	RequiredMetaClass	the meta-class for the object to find; if the object that is resolved is not of this class type, the result is nullptr.
  * @param	PortFlags			bitmask of EPropertyPortFlags that can modify the behavior of the search
  * @param	Buffer				the text to parse; should point to a textual representation of an object reference.  Can be just the object name (either fully 
  *								fully qualified or not), or can be formatted as a const object reference (i.e. SomeClass'SomePackage.TheObject')
@@ -246,14 +246,14 @@ bool UObjectPropertyBase::ParseObjectPropertyValue( const UProperty* Property, U
 
 	FString Temp;
 	Buffer = UPropertyHelpers::ReadToken(Buffer, Temp, true);
-	if ( Buffer == NULL )
+	if ( Buffer == nullptr )
 	{
 		return false;
 	}
 
 	if ( Temp == TEXT("None") )
 	{
-		out_ResolvedValue = NULL;
+		out_ResolvedValue = nullptr;
 	}
 	else
 	{
@@ -261,13 +261,13 @@ bool UObjectPropertyBase::ParseObjectPropertyValue( const UProperty* Property, U
 
 		SkipWhitespace(Buffer);
 
-		bool bWarnOnNULL = (PortFlags&PPF_CheckReferences)!=0;
+		bool bWarnOnnullptr = (PortFlags&PPF_CheckReferences)!=0;
 
 		if( *Buffer == TCHAR('\'') )
 		{
 			FString ObjectText;
 			Buffer = UPropertyHelpers::ReadToken( ++Buffer, ObjectText, true );
-			if( Buffer == NULL )
+			if( Buffer == nullptr )
 			{
 				return false;
 			}
@@ -287,21 +287,21 @@ bool UObjectPropertyBase::ParseObjectPropertyValue( const UProperty* Property, U
 			out_ResolvedValue = UObjectPropertyBase::FindImportedObject(Property, OwnerObject, ObjectClass, RequiredMetaClass, *Temp, PortFlags);
 		}
 
-		if ( out_ResolvedValue != NULL && !out_ResolvedValue->GetClass()->IsChildOf(RequiredMetaClass) )
+		if ( out_ResolvedValue != nullptr && !out_ResolvedValue->GetClass()->IsChildOf(RequiredMetaClass) )
 		{
-			if (bWarnOnNULL )
+			if (bWarnOnnullptr )
 			{
 				UE_LOG(LogProperty, Error, TEXT("%s: bad cast in '%s'"), *Property->GetFullName(), InBuffer );
 			}
 
-			out_ResolvedValue = NULL;
+			out_ResolvedValue = nullptr;
 			return false;
 		}
 
 		// If we couldn't find it or load it, we'll have to do without it.
-		if ( out_ResolvedValue == NULL )
+		if ( out_ResolvedValue == nullptr )
 		{
-			if( bWarnOnNULL )
+			if( bWarnOnnullptr )
 			{
 				UE_LOG(LogProperty, Warning, TEXT("%s: unresolved reference to '%s'"), *Property->GetFullName(), InBuffer );
 			}
@@ -315,7 +315,7 @@ bool UObjectPropertyBase::ParseObjectPropertyValue( const UProperty* Property, U
 const TCHAR* UObjectPropertyBase::ImportText_Internal( const TCHAR* InBuffer, void* Data, int32 PortFlags, UObject* Parent, FOutputDevice* ErrorText ) const
 {
 	const TCHAR* Buffer = InBuffer;
-	UObject* Result = NULL;
+	UObject* Result = nullptr;
 
 	bool bOk = ParseObjectPropertyValue(this, Parent, PropertyClass, PortFlags, Buffer, Result);
 
@@ -325,7 +325,7 @@ const TCHAR* UObjectPropertyBase::ImportText_Internal( const TCHAR* InBuffer, vo
 
 UObject* UObjectPropertyBase::FindImportedObject( const UProperty* Property, UObject* OwnerObject, UClass* ObjectClass, UClass* RequiredMetaClass, const TCHAR* Text, uint32 PortFlags/*=0*/ )
 {
-	UObject*	Result = NULL;
+	UObject*	Result = nullptr;
 	check( ObjectClass->IsChildOf(RequiredMetaClass) );
 
 	bool AttemptNonQualifiedSearch = (PortFlags & PPF_AttemptNonQualifiedSearch) != 0; 
@@ -334,16 +334,16 @@ UObject* UObjectPropertyBase::FindImportedObject( const UProperty* Property, UOb
 	// looking through the archetype chain at each outer and stop once the outer chain reaches the owning class's default object
 	if (PortFlags & PPF_ParsingDefaultProperties)
 	{
-		for (UObject* SearchStart = OwnerObject; Result == NULL && SearchStart != NULL; SearchStart = SearchStart->GetOuter())
+		for (UObject* SearchStart = OwnerObject; Result == nullptr && SearchStart != nullptr; SearchStart = SearchStart->GetOuter())
 		{
 			UObject* ScopedSearchRoot = SearchStart;
-			while (Result == NULL && ScopedSearchRoot != NULL)
+			while (Result == nullptr && ScopedSearchRoot != nullptr)
 			{
 				Result = StaticFindObjectSafe(ObjectClass, ScopedSearchRoot, Text);
 				// don't think it's possible to get a non-subobject here, but it doesn't hurt to check
-				if (Result != NULL && !Result->IsTemplate(RF_ClassDefaultObject))
+				if (Result != nullptr && !Result->IsTemplate(RF_ClassDefaultObject))
 				{
-					Result = NULL;
+					Result = nullptr;
 				}
 
 				ScopedSearchRoot = ScopedSearchRoot->GetArchetype();
@@ -360,33 +360,43 @@ UObject* UObjectPropertyBase::FindImportedObject( const UProperty* Property, UOb
 	// fully qualified, and this will step up the nested object chain to solve any name
 	// collisions within a nested object tree
 	UObject* ScopedSearchRoot = OwnerObject;
-	while (Result == NULL && ScopedSearchRoot != NULL)
+	while (Result == nullptr && ScopedSearchRoot != nullptr)
 	{
 		Result = StaticFindObjectSafe(ObjectClass, ScopedSearchRoot, Text);
 		// disallow class default subobjects here while importing defaults
 		// this prevents the use of a subobject name that doesn't exist in the scope of the default object being imported
 		// from grabbing some other subobject with the same name and class in some other arbitrary default object
-		if (Result != NULL && (PortFlags & PPF_ParsingDefaultProperties) && Result->IsTemplate(RF_ClassDefaultObject))
+		if (Result != nullptr && (PortFlags & PPF_ParsingDefaultProperties) && Result->IsTemplate(RF_ClassDefaultObject))
 		{
-			Result = NULL;
+			Result = nullptr;
 		}
 
 		ScopedSearchRoot = ScopedSearchRoot->GetOuter();
 	}
 
-	if (Result == NULL)
+	if (Result == nullptr)
 	{
 		// attempt to find a fully qualified object
-		Result = StaticFindObjectSafe(ObjectClass, NULL, Text);
+		Result = StaticFindObjectSafe(ObjectClass, nullptr, Text);
+
+		if (Result == nullptr && (PortFlags & PPF_SerializedAsImportText))
+		{
+			// Check string asset redirectors
+			FStringAssetReference Path = FString(Text);
+			if (Path.PreSavePath())
+			{
+				Result = StaticFindObjectSafe(ObjectClass, nullptr, *Path.ToString());
+			}
+		}
 
-		if (Result == NULL)
+		if (Result == nullptr)
 		{
 			// match any object of the correct class whose path contains the specified path
 			Result = StaticFindObjectSafe(ObjectClass, ANY_PACKAGE, Text);
 			// disallow class default subobjects here while importing defaults
-			if (Result != NULL && (PortFlags & PPF_ParsingDefaultProperties) && Result->IsTemplate(RF_ClassDefaultObject))
+			if (Result != nullptr && (PortFlags & PPF_ParsingDefaultProperties) && Result->IsTemplate(RF_ClassDefaultObject))
 			{
-				Result = NULL;
+				Result = nullptr;
 			}
 		}
 	}
@@ -435,7 +445,7 @@ UObject* UObjectPropertyBase::FindImportedObject( const UProperty* Property, UOb
 				uint32 LoadFlags = LOAD_NoWarn | LOAD_FindIfFail;
 
 				UE_LOG(LogProperty, Verbose, TEXT("FindImportedObject is attempting to import [%s] (class = %s) with StaticLoadObject"), Text, *GetFullNameSafe(ObjectClass));
-				Result = StaticLoadObject(ObjectClass, NULL, Text, NULL, LoadFlags, NULL);
+				Result = StaticLoadObject(ObjectClass, nullptr, Text, nullptr, LoadFlags, nullptr);
 
 #if USE_DEFERRED_DEPENDENCY_CHECK_VERIFICATION_TESTS
 				check(!bDeferAssetImports || !Result || !FBlueprintSupport::IsInBlueprintPackage(Result));
@@ -451,7 +461,7 @@ UObject* UObjectPropertyBase::FindImportedObject( const UProperty* Property, UOb
 		if ( !ObjectProperty || !ObjectProperty->AllowCrossLevel())
 		{
 			UE_LOG(LogProperty, Warning, TEXT("Illegal TEXT reference to a private object in external package (%s) from referencer (%s).  Import failed..."), *Result->GetFullName(), *OwnerObject->GetFullName());
-			Result = NULL;
+			Result = nullptr;
 		}
 	}
 
@@ -467,7 +477,7 @@ FName UObjectPropertyBase::GetID() const
 UObject* UObjectPropertyBase::GetObjectPropertyValue(const void* PropertyValueAddress) const
 {
 	check(0);
-	return NULL;
+	return nullptr;
 }
 
 void UObjectPropertyBase::SetObjectPropertyValue(void* PropertyValueAddress, UObject* Value) const
@@ -515,13 +525,13 @@ void UObjectPropertyBase::CheckValidObject(void* Value) const
 		if ((PropertyClass != nullptr) && !ObjectClass->IsChildOf(PropertyClass) && !bIsReplacingClassRefs && !bIsDeferringValueLoad)
 		{
 			UE_LOG(LogProperty, Warning,
-				TEXT("Serialized %s for a property of %s. Reference will be NULLed.\n    Property = %s\n    Item = %s"),
+				TEXT("Serialized %s for a property of %s. Reference will be nullptred.\n    Property = %s\n    Item = %s"),
 				*Object->GetClass()->GetFullName(),
 				*PropertyClass->GetFullName(),
 				*GetFullName(),
 				*Object->GetFullName()
 				);
-			SetObjectPropertyValue(Value, NULL);
+			SetObjectPropertyValue(Value, nullptr);
 		}
 	}
 }
diff --git a/Engine/Source/Runtime/CoreUObject/Private/UObject/ScriptCore.cpp b/Engine/Source/Runtime/CoreUObject/Private/UObject/ScriptCore.cpp
--- a/Engine/Source/Runtime/CoreUObject/Private/UObject/ScriptCore.cpp
+++ b/Engine/Source/Runtime/CoreUObject/Private/UObject/ScriptCore.cpp
@@ -346,6 +346,10 @@ FString FFrame::GetScriptCallstack()
 			ScriptStack += TEXT("\t") + BlueprintExceptionTracker.ScriptStack[i].GetStackDescription() + TEXT("\n");
 		}
 	}
+	else
+	{
+		ScriptStack += TEXT("\t[Empty] (FFrame::GetScriptCallstack() called from native code)");
+	}
 #else
 	ScriptStack = TEXT("Unable to display Script Callstack. Compile with DO_BLUEPRINT_GUARD=1");
 #endif
@@ -463,10 +467,17 @@ FString FFrame::GetStackTrace() const
 	}
 	
 	// and then dump them to a string
-	Result += FString( TEXT("Script call stack:\n") );
-	for (int32 Index = FrameStack.Num() - 1; Index >= 0; Index--)
+	if (FrameStack.Num() > 0)
+	{
+		Result += FString(TEXT("Script call stack:\n"));
+		for (int32 Index = FrameStack.Num() - 1; Index >= 0; Index--)
+		{
+			Result += FString::Printf(TEXT("\t%s\n"), *FrameStack[Index]->Node->GetFullName());
+		}
+	}
+	else
 	{
-		Result += FString::Printf(TEXT("\t%s\n"), *FrameStack[Index]->Node->GetFullName());
+		Result += FString(TEXT("Script call stack: [Empty] (FFrame::GetStackTrace() called from native code)"));
 	}
 
 	return Result;
@@ -1032,8 +1043,18 @@ bool UObject::CallFunctionByNameWithArguments(const TCHAR* Str, FOutputDevice& A
 	uint8* Parms = (uint8*)FMemory_Alloca(Function->ParmsSize);
 	FMemory::Memzero( Parms, Function->ParmsSize );
 
+	for (TFieldIterator<UProperty> It(Function); It && It->HasAnyPropertyFlags(CPF_Parm); ++It)
+	{
+		UProperty* LocalProp = *It;
+		checkSlow(LocalProp);
+		if (!LocalProp->HasAnyPropertyFlags(CPF_ZeroConstructor))
+		{
+			LocalProp->InitializeValue_InContainer(Parms);
+		}
+	}
+
 	const uint32 ExportFlags = PPF_None;
-	bool Failed = 0;
+	bool bFailed = 0;
 	int32 NumParamsEvaluated = 0;
 	for( TFieldIterator<UProperty> It(Function); It && (It->PropertyFlags & (CPF_Parm|CPF_ReturnParm))==CPF_Parm; ++It, NumParamsEvaluated++ )
 	{
@@ -1095,22 +1116,22 @@ bool UObject::CallFunctionByNameWithArguments(const TCHAR* Str, FOutputDevice& A
 		{
 			FFormatNamedArguments Arguments;
 			Arguments.Add(TEXT("Message"), FText::FromName( Message ));
-			Arguments.Add(TEXT("PropertyName"), FText::FromString( It->GetName() ));
-			Ar.Logf( TEXT("%s"), *FText::Format( NSLOCTEXT( "Core", "BadProperty", "'{Message}': Bad or missing property '{PropertyName}'" ), Arguments ).ToString() );
-			Failed = true;
+			Arguments.Add(TEXT("PropertyName"), FText::FromName(It->GetFName()));
+			Arguments.Add(TEXT("FunctionName"), FText::FromName(Function->GetFName()));
+			Ar.Logf( TEXT("%s"), *FText::Format( NSLOCTEXT( "Core", "BadProperty", "'{Message}': Bad or missing property '{PropertyName}' when trying to call {FunctionName}" ), Arguments ).ToString() );
+			bFailed = true;
 
 			break;
 		}
-
 	}
 
-	if( !Failed )
+	if( !bFailed )
 	{
 		ProcessEvent( Function, Parms );
 	}
 
 	//!!destructframe see also UObject::ProcessEvent
-	for( TFieldIterator<UProperty> It(Function); It && (It->PropertyFlags & (CPF_Parm|CPF_ReturnParm))==CPF_Parm; ++It )
+	for( TFieldIterator<UProperty> It(Function); It && It->HasAnyPropertyFlags(CPF_Parm); ++It )
 	{
 		It->DestroyValue_InContainer(Parms);
 	}
diff --git a/Engine/Source/Runtime/CoreUObject/Public/Misc/StringAssetReference.h b/Engine/Source/Runtime/CoreUObject/Public/Misc/StringAssetReference.h
--- a/Engine/Source/Runtime/CoreUObject/Public/Misc/StringAssetReference.h
+++ b/Engine/Source/Runtime/CoreUObject/Public/Misc/StringAssetReference.h
@@ -106,8 +106,8 @@ struct COREUOBJECT_API FStringAssetReference
 	/** Serializes the internal path and also handles save/PIE fixups. Call this from the archiver overrides */
 	void SerializePath(FArchive& Ar, bool bSkipSerializeIfArchiveHasSize = false);
 
-	/** Fixes up string asset reference for saving, call if saving with a method that skips SerializePath. This can modify the path */
-	void PreSavePath();
+	/** Fixes up string asset reference for saving, call if saving with a method that skips SerializePath. This can modify the path, it will return true if it was modified */
+	bool PreSavePath();
 
 	/** Handles when a string asset reference has been loaded, call if loading with a method that skips SerializePath. This does not modify path but might call callbacks */
 	void PostLoadPath() const;
diff --git a/Engine/Source/Runtime/CoreUObject/Public/Serialization/ArchiveReplaceObjectRef.h b/Engine/Source/Runtime/CoreUObject/Public/Serialization/ArchiveReplaceObjectRef.h
--- a/Engine/Source/Runtime/CoreUObject/Public/Serialization/ArchiveReplaceObjectRef.h
+++ b/Engine/Source/Runtime/CoreUObject/Public/Serialization/ArchiveReplaceObjectRef.h
@@ -9,11 +9,54 @@
 /** Base class for object replacement archives */ 
 class COREUOBJECT_API FArchiveReplaceObjectRefBase : public FArchiveUObject
 {
+public:
+
+	/**
+	* Returns the number of times the object was referenced
+	*/
+	int64 GetCount() const { return Count; }
+
+	/**
+	* Returns a reference to the object this archive is operating on
+	*/
+	const UObject* GetSearchObject() const { return SearchObject; }
+
+	/**
+	* Returns a reference to the replaced references map
+	*/
+	const TMap<UObject*, TArray<UProperty*>>& GetReplacedReferences() const { return ReplacedReferences; }
+
+	/**
+	* Returns the name of this archive.
+	**/
+	virtual FString GetArchiveName() const { return TEXT("ReplaceObjectRef"); }
+	
 protected:
 	/**
 	* Serializes a single object
 	*/
 	void SerializeObject(UObject* ObjectToSerialize);
+
+	/** Initial object to start the reference search from */
+	UObject* SearchObject;
+
+	/** The number of times encountered */
+	int32 Count;
+
+	/** List of objects that have already been serialized */
+	TSet<UObject*> SerializedObjects;
+
+	/** Object that will be serialized */
+	TArray<UObject*> PendingSerializationObjects;
+
+	/** Map of referencing objects to referencing properties */
+	TMap<UObject*, TArray<UProperty*>> ReplacedReferences;
+
+	/**
+	* Whether references to non-public objects not contained within the SearchObject
+	* should be set to null
+	*/
+	bool bNullPrivateReferences;
 };
 
 /*----------------------------------------------------------------------------
@@ -52,9 +95,12 @@ public:
 		bool bDelayStart = false,
 		bool bIgnoreClassGeneratedByRef = true
 	)
-	: SearchObject(InSearchObject), ReplacementMap(inReplacementMap)
-	, Count(0), bNullPrivateReferences(bNullPrivateRefs)
+	: ReplacementMap(inReplacementMap)
 	{
+		SearchObject = InSearchObject;
+		Count = 0;
+		bNullPrivateReferences = bNullPrivateRefs;
+
 		ArIsObjectReferenceCollector = true;
 		ArIsModifyingWeakAndStrongReferences = true;		// Also replace weak references too!
 		ArIgnoreArchetypeRef = bIgnoreArchetypeRef;
@@ -88,24 +134,6 @@ public:
 		}
 	}
 
-	/**
-	 * Returns the number of times the object was referenced
-	 */
-	int64 GetCount() const
-	{
-		return Count;
-	}
-
-	/**
-	 * Returns a reference to the object this archive is operating on
-	 */
-	const UObject* GetSearchObject() const { return SearchObject; }
-
-	/**
-	 * Returns a reference to the replaced references map
-	 */
-	const TMap<UObject*, TArray<UProperty*>>& GetReplacedReferences() const { return ReplacedReferences; }
-
 	/**
 	 * Serializes the reference to the object
 	 */
@@ -148,33 +176,8 @@ public:
 		return *this;
 	}
 
-	/**
-	 * Returns the name of this archive.
-	 **/
-	virtual FString GetArchiveName() const { return TEXT("ReplaceObjectRef"); }
-
 protected:
-	/** Initial object to start the reference search from */
-	UObject* SearchObject;
-
 	/** Map of objects to find references to -> object to replace references with */
 	const TMap<T*,T*>& ReplacementMap;
 	
-	/** The number of times encountered */
-	int32 Count;
-
-	/** List of objects that have already been serialized */
-	TSet<UObject*> SerializedObjects;
-
-	/** Object that will be serialized */
-	TArray<UObject*> PendingSerializationObjects;
-
-	/** Map of referencing objects to referencing properties */
-	TMap<UObject*, TArray<UProperty*>> ReplacedReferences;
-
-	/**
-	 * Whether references to non-public objects not contained within the SearchObject
-	 * should be set to null
-	 */
-	bool bNullPrivateReferences;
 };
diff --git a/Engine/Source/Runtime/CoreUObject/Public/UObject/AssetPtr.h b/Engine/Source/Runtime/CoreUObject/Public/UObject/AssetPtr.h
--- a/Engine/Source/Runtime/CoreUObject/Public/UObject/AssetPtr.h
+++ b/Engine/Source/Runtime/CoreUObject/Public/UObject/AssetPtr.h
@@ -137,6 +137,12 @@ public:
 		AssetPtr.Reset();
 	}
 
+	/** Resets the weak ptr only, call this when ObjectId may change */
+	FORCEINLINE void ResetWeakPtr()
+	{
+		AssetPtr.ResetWeakPtr();
+	}
+
 	/** Copy from an object already in memory */
 	template <typename U>
 	FORCEINLINE TAssetPtr& operator=(const U* Object)
@@ -396,6 +402,12 @@ public:
 		AssetPtr.Reset();
 	}
 
+	/** Resets the weak ptr only, call this when ObjectId may change */
+	FORCEINLINE void ResetWeakPtr()
+	{
+		AssetPtr.ResetWeakPtr();
+	}
+
 	/** Copy from a class already in memory */
 	FORCEINLINE void operator=(const UClass* From)
 	{
