Copying //UE4/Dev-Framework to //UE4/Dev-Main (Source: //UE4/Dev-Framework @ 3716594)

diff --git a/Engine/Source/Editor/UnrealEd/Private/MaterialGraphSchema.cpp b/Engine/Source/Editor/UnrealEd/Private/MaterialGraphSchema.cpp
--- a/Engine/Source/Editor/UnrealEd/Private/MaterialGraphSchema.cpp
+++ b/Engine/Source/Editor/UnrealEd/Private/MaterialGraphSchema.cpp
@@ -164,22 +164,23 @@ UEdGraphNode* FMaterialGraphSchemaAction_Paste::PerformAction(class UEdGraph* Pa
 //////////////////////////
 // UMaterialGraphSchema //
 
+const FName UMaterialGraphSchema::PC_Mask(TEXT("mask"));
+const FName UMaterialGraphSchema::PC_Required(TEXT("required"));
+const FName UMaterialGraphSchema::PC_Optional(TEXT("optional"));
+const FName UMaterialGraphSchema::PC_MaterialInput(TEXT("materialinput"));
+
+const FName UMaterialGraphSchema::PSC_Red(TEXT("red"));
+const FName UMaterialGraphSchema::PSC_Green(TEXT("green"));
+const FName UMaterialGraphSchema::PSC_Blue(TEXT("blue"));
+const FName UMaterialGraphSchema::PSC_Alpha(TEXT("alpha"));
+
+const FLinearColor UMaterialGraphSchema::ActivePinColor = FLinearColor::White;
+const FLinearColor UMaterialGraphSchema::InactivePinColor = FLinearColor(0.05f, 0.05f, 0.05f);
+const FLinearColor UMaterialGraphSchema::AlphaPinColor = FLinearColor(0.5f, 0.5f, 0.5f);
+
 UMaterialGraphSchema::UMaterialGraphSchema(const FObjectInitializer& ObjectInitializer)
 	: Super(ObjectInitializer)
 {
-	PC_Mask = TEXT("mask");
-	PC_Required = TEXT("required");
-	PC_Optional = TEXT("optional");
-	PC_MaterialInput = TEXT("materialinput");
-
-	PSC_Red = TEXT("red");
-	PSC_Green = TEXT("green");
-	PSC_Blue = TEXT("blue");
-	PSC_Alpha = TEXT("alpha");
-
-	ActivePinColor = FLinearColor::White;
-	InactivePinColor = FLinearColor(0.05f, 0.05f, 0.05f);
-	AlphaPinColor = FLinearColor(0.5f, 0.5f, 0.5f);
 }
 
 void UMaterialGraphSchema::SelectAllInputNodes(UEdGraph* Graph, UEdGraphPin* InGraphPin)
@@ -224,7 +225,7 @@ void UMaterialGraphSchema::GetBreakLinkToSubMenuActions( class FMenuBuilder& Men
 		FText Title = FText::FromString( TitleString );
 		if ( Pin->PinName != TEXT("") )
 		{
-			TitleString = FString::Printf(TEXT("%s (%s)"), *TitleString, *Pin->PinName);
+			TitleString = FString::Printf(TEXT("%s (%s)"), *TitleString, *Pin->PinName.ToString());
 
 			// Add name of connection if possible
 			FFormatNamedArguments Args;
@@ -441,7 +442,7 @@ void UMaterialGraphSchema::GetContextMenuActions(const UEdGraph* CurrentGraph, c
 							if (FunctionOutput)
 							{
 								FFormatNamedArguments Arguments;
-								Arguments.Add(TEXT("Name"), FText::FromString( *FunctionOutput->OutputName ));
+								Arguments.Add(TEXT("Name"), FText::FromName( FunctionOutput->OutputName ));
 								const FText Label = FText::Format( LOCTEXT( "ConnectToFunction", "Connect To {Name}" ), Arguments );
 								const FText ToolTip = FText::Format( LOCTEXT( "ConnectToFunctionTooltip", "Connects to the function output {Name}" ), Arguments );
 								MenuBuilder->AddMenuEntry(Label,
@@ -636,7 +637,7 @@ void UMaterialGraphSchema::BreakPinLinks(UEdGraphPin& TargetPin, bool bSendsNode
 	}
 }
 
-void UMaterialGraphSchema::BreakSinglePinLink(UEdGraphPin* SourcePin, UEdGraphPin* TargetPin)
+void UMaterialGraphSchema::BreakSinglePinLink(UEdGraphPin* SourcePin, UEdGraphPin* TargetPin) const
 {
 	const FScopedTransaction Transaction( NSLOCTEXT("UnrealEd", "GraphEd_BreakSinglePinLink", "Break Pin Link") );
 
diff --git a/Engine/Source/Editor/UnrealEd/Private/PlayLevel.cpp b/Engine/Source/Editor/UnrealEd/Private/PlayLevel.cpp
--- a/Engine/Source/Editor/UnrealEd/Private/PlayLevel.cpp
+++ b/Engine/Source/Editor/UnrealEd/Private/PlayLevel.cpp
@@ -3586,6 +3586,7 @@ UWorld* UEditorEngine::CreatePIEWorldByDuplication(FWorldContext &WorldContext,
 	PlayWorldPackage->PIEInstanceID = WorldContext.PIEInstance;
 	PlayWorldPackage->FileName = InPackage->FileName;
 	PlayWorldPackage->SetGuid( InPackage->GetGuid() );
+	PlayWorldPackage->MarkAsFullyLoaded();
 
 	// check(GPlayInEditorID == -1 || GPlayInEditorID == WorldContext.PIEInstance);
 	// Currently GPlayInEditorID is not correctly reset after map loading, so it's not safe to assert here
diff --git a/Engine/Source/Editor/UnrealEd/Private/SCurveEditor.cpp b/Engine/Source/Editor/UnrealEd/Private/SCurveEditor.cpp
--- a/Engine/Source/Editor/UnrealEd/Private/SCurveEditor.cpp
+++ b/Engine/Source/Editor/UnrealEd/Private/SCurveEditor.cpp
@@ -36,6 +36,8 @@
 
 #define LOCTEXT_NAMESPACE "SCurveEditor"
 
+DEFINE_LOG_CATEGORY_STATIC(LogCurveEditor, Log, All);
+
 const static FVector2D	CONST_KeySize		= FVector2D(11,11);
 const static FVector2D	CONST_TangentSize	= FVector2D(7,7);
 const static FVector2D	CONST_CurveSize		= FVector2D(12,12);
@@ -3242,10 +3244,17 @@ void SCurveEditor::OnBakeCurveSampleRateCommitted(const FText& InText, ETextComm
 	{
 		double NewBakeSampleRate = FCString::Atod(*InText.ToString());
 		const bool bIsNumber = InText.IsNumeric(); 
-		if(!bIsNumber)
+		if (!bIsNumber)
+		{
 			return;
+		}
+		if (NewBakeSampleRate <= 0.0)
+		{
+			UE_LOG(LogCurveEditor, Error, TEXT("Invalid Bake Sample Rate"));
+			return;
+		}
 
-		float BakeSampleRate = (float)NewBakeSampleRate;
+		const float BakeSampleRate = (float)NewBakeSampleRate;
 
 		const FScopedTransaction Transaction(LOCTEXT("CurveEditor_BakeCurve", "Bake Curve"));
 		CurveOwner->ModifyOwner();
@@ -3273,8 +3282,15 @@ void SCurveEditor::OnBakeCurveSampleRateCommitted(const FText& InText, ETextComm
 		{
 			for (auto CurveToBake : CurveRangeMap)
 			{
-				CurveToBake.Key->BakeCurve(BakeSampleRate, CurveToBake.Value.Min, CurveToBake.Value.Max);
-				ChangedCurveEditInfos.Add(GetViewModelForCurve(CurveToBake.Key)->CurveInfo);
+				if (CurveToBake.Value.Min != CurveToBake.Value.Max)
+				{
+					CurveToBake.Key->BakeCurve(BakeSampleRate, CurveToBake.Value.Min, CurveToBake.Value.Max);
+					ChangedCurveEditInfos.Add(GetViewModelForCurve(CurveToBake.Key)->CurveInfo);
+				}
+				else
+				{
+					UE_LOG(LogCurveEditor, Warning, TEXT("Unable to bake single-point curve. Check if you don't have a single key selected before Baking."));
+				}
 			}
 		}
 		else
diff --git a/Engine/Source/Editor/UnrealEd/Private/Settings/SettingsClasses.cpp b/Engine/Source/Editor/UnrealEd/Private/Settings/SettingsClasses.cpp
--- a/Engine/Source/Editor/UnrealEd/Private/Settings/SettingsClasses.cpp
+++ b/Engine/Source/Editor/UnrealEd/Private/Settings/SettingsClasses.cpp
@@ -520,7 +520,6 @@ void ULevelEditorViewportSettings::PostEditChangeProperty( struct FPropertyChang
 UProjectPackagingSettings::UProjectPackagingSettings( const FObjectInitializer& ObjectInitializer )
 	: Super(ObjectInitializer)
 {
-	bWarnIfPackagedWithoutNativizationFlag = true;
 }
 
 
@@ -638,36 +637,6 @@ void UProjectPackagingSettings::PostEditChangeProperty( FPropertyChangedEvent& P
 			}
 		}
 	}
-	else if (Name == FName((TEXT("BlueprintNativizationMethod"))))
-	{
-		IProjectManager& ProjManager = IProjectManager::Get();
-		{
-			// NOTE: these are hardcoded to match the path constructed by AddBlueprintPluginPathArgument() on CookCommand.Automation.cs, and the defaults in 
-			//       FBlueprintNativeCodeGenPaths::GetDefaultCodeGenPaths(); if you alter this (or either of those) then you need to update the others
-			const FString NativizedPluginDir  = TEXT("./Intermediate/Plugins");
-			const FString NativizedPluginName = TEXT("NativizedAssets");
-			const FString FullPluginPath = FPaths::ConvertRelativePathToFull( FPaths::ConvertRelativePathToFull(FPaths::ProjectDir()), NativizedPluginDir );
-
-			if (BlueprintNativizationMethod == EProjectPackagingBlueprintNativizationMethod::Disabled)
-			{
-				
-				ProjManager.UpdateAdditionalPluginDirectory(FullPluginPath, /*bAddOrRemove =*/false);
-				FText PluginDisableFailure;
-				ProjManager.SetPluginEnabled(NativizedPluginName, /*bEnabled =*/false, PluginDisableFailure);
-			}
-			else
-			{
-				ProjManager.UpdateAdditionalPluginDirectory(FullPluginPath, /*bAddOrRemove =*/true);
-				// plugin is enabled by default, so let's remove it from the uproject list entirely (else it causes problem in the packaged project)
-				// SetPluginEnabled() will only remove it if the plugin exists (it may not yet), so we rely on this explicit removal
-				FText PluginEnableFailure;
-				ProjManager.RemovePluginReference(NativizedPluginName, PluginEnableFailure);
-			}
-
-			FText SaveFailure;
-			ProjManager.SaveCurrentProjectToDisk(SaveFailure);
-		}
-	}
 	else if (Name == FName((TEXT("NativizeBlueprintAssets"))))
 	{
 		int32 AssetIndex;
diff --git a/Engine/Source/Editor/UnrealEd/Private/UnrealEdMisc.cpp b/Engine/Source/Editor/UnrealEd/Private/UnrealEdMisc.cpp
--- a/Engine/Source/Editor/UnrealEd/Private/UnrealEdMisc.cpp
+++ b/Engine/Source/Editor/UnrealEd/Private/UnrealEdMisc.cpp
@@ -88,6 +88,10 @@
 DEFINE_LOG_CATEGORY_STATIC(LogUnrealEdMisc, Log, All);
 
 bool FTickableEditorObject::bCollectionIntact = true;
+bool FTickableEditorObject::bIsTickingObjects = false;
+FTickableEditorObject::TTickableObjectsCollection FTickableEditorObject::TickableObjects;
+TArray<FTickableEditorObject*> FTickableEditorObject::PendingTickableObjects;
+
 
 namespace
 {
diff --git a/Engine/Source/Editor/UnrealEd/Public/EdGraphUtilities.h b/Engine/Source/Editor/UnrealEd/Public/EdGraphUtilities.h
--- a/Engine/Source/Editor/UnrealEd/Public/EdGraphUtilities.h
+++ b/Engine/Source/Editor/UnrealEd/Public/EdGraphUtilities.h
@@ -132,13 +132,13 @@ public:
 	};
 
 	/** returns true if the ParameterName is marked as a SetParam in the relevant Function */
-	static bool IsSetParam(const UFunction* Function, const FString& ParameterName);
+	static bool IsSetParam(const UFunction* Function, const FName ParameterName);
 
 	/** returns true if the ParameterName is marked as a MapParam, MapKeyParam, or MapValueParam in the relevant Function */
-	static bool IsMapParam(const UFunction* Function, const FString& ParameterName);
+	static bool IsMapParam(const UFunction* Function, const FName ParameterName);
 
 	/** returns true if the ParameterName is marked as ArrayTypeDependentParams */
-	static bool IsArrayDependentParam(const UFunction* Function, const FString& ParameterName );
+	static bool IsArrayDependentParam(const UFunction* Function, const FName ParameterName );
 
 	/** returns the first pin marked as an ArrayParam, usually returning nullptr */
 	static UEdGraphPin* FindArrayParamPin(const UFunction* Function, const UEdGraphNode* Node);
@@ -207,7 +207,7 @@ struct FWeakGraphPinPtr
 	/** Compare weak pointers for equality */ 
 	FORCEINLINE bool operator==(const FWeakGraphPinPtr &OtherPinPtr) const
 	{
-		return PinName.Equals(OtherPinPtr.PinName)
+		return PinName == OtherPinPtr.PinName
 			&& PinReference == OtherPinPtr.PinReference
 			&& NodeObjectPtr == OtherPinPtr.NodeObjectPtr;
 	}
@@ -215,7 +215,7 @@ struct FWeakGraphPinPtr
 	/** Compare weak pointers for inequality */ 
 	FORCEINLINE bool operator!=(const FWeakGraphPinPtr &OtherPinPtr) const
 	{
-		return !PinName.Equals(OtherPinPtr.PinName)
+		return PinName != OtherPinPtr.PinName
 			|| PinReference.Get() != OtherPinPtr.PinReference.Get()
 			|| NodeObjectPtr != OtherPinPtr.NodeObjectPtr;
 	}
@@ -243,7 +243,7 @@ struct FWeakGraphPinPtr
 
 private:
 	/** Pin name string */
-	FString PinName;
+	FName PinName;
 
 	/** Weak reference to the UEdGraphPin object */
 	FEdGraphPinReference PinReference;
diff --git a/Engine/Source/Editor/UnrealEd/Public/Kismet2/BlueprintEditorUtils.h b/Engine/Source/Editor/UnrealEd/Public/Kismet2/BlueprintEditorUtils.h
--- a/Engine/Source/Editor/UnrealEd/Public/Kismet2/BlueprintEditorUtils.h
+++ b/Engine/Source/Editor/UnrealEd/Public/Kismet2/BlueprintEditorUtils.h
@@ -687,8 +687,11 @@ public:
 	 * @param [out]		HiddenPins		Set of pins that should be hidden
 	 * @param [out]		OutInternalPins	Subset of hidden pins that are marked for internal use only rather than marked as hidden (optional)
 	 */
+	DEPRECATED(4.19, "Use version that passes sets by name")
 	static void GetHiddenPinsForFunction(UEdGraph const* Graph, UFunction const* Function, TSet<FString>& HiddenPins, TSet<FString>* OutInternalPins = nullptr);
 
+	static void GetHiddenPinsForFunction(UEdGraph const* Graph, UFunction const* Function, TSet<FName>& HiddenPins, TSet<FName>* OutInternalPins = nullptr);
+
 	/** Makes sure that calls to parent functions are valid, and removes them if not */
 	static void ConformCallsToParentFunctions(UBlueprint* Blueprint);
 
diff --git a/Engine/Source/Editor/UnrealEd/Public/Settings/EditorProjectSettings.h b/Engine/Source/Editor/UnrealEd/Public/Settings/EditorProjectSettings.h
--- a/Engine/Source/Editor/UnrealEd/Public/Settings/EditorProjectSettings.h
+++ b/Engine/Source/Editor/UnrealEd/Public/Settings/EditorProjectSettings.h
@@ -168,6 +168,10 @@ public:
 	UPROPERTY(EditAnywhere, config, Category=Blueprints)
 	uint32 bUseCompilationManager:1;
 	
+	/** Flag to enable faster compiles for individual blueprints if they have no function signature changes */
+	UPROPERTY(EditAnywhere, config, Category=Blueprints)
+	uint32 bSkipUnneededDependencyCompilation:1;
+
 	// UObject interface
 	virtual void PostEditChangeProperty(struct FPropertyChangedEvent& PropertyChangedEvent) override;
 	// End of UObject interface
diff --git a/Engine/Source/Editor/UnrealEd/Public/TickableEditorObject.h b/Engine/Source/Editor/UnrealEd/Public/TickableEditorObject.h
--- a/Engine/Source/Editor/UnrealEd/Public/TickableEditorObject.h
+++ b/Engine/Source/Editor/UnrealEd/Public/TickableEditorObject.h
@@ -14,30 +14,66 @@ class FTickableEditorObject : public FTickableObjectBase
 {
 public:
 
-	static void TickObjects(float DeltaSeconds)
+	static void TickObjects(const float DeltaSeconds)
 	{
-		const TArray<FTickableEditorObject*>& TickableObjects = GetTickableObjects();
+		for (FTickableEditorObject* PendingTickable : PendingTickableObjects)
+		{
+			AddTickableObject(TickableObjects, PendingTickable);
+		}
+		PendingTickableObjects.Empty();
 
-		for (int32 ObjectIndex=0; ObjectIndex < TickableObjects.Num(); ++ObjectIndex)
+		if (TickableObjects.Num() > 0)
 		{
-			FTickableEditorObject* TickableObject = TickableObjects[ObjectIndex];
-			if (TickableObject->IsTickable())
+			check(!bIsTickingObjects);
+			bIsTickingObjects = true;
+
+			bool bNeedsCleanup = false;
+
+			for (const FTickableObjectEntry& TickableEntry : TickableObjects)
+			{
+				if (FTickableObjectBase* TickableObject = TickableEntry.TickableObject)
+				{
+					if ((TickableEntry.TickType == ETickableTickType::Always) || TickableObject->IsTickable())
+					{
+						TickableObject->Tick(DeltaSeconds);
+					}
+
+					// In case it was removed during tick
+					if (TickableEntry.TickableObject == nullptr)
+					{
+						bNeedsCleanup = true;
+					}
+				}
+				else
+				{
+					bNeedsCleanup = true;
+				}
+			}
+
+			if (bNeedsCleanup)
 			{
-				TickableObject->Tick(DeltaSeconds);
+				TickableObjects.RemoveAll([](const FTickableObjectEntry& Entry) { return (Entry.TickableObject == nullptr); });
 			}
+
+			bIsTickingObjects = false;
 		}
 	}
 
 	/** Registers this instance with the static array of tickable objects. */
 	FTickableEditorObject()
 	{
-		GetTickableObjects().Add(this);
+		check(!PendingTickableObjects.Contains(this));
+		check(!TickableObjects.Contains(this));
+		PendingTickableObjects.Add(this);
 	}
 
 	/** Removes this instance from the static array of tickable objects. */
 	virtual ~FTickableEditorObject()
 	{
-		UnregisterTickableObject(this);
+		if (bCollectionIntact && PendingTickableObjects.Remove(this) == 0)
+		{
+			RemoveTickableObject(TickableObjects, this, bIsTickingObjects);
+		}
 	}
 
 private:
@@ -48,7 +84,7 @@ private:
 	 * Some tickable objects can outlive the collection
 	 * (global/static destructor order is unpredictable).
 	 */
-	class TTickableObjectsCollection : public TArray<FTickableEditorObject*>
+	class TTickableObjectsCollection : public TArray<FTickableObjectEntry>
 	{
 	public:
 		~TTickableObjectsCollection()
@@ -61,22 +97,8 @@ private:
 
 	/** True if collection of tickable objects is still intact. */
 	UNREALED_API static bool bCollectionIntact;
-
-	/** Avoids removal if the object outlived the collection. */
-	UNREALED_API static void UnregisterTickableObject(FTickableEditorObject* Obj)
-	{
-		if (bCollectionIntact)
-		{
-			const int32 Pos = GetTickableObjects().Find(Obj);
-			check(Pos!=INDEX_NONE);
-			GetTickableObjects().RemoveAt(Pos);
-		}
-	}
-
-	/** Returns the array of tickable editor objects */
-	UNREALED_API static TArray<FTickableEditorObject*>& GetTickableObjects()
-	{
-		static TTickableObjectsCollection TickableObjects;
-		return TickableObjects;
-	}
+	/** True if currently ticking of tickable editor objects. */
+	UNREALED_API static bool bIsTickingObjects;
+	UNREALED_API static TTickableObjectsCollection TickableObjects;
+	UNREALED_API static TArray<FTickableEditorObject*> PendingTickableObjects;
 };
diff --git a/Engine/Source/Editor/WorldBrowser/Private/LevelCollectionModel.h b/Engine/Source/Editor/WorldBrowser/Private/LevelCollectionModel.h
--- a/Engine/Source/Editor/WorldBrowser/Private/LevelCollectionModel.h
+++ b/Engine/Source/Editor/WorldBrowser/Private/LevelCollectionModel.h
@@ -37,9 +37,9 @@ public:
 	virtual ~FLevelCollectionModel();
 
 	/** FTickableEditorObject interface */
-	void Tick( float DeltaTime ) override;
-	bool IsTickable() const override { return true; }
-	TStatId GetStatId() const override;
+	virtual void Tick( float DeltaTime ) override;
+	virtual ETickableTickType GetTickableTickType() const override { return ETickableTickType::Always; }
+	virtual TStatId GetStatId() const override;
 	/** FTickableEditorObject interface */
 	
 	/**	@return	Whether level collection is read only now */
diff --git a/Engine/Source/Programs/AutomationTool/AutomationUtils/ProjectParams.cs b/Engine/Source/Programs/AutomationTool/AutomationUtils/ProjectParams.cs
--- a/Engine/Source/Programs/AutomationTool/AutomationUtils/ProjectParams.cs
+++ b/Engine/Source/Programs/AutomationTool/AutomationUtils/ProjectParams.cs
@@ -693,7 +693,18 @@ namespace AutomationTool
 			this.IterativeDeploy = GetParamValueIfNotSpecified(Command, IterativeDeploy, this.IterativeDeploy, new string[] {"iterativedeploy", "iterate" } );
 			this.FastCook = GetParamValueIfNotSpecified(Command, FastCook, this.FastCook, "FastCook");
 			this.IgnoreCookErrors = GetParamValueIfNotSpecified(Command, IgnoreCookErrors, this.IgnoreCookErrors, "IgnoreCookErrors");
-            this.RunAssetNativization = GetParamValueIfNotSpecified(Command, RunAssetNativization, this.RunAssetNativization, "nativizeAssets");
+
+            // Determine whether or not we're going to nativize Blueprint assets at cook time.
+            this.RunAssetNativization = false;
+            ConfigHierarchy GameIni = ConfigCache.ReadHierarchy(ConfigHierarchyType.Game, RawProjectPath.Directory, HostPlatform.Current.HostEditorPlatform);
+            if (GameIni != null)
+            {
+                string BlueprintNativizationMethod;
+                if (GameIni.TryGetValue("/Script/UnrealEd.ProjectPackagingSettings", "BlueprintNativizationMethod", out BlueprintNativizationMethod))
+                {
+                    this.RunAssetNativization = !string.IsNullOrEmpty(BlueprintNativizationMethod) && BlueprintNativizationMethod != "Disabled";
+                }
+            }
 
             string DeviceString = ParseParamValueIfNotSpecified(Command, Device, "device", String.Empty).Trim(new char[] { '\"' });
             if(DeviceString == "")
@@ -1160,19 +1171,8 @@ namespace AutomationTool
         /// <summary>
         /// Determines if Blueprint assets should be substituted with auto-generated code.
         /// </summary>
-        [Help("nativizeAssets", "Runs a \"nativization\" pass on Blueprint assets, converting then into C++ (replacing the assets with the generated source).")]
         public bool RunAssetNativization;
 
-        public struct BlueprintPluginKey
-        {
-            public bool Client;
-            public UnrealTargetPlatform TargetPlatform;
-        }
-        /// <summary>
-        /// Shared: Ref to an auto-generated plugin file that should be incorporated into the project's build
-        /// </summary>
-        public Dictionary<BlueprintPluginKey, FileReference> BlueprintPluginPaths = new Dictionary<BlueprintPluginKey, FileReference>();
-
         #endregion
 
         #region Build
diff --git a/Engine/Source/Programs/AutomationTool/Scripts/BuildCookRun.Automation.cs b/Engine/Source/Programs/AutomationTool/Scripts/BuildCookRun.Automation.cs
--- a/Engine/Source/Programs/AutomationTool/Scripts/BuildCookRun.Automation.cs
+++ b/Engine/Source/Programs/AutomationTool/Scripts/BuildCookRun.Automation.cs
@@ -215,29 +215,7 @@ public class BuildCookRun : BuildCommand
             // as well (which would be built above)
             Project.Build(this, Params, WorkingCL, ClientTargets | ProjectBuildTargets.CrashReporter);
         }
-        else
-        {
-            ConfigHierarchy GameIni = ConfigCache.ReadHierarchy(ConfigHierarchyType.Game, Params.RawProjectPath.Directory, HostPlatform.Current.HostEditorPlatform);
-            if (GameIni != null)
-            {
-                List<string> NativizeBlueprintAssets = null;
-                string BlueprintNativizationMethod = "Disabled";
-                bool bWarnIfPackagedWithoutNativizationFlag = true;
-                GameIni.GetString("/Script/UnrealEd.ProjectPackagingSettings", "BlueprintNativizationMethod", out BlueprintNativizationMethod);
-                GameIni.GetBool("/Script/UnrealEd.ProjectPackagingSettings", "bWarnIfPackagedWithoutNativizationFlag", out bWarnIfPackagedWithoutNativizationFlag);
-                GameIni.GetArray("/Script/UnrealEd.ProjectPackagingSettings", "NativizeBlueprintAssets", out NativizeBlueprintAssets);
-
-                if (bWarnIfPackagedWithoutNativizationFlag && BlueprintNativizationMethod != "Disabled")
-                {
-                    // Warn if we're cooking without the -nativizeAssets flag, when the project settings specify a nativization method.
-                    // If the "exclusive" (whitelist) method is set, we only warn if at least one asset has been selected for conversion.
-                    if (BlueprintNativizationMethod != "Exclusive" || (NativizeBlueprintAssets != null && NativizeBlueprintAssets.Count > 0))
-                    {
-                        LogWarning("Project is configured for Blueprint nativization, but the conversion flag (-nativizeAssets) was omitted from the command line. No nativized assets have been built as a result.");
-                    }
-                }
-            }
-        }
+        
 		Project.CopyBuildToStagingDirectory(Params);
 		Project.Package(Params, WorkingCL);
 		Project.Archive(Params);
diff --git a/Engine/Source/Programs/AutomationTool/Scripts/BuildProjectCommand.Automation.cs b/Engine/Source/Programs/AutomationTool/Scripts/BuildProjectCommand.Automation.cs
--- a/Engine/Source/Programs/AutomationTool/Scripts/BuildProjectCommand.Automation.cs
+++ b/Engine/Source/Programs/AutomationTool/Scripts/BuildProjectCommand.Automation.cs
@@ -38,39 +38,6 @@ public partial class Project : CommandUtils
 {
 	#region Build Command
 
-    static string GetBlueprintPluginPathArgument(ProjectParams Params, bool Client, UnrealTargetPlatform TargetPlatform)
-    {
-        string ScriptPluginArgs = "";
-
-        // if we're utilizing an auto-generated code plugin/module (a product of 
-        // the cook process), make sure to compile it along with the targets here
-
-        if (Params.RunAssetNativization)
-        {
-            ProjectParams.BlueprintPluginKey PluginKey = new ProjectParams.BlueprintPluginKey();
-            PluginKey.Client = Client;
-            PluginKey.TargetPlatform = TargetPlatform;
-            FileReference CodePlugin = null;
-            if(Params.BlueprintPluginPaths.TryGetValue(PluginKey, out CodePlugin))
-            {
-                if (FileReference.Exists(CodePlugin))
-                {
-                    ScriptPluginArgs += "-PLUGIN \"" + CodePlugin + "\" ";
-                }
-                else
-                {
-                    LogWarning("Failed to find generated plugin for Blueprint nativization [" + CodePlugin.FullName + "]. Nativization has run, but maybe there are no Blueprints which were converted?");
-                }
-            }
-            else
-            {
-                LogWarning("BlueprintPluginPath for " + TargetPlatform + " " + (Client ? "client" : "server") + " was not found");
-            }
-        }
-
-        return ScriptPluginArgs;
-    }
-
 	/// <summary>
 	/// PlatformSupportsCrashReporter
 	/// </summary>
@@ -149,9 +116,8 @@ public partial class Project : CommandUtils
 			{
                 foreach (var ClientPlatformType in UniquePlatformTypes)
 				{
-                    string ScriptPluginArgs = GetBlueprintPluginPathArgument(Params, true, ClientPlatformType);
                     CrashReportPlatforms.Add(ClientPlatformType);
-					Agenda.AddTargets(Params.ClientCookedTargets.ToArray(), ClientPlatformType, BuildConfig, Params.CodeBasedUprojectPath, InAddArgs: ScriptPluginArgs + " -remoteini=\"" + Params.RawProjectPath.Directory.FullName + "\"" + AdditionalArgs);
+					Agenda.AddTargets(Params.ClientCookedTargets.ToArray(), ClientPlatformType, BuildConfig, Params.CodeBasedUprojectPath, InAddArgs: " -remoteini=\"" + Params.RawProjectPath.Directory.FullName + "\"" + AdditionalArgs);
 				}
 			}
 		}
@@ -163,9 +129,8 @@ public partial class Project : CommandUtils
 			{
 				foreach (var ServerPlatformType in UniquePlatformTypes)
 				{
-                    string ScriptPluginArgs = GetBlueprintPluginPathArgument(Params, false, ServerPlatformType);
                     CrashReportPlatforms.Add(ServerPlatformType);
-					Agenda.AddTargets(Params.ServerCookedTargets.ToArray(), ServerPlatformType, BuildConfig, Params.CodeBasedUprojectPath, InAddArgs: ScriptPluginArgs + " -remoteini=\"" + Params.RawProjectPath.Directory.FullName + "\"" + AdditionalArgs);
+					Agenda.AddTargets(Params.ServerCookedTargets.ToArray(), ServerPlatformType, BuildConfig, Params.CodeBasedUprojectPath, InAddArgs: " -remoteini=\"" + Params.RawProjectPath.Directory.FullName + "\"" + AdditionalArgs);
 				}
 			}
 		}
diff --git a/Engine/Source/Programs/AutomationTool/Scripts/CookCommand.Automation.cs b/Engine/Source/Programs/AutomationTool/Scripts/CookCommand.Automation.cs
--- a/Engine/Source/Programs/AutomationTool/Scripts/CookCommand.Automation.cs
+++ b/Engine/Source/Programs/AutomationTool/Scripts/CookCommand.Automation.cs
@@ -21,31 +21,6 @@ public partial class Project : CommandUtils
 {
     #region Cook Command
 
-    static string AddBlueprintPluginPathArgument(ProjectParams Params, bool Client, UnrealTargetPlatform TargetPlatform, string PlatformToCook)
-    {
-        string PluginPath = "";
-
-        if (Params.RunAssetNativization)
-        {
-            // if you change or remove this placeholder value, then you should reflect those changes in the CookCommandlet (in 
-            // BlueprintNativeCodeGenManifest.cpp - where it searches and replaces this value)
-            string PlatformPlaceholderPattern = "<PLAT>";
-
-            string ProjectDir = Params.RawProjectPath.Directory.ToString();
-            // NOTE: in UProjectPackagingSettings::PostEditChangeProperty() there is a hardcoded file path/name that is set to match this; 
-            //       if you alter this path then you need to update that and likely FBlueprintNativeCodeGenPaths::GetDefaultCodeGenPaths() as well
-            PluginPath = CombinePaths(ProjectDir, "Intermediate", "Plugins", PlatformPlaceholderPattern, "NativizedAssets", "NativizedAssets.uplugin");
-
-            ProjectParams.BlueprintPluginKey PluginKey = new ProjectParams.BlueprintPluginKey();
-            PluginKey.Client = Client;
-            PluginKey.TargetPlatform = TargetPlatform;
-
-            Params.BlueprintPluginPaths.Add(PluginKey, new FileReference(PluginPath.Replace(PlatformPlaceholderPattern, PlatformToCook)));
-        }
-        return PluginPath;
-    }
-
-
     public static void Cook(ProjectParams Params)
 	{
 		if ((!Params.Cook && !(Params.CookOnTheFly && !Params.SkipServer)) || Params.SkipCook)
@@ -106,8 +81,6 @@ public partial class Project : CommandUtils
 		}
 		else
 		{
-            string NativizedPluginPath = "";
-
             var PlatformsToCook = new HashSet<string>();
             if (!Params.NoClient)
 			{
@@ -117,7 +90,6 @@ public partial class Project : CommandUtils
 					var DataPlatformDesc = Params.GetCookedDataPlatformForClientTarget(ClientPlatform);
                     string PlatformToCook = Platform.Platforms[DataPlatformDesc].GetCookPlatform(false, Params.Client);
                     PlatformsToCook.Add(PlatformToCook);
-                    NativizedPluginPath = AddBlueprintPluginPathArgument(Params, true, DataPlatformDesc.Type, PlatformToCook);
                 }
 			}
 			if (Params.DedicatedServer)
@@ -128,7 +100,6 @@ public partial class Project : CommandUtils
 					var DataPlatformDesc = Params.GetCookedDataPlatformForServerTarget(ServerPlatform);
                     string PlatformToCook = Platform.Platforms[DataPlatformDesc].GetCookPlatform(true, false);
                     PlatformsToCook.Add(PlatformToCook);
-                    NativizedPluginPath = AddBlueprintPluginPathArgument(Params, false, DataPlatformDesc.Type, PlatformToCook);
                 }
 			}
 
@@ -254,16 +225,7 @@ public partial class Project : CommandUtils
                 {
                     CommandletParams += " -compressed";
                 }
-                // we provide the option for users to run a conversion on certain (script) assets, translating them 
-                // into native source code... the cooker needs to 
-                if (Params.RunAssetNativization)
-                {
-                    CommandletParams += " -NativizeAssets";
-                    if (NativizedPluginPath.Length > 0)
-                    {
-                        CommandletParams += "=\"" + NativizedPluginPath + "\"";
-                    }
-                }
+                
                 if (Params.HasAdditionalCookerOptions)
                 {
                     string FormatedAdditionalCookerParams = Params.AdditionalCookerOptions.TrimStart(new char[] { '\"', ' ' }).TrimEnd(new char[] { '\"', ' ' });
diff --git a/Engine/Source/Programs/UnrealBuildTool/Configuration/UEBuildTarget.cs b/Engine/Source/Programs/UnrealBuildTool/Configuration/UEBuildTarget.cs
--- a/Engine/Source/Programs/UnrealBuildTool/Configuration/UEBuildTarget.cs
+++ b/Engine/Source/Programs/UnrealBuildTool/Configuration/UEBuildTarget.cs
@@ -796,9 +796,34 @@ namespace UnrealBuildTool
 			{
 				RulesObject.bDeployAfterCompile = false;
             }
-			
-			// Generate a build target from this rules module
-			UEBuildTarget BuildTarget = new UEBuildTarget(Desc, new ReadOnlyTargetRules(RulesObject), RulesAssembly, TargetFileName);
+
+            // Include generated code plugin if not building an editor target and project is configured for nativization
+            if (RulesObject.ProjectFile != null && RulesObject.Type != TargetType.Editor && ShouldIncludeNativizedAssets(RulesObject.ProjectFile.Directory))
+            {
+                string PlatformName;
+                if (RulesObject.Platform == UnrealTargetPlatform.Win32 || RulesObject.Platform == UnrealTargetPlatform.Win64)
+                {
+                    PlatformName = "Windows";
+                }
+                else
+                {
+                    PlatformName = RulesObject.Platform.ToString();
+                }
+
+                FileReference PluginFile = FileReference.Combine(RulesObject.ProjectFile.Directory, "Intermediate", "Plugins", "NativizedAssets", PlatformName, RulesObject.Type.ToString(), "NativizedAssets.uplugin");
+                if (FileReference.Exists(PluginFile))
+                {
+                    Desc.ForeignPlugins.Add(PluginFile);
+                    RulesAssembly = RulesCompiler.CreatePluginRulesAssembly(PluginFile, RulesAssembly);
+                }
+                else
+                {
+                    throw new BuildException("{0} is configured for nativization, but is missing the generated code plugin at \"{1}\". Make sure to cook {2} data before attempting to build the {3} target.", RulesObject.Name, PluginFile.FullName, RulesObject.Type.ToString(), RulesObject.Platform.ToString());
+                }
+            }
+
+            // Generate a build target from this rules module
+            UEBuildTarget BuildTarget = new UEBuildTarget(Desc, new ReadOnlyTargetRules(RulesObject), RulesAssembly, TargetFileName);
 
 			if (UnrealBuildTool.bPrintPerformanceInfo)
 			{
@@ -4738,5 +4763,23 @@ namespace UnrealBuildTool
 
 			return FilteredFileItems;
 		}
+
+        static bool ShouldIncludeNativizedAssets(DirectoryReference GameProjectDirectory)
+        {
+            ConfigHierarchy Config = ConfigCache.ReadHierarchy(ConfigHierarchyType.Game, GameProjectDirectory, BuildHostPlatform.Current.Platform);
+            if (Config != null)
+            {
+                // Determine whether or not the user has enabled nativization of Blueprint assets at cook time (default is 'Disabled')
+                string BlueprintNativizationMethod;
+                if (!Config.TryGetValue("/Script/UnrealEd.ProjectPackagingSettings", "BlueprintNativizationMethod", out BlueprintNativizationMethod))
+                {
+                    BlueprintNativizationMethod = "Disabled";
+                }
+
+                return BlueprintNativizationMethod != "Disabled";
+            }
+
+            return false;
+        }
 	}
 }
diff --git a/Engine/Source/Programs/UnrealBuildTool/ProjectFiles/ProjectFileGenerator.cs b/Engine/Source/Programs/UnrealBuildTool/ProjectFiles/ProjectFileGenerator.cs
--- a/Engine/Source/Programs/UnrealBuildTool/ProjectFiles/ProjectFileGenerator.cs
+++ b/Engine/Source/Programs/UnrealBuildTool/ProjectFiles/ProjectFileGenerator.cs
@@ -1083,27 +1083,30 @@ namespace UnrealBuildTool
 				EngineProject.AddFilesToProject(SourceFileSearch.FindFiles(UHTConfigDirectory), UnrealBuildTool.EngineDirectory);
 			}
 		}
+
+        /// <summary>
+        /// Finds any additional plugin files.
+        /// </summary>
+        /// <returns>List of additional plugin files</returns>
 		private List<FileReference> DiscoverExtraPlugins(List<FileReference> AllGameProjects)
 		{
 			List<FileReference> AddedPlugins = new List<FileReference>();
-			foreach (FileReference GameProject in AllGameProjects)
+
+            foreach (FileReference GameProject in AllGameProjects)
 			{
-				ProjectDescriptor ProjectDesc = ProjectDescriptor.FromFile(GameProject);
-				if (ProjectDesc.AdditionalPluginDirectories != null)
-				{
-					foreach (string AdditionalPluginDirectory in ProjectDesc.AdditionalPluginDirectories)
-					{
-						// @TODO: Right now, we're only including additional plugins that are still inside a game directory.
-						//        This is so FindProjectForModule() succeeds in matching up a project with the plugin. If 
-						//        the plugin is outside of a game directory, then we'll need to add special handling to AddProjectsForAllModules()
-						//        and generate a standalone project, since multiple games could share the same plugin
-						DirectoryReference PluginDir = DirectoryReference.Combine(GameProject.Directory, AdditionalPluginDirectory);
-						if (PluginDir.IsUnderDirectory(GameProject.Directory))
-						{
-							AddedPlugins.AddRange(Plugins.EnumeratePlugins(PluginDir));
-						}
-					}
-				}
+                // Check the user preference to see if they'd like to include nativized assets as a generated project.
+                bool bIncludeNativizedAssets = false;
+                ConfigHierarchy Config = ConfigCache.ReadHierarchy(ConfigHierarchyType.Game, GameProject.Directory, BuildHostPlatform.Current.Platform);
+                if (Config != null)
+                {
+                    Config.TryGetValue("/Script/UnrealEd.ProjectPackagingSettings", "bIncludeNativizedAssetsInProjectGeneration", out bIncludeNativizedAssets);
+                }
+
+                // Note: Whether or not we include nativized assets here has no bearing on whether or not they actually get built.
+                if (bIncludeNativizedAssets)
+                {
+                    AddedPlugins.AddRange(Plugins.EnumeratePlugins(DirectoryReference.Combine(GameProject.Directory, "Intermediate", "Plugins")).Where(x => x.GetFileNameWithoutExtension() == "NativizedAssets"));
+                }
 			}
 			return AddedPlugins;
 		}
diff --git a/Engine/Source/Programs/UnrealHeaderTool/Private/CodeGenerator.cpp b/Engine/Source/Programs/UnrealHeaderTool/Private/CodeGenerator.cpp
--- a/Engine/Source/Programs/UnrealHeaderTool/Private/CodeGenerator.cpp
+++ b/Engine/Source/Programs/UnrealHeaderTool/Private/CodeGenerator.cpp
@@ -3299,7 +3299,7 @@ void FNativeClassHeaderGenerator::ExportGeneratedEnumInitCode(FOutputDevice& Out
 		TEXT("\tstatic FCompiledInDeferEnum Z_CompiledInDeferEnum_UEnum_%s(%s_StaticEnum, TEXT(\"%s\"), TEXT(\"%s\"), %s, %s, %s);\r\n"),
 		*EnumNameCpp,
 		*EnumNameCpp,
-		*Enum->GetOutermost()->GetName(),
+		bIsDynamic ? *FClass::GetTypePackageName(Enum) : *Enum->GetOutermost()->GetName(),
 		*OverriddenEnumNameCpp,
 		bIsDynamic ? TEXT("true") : TEXT("false"),
 		bIsDynamic ? *AsTEXT(FClass::GetTypePackageName(Enum)) : TEXT("nullptr"),
diff --git a/Engine/Source/Runtime/AIModule/Classes/BehaviorTree/BTCompositeNode.h b/Engine/Source/Runtime/AIModule/Classes/BehaviorTree/BTCompositeNode.h
--- a/Engine/Source/Runtime/AIModule/Classes/BehaviorTree/BTCompositeNode.h
+++ b/Engine/Source/Runtime/AIModule/Classes/BehaviorTree/BTCompositeNode.h
@@ -170,28 +170,34 @@ class AIMODULE_API UBTCompositeNode : public UBTNode
 	/** is child execution allowed by decorators? */
 	bool DoDecoratorsAllowExecution(UBehaviorTreeComponent& OwnerComp, int32 InstanceIdx, int32 ChildIdx) const;
 
+	bool IsApplyingDecoratorScope() const;
+
 protected:
 
-	/** execution index of last node in child branches */
-	uint16 LastExecutionIndex;
+	/** if set, all decorators in branch below will be removed when execution flow leaves (decorators on this node are not affected) */
+	UPROPERTY(EditAnywhere, Category = Composite)
+	uint32 bApplyDecoratorScope : 1;
 
 	/** if set, NotifyChildExecution will be called */
-	uint8 bUseChildExecutionNotify : 1;
+	uint32 bUseChildExecutionNotify : 1;
 
 	/** if set, NotifyNodeActivation will be called */
-	uint8 bUseNodeActivationNotify : 1;
+	uint32 bUseNodeActivationNotify : 1;
 
 	/** if set, NotifyNodeDeactivation will be called */
-	uint8 bUseNodeDeactivationNotify : 1;
+	uint32 bUseNodeDeactivationNotify : 1;
 
 	/** if set, CanNotifyDecoratorsOnActivation will be called */
-	uint8 bUseDecoratorsActivationCheck : 1;
+	uint32 bUseDecoratorsActivationCheck : 1;
 
 	/** if set, CanNotifyDecoratorsOnDeactivation will be called */
-	uint8 bUseDecoratorsDeactivationCheck : 1;
+	uint32 bUseDecoratorsDeactivationCheck : 1;
 
 	/** if set, CanNotifyDecoratorsOnFailedActivation will be called */
-	uint8 bUseDecoratorsFailedActivationCheck : 1;
+	uint32 bUseDecoratorsFailedActivationCheck : 1;
+
+	/** execution index of last node in child branches */
+	uint16 LastExecutionIndex;
 
 	/** called just after child execution, allows to modify result */
 	virtual void NotifyChildExecution(UBehaviorTreeComponent& OwnerComp, uint8* NodeMemory, int32 ChildIdx, EBTNodeResult::Type& NodeResult) const;
@@ -249,3 +255,8 @@ FORCEINLINE uint16 UBTCompositeNode::GetLastExecutionIndex() const
 {
 	return LastExecutionIndex;
 }
+
+FORCEINLINE bool UBTCompositeNode::IsApplyingDecoratorScope() const
+{
+	return bApplyDecoratorScope;
+}
diff --git a/Engine/Source/Runtime/AIModule/Classes/BehaviorTree/BehaviorTreeComponent.h b/Engine/Source/Runtime/AIModule/Classes/BehaviorTree/BehaviorTreeComponent.h
--- a/Engine/Source/Runtime/AIModule/Classes/BehaviorTree/BehaviorTreeComponent.h
+++ b/Engine/Source/Runtime/AIModule/Classes/BehaviorTree/BehaviorTreeComponent.h
@@ -20,6 +20,7 @@ class UBTNode;
 class UBTTask_RunBehavior;
 class UBTTask_RunBehaviorDynamic;
 class UBTTaskNode;
+struct FScopedBehaviorTreeLock;
 
 struct FBTNodeExecutionInfo
 {
@@ -58,17 +59,11 @@ struct FBTPendingExecutionInfo
 	/** if set, request can't be executed */
 	uint32 bLocked : 1;
 
-	/** if set, active task is in the middle of Abort call */
-	uint32 bAborting : 1;
-
-	FBTPendingExecutionInfo() : NextTask(NULL), bOutOfNodes(false), bLocked(false), bAborting(false) {}
+	FBTPendingExecutionInfo() : NextTask(NULL), bOutOfNodes(false), bLocked(false) {}
 	bool IsSet() const { return (NextTask || bOutOfNodes) && !bLocked; }
 
 	void Lock() { bLocked = true; }
 	void Unlock() { bLocked = false; }
-
-	void OnAbortStart() { bAborting = true; }
-	void OnAbortProcessed() { bAborting = false; }
 };
 
 struct FBTTreeStartInfo
@@ -149,8 +144,11 @@ public:
 	/** unregister all aux nodes less important than given index */
 	void UnregisterAuxNodesUpTo(const FBTNodeIndex& Index);
 
+	/** unregister all aux nodes between given execution index range: FromIndex < AuxIndex < ToIndex */
+	void UnregisterAuxNodesInRange(const FBTNodeIndex& FromIndex, const FBTNodeIndex& ToIndex);
+
 	/** unregister all aux nodes in branch of tree */
-	void UnregisterAuxNodesInBranch(const UBTCompositeNode* Node);
+	void UnregisterAuxNodesInBranch(const UBTCompositeNode* Node, bool bApplyImmediately = true);
 
 	/** BEGIN UActorComponent overrides */
 	virtual void TickComponent(float DeltaTime, enum ELevelTick TickType, FActorComponentTickFunction *ThisTickFunction) override;
@@ -268,6 +266,12 @@ protected:
 	/** index of last active instance on stack */
 	uint16 ActiveInstanceIdx;
 
+	/** if set, StopTree calls will be deferred */
+	uint8 StopTreeLock;
+
+	/** if set, StopTree will be called at the end of tick */
+	uint8 bDeferredStopTree : 1;
+
 	/** loops tree execution */
 	uint8 bLoopExecution : 1;
 
@@ -335,6 +339,9 @@ protected:
 	/** apply pending tree initialization */
 	void ProcessPendingInitialize();
 
+	/** restore state of tree to state before search */
+	void RollbackSearchChanges();
+
 	/** make a snapshot for debugger */
 	void StoreDebuggerExecutionStep(EBTExecutionSnap::Type SnapType);
 
@@ -370,6 +377,7 @@ protected:
 	friend UBTTask_RunBehaviorDynamic;
 	friend FBehaviorTreeDebugger;
 	friend FBehaviorTreeInstance;
+	friend FScopedBehaviorTreeLock;
 };
 
 //////////////////////////////////////////////////////////////////////////
diff --git a/Engine/Source/Runtime/AIModule/Classes/BehaviorTree/BehaviorTreeTypes.h b/Engine/Source/Runtime/AIModule/Classes/BehaviorTree/BehaviorTreeTypes.h
--- a/Engine/Source/Runtime/AIModule/Classes/BehaviorTree/BehaviorTreeTypes.h
+++ b/Engine/Source/Runtime/AIModule/Classes/BehaviorTree/BehaviorTreeTypes.h
@@ -409,12 +409,18 @@ struct FBehaviorTreeSearchData
 	/** search unique number */
 	int32 SearchId;
 
+	/** active instance index to rollback to */
+	int32 RollbackInstanceIdx;
+
 	/** if set, current search will be restarted in next tick */
 	uint32 bPostponeSearch : 1;
 
 	/** set when task search is in progress */
 	uint32 bSearchInProgress : 1;
 
+	/** if set, active node state/memory won't be rolled back */
+	uint32 bPreserveActiveNodeMemoryOnRollback : 1;
+
 	/** adds update info to PendingUpdates array, removing all previous updates for this node */
 	void AddUniqueUpdate(const FBehaviorTreeSearchUpdate& UpdateInfo);
 
@@ -425,7 +431,7 @@ struct FBehaviorTreeSearchData
 	void Reset();
 
 	FBehaviorTreeSearchData(UBehaviorTreeComponent& InOwnerComp) 
-		: OwnerComp(InOwnerComp), bPostponeSearch(false), bSearchInProgress(false)
+		: OwnerComp(InOwnerComp), RollbackInstanceIdx(INDEX_NONE), bPostponeSearch(false), bSearchInProgress(false), bPreserveActiveNodeMemoryOnRollback(false)
 	{}
 
 private:
diff --git a/Engine/Source/Runtime/AIModule/Classes/BehaviorTree/Composites/BTComposite_SimpleParallel.h b/Engine/Source/Runtime/AIModule/Classes/BehaviorTree/Composites/BTComposite_SimpleParallel.h
--- a/Engine/Source/Runtime/AIModule/Classes/BehaviorTree/Composites/BTComposite_SimpleParallel.h
+++ b/Engine/Source/Runtime/AIModule/Classes/BehaviorTree/Composites/BTComposite_SimpleParallel.h
@@ -50,13 +50,13 @@ struct FBTParallelMemory : public FBTCompositeMemory
  * Simple Parallel composite node.
  * Allows for running two children: one which must be a single task node (with optional decorators), and the other of which can be a complete subtree.
  */
-UCLASS()
+UCLASS(HideCategories=(Composite))
 class AIMODULE_API UBTComposite_SimpleParallel : public UBTCompositeNode
 {
 	GENERATED_UCLASS_BODY()
 
 	/** how background tree should be handled when main task finishes execution */
-	UPROPERTY(EditInstanceOnly, Category=Parallel)
+	UPROPERTY(EditInstanceOnly, Category = Parallel)
 	TEnumAsByte<EBTParallelMode::Type> FinishMode;
 
 	/** handle child updates */
diff --git a/Engine/Source/Runtime/AIModule/Classes/EnvironmentQuery/EnvQueryDebugHelpers.h b/Engine/Source/Runtime/AIModule/Classes/EnvironmentQuery/EnvQueryDebugHelpers.h
--- a/Engine/Source/Runtime/AIModule/Classes/EnvironmentQuery/EnvQueryDebugHelpers.h
+++ b/Engine/Source/Runtime/AIModule/Classes/EnvironmentQuery/EnvQueryDebugHelpers.h
@@ -167,23 +167,30 @@ public:
 
 #if ENABLE_VISUAL_LOG && USE_EQS_DEBUGGER
 	static void LogQuery(FEnvQueryInstance& Query, const FLogCategoryBase& Category, ELogVerbosity::Type Verbosity);
+	static void LogQuery(FEnvQueryInstance& Query, const FName& CategoryName, ELogVerbosity::Type Verbosity);
 
 private:
-	static void LogQueryInternal(FEnvQueryInstance& Query, const FLogCategoryBase& Category, ELogVerbosity::Type Verbosity, float TimeSeconds, FVisualLogEntry *CurrentEntry);
+	static void LogQueryInternal(FEnvQueryInstance& Query, const FName& CategoryName, ELogVerbosity::Type Verbosity, float TimeSeconds, FVisualLogEntry *CurrentEntry);
 #endif
 };
 
 #if ENABLE_VISUAL_LOG && USE_EQS_DEBUGGER
 inline void UEnvQueryDebugHelpers::LogQuery(FEnvQueryInstance& Query, const FLogCategoryBase& Category, ELogVerbosity::Type Verbosity)
+{
+	const FName CategoryName = Category.GetCategoryName(); 
+	LogQuery(Query, CategoryName, Verbosity);
+}
+
+inline void UEnvQueryDebugHelpers::LogQuery(FEnvQueryInstance& Query, const FName& CategoryName, ELogVerbosity::Type Verbosity)
 {
 	UWorld *World = nullptr;
 	FVisualLogEntry *CurrentEntry = nullptr;
-	if (FVisualLogger::CheckVisualLogInputInternal(Query.Owner.Get(), Category, Verbosity, &World, &CurrentEntry) == false)
+	if (FVisualLogger::CheckVisualLogInputInternal(Query.Owner.Get(), CategoryName, Verbosity, &World, &CurrentEntry) == false)
 	{
 		return;
 	}
 
-	LogQueryInternal(Query, Category, Verbosity, World->TimeSeconds, CurrentEntry);
+	LogQueryInternal(Query, CategoryName, Verbosity, World->TimeSeconds, CurrentEntry);
 }
 #endif
 
diff --git a/Engine/Source/Runtime/AIModule/Classes/EnvironmentQuery/EnvQueryManager.h b/Engine/Source/Runtime/AIModule/Classes/EnvironmentQuery/EnvQueryManager.h
--- a/Engine/Source/Runtime/AIModule/Classes/EnvironmentQuery/EnvQueryManager.h
+++ b/Engine/Source/Runtime/AIModule/Classes/EnvironmentQuery/EnvQueryManager.h
@@ -168,7 +168,7 @@ class AIMODULE_API UEnvQueryManager : public UObject, public FTickableGameObject
 	virtual void Tick(float DeltaTime) override;
 
 	/** [FTickableGameObject] always tick, unless it's the default object */
-	virtual bool IsTickable() const override { return HasAnyFlags(RF_ClassDefaultObject) == false; }
+	virtual ETickableTickType GetTickableTickType() const override { return (HasAnyFlags(RF_ClassDefaultObject) ? ETickableTickType::Never : ETickableTickType::Always); }
 
 	/** [FTickableGameObject] tick stats */
 	virtual TStatId GetStatId() const override;
diff --git a/Engine/Source/Runtime/AIModule/Classes/EnvironmentQuery/EnvQueryTypes.h b/Engine/Source/Runtime/AIModule/Classes/EnvironmentQuery/EnvQueryTypes.h
--- a/Engine/Source/Runtime/AIModule/Classes/EnvironmentQuery/EnvQueryTypes.h
+++ b/Engine/Source/Runtime/AIModule/Classes/EnvironmentQuery/EnvQueryTypes.h
@@ -205,6 +205,7 @@ enum class EAIParamType : uint8
 	Float,
 	Int,
 	Bool,
+	MAX UMETA(Hidden)
 };
 
 UENUM()
@@ -278,6 +279,10 @@ struct AIMODULE_API FEnvNamedValue
 
 	UPROPERTY(EditAnywhere, BlueprintReadWrite, Category=Param)
 	float Value;
+
+	FEnvNamedValue()
+		: ParamType(EAIParamType::Float), Value(0.f)
+	{}
 };
 
 USTRUCT()
@@ -1251,6 +1256,8 @@ struct AIMODULE_API FAIDynamicParam
 
 	FAIDynamicParam()
 	{
+		ParamType = EAIParamType::Float;
+		Value = 0.f;
 		BBKey.AllowNoneAsValue(true);
 	}
 
diff --git a/Engine/Source/Runtime/AIModule/Classes/Perception/AIPerceptionSystem.h b/Engine/Source/Runtime/AIModule/Classes/Perception/AIPerceptionSystem.h
--- a/Engine/Source/Runtime/AIModule/Classes/Perception/AIPerceptionSystem.h
+++ b/Engine/Source/Runtime/AIModule/Classes/Perception/AIPerceptionSystem.h
@@ -45,7 +45,7 @@ public:
 	virtual void Tick(float DeltaTime) override;
 
 	/** [FTickableGameObject] always tick, unless it's the default object */
-	virtual bool IsTickable() const override { return HasAnyFlags(RF_ClassDefaultObject) == false; }
+	virtual ETickableTickType GetTickableTickType() const override { return (HasAnyFlags(RF_ClassDefaultObject) ? ETickableTickType::Never : ETickableTickType::Always); }
 
 	/** [FTickableGameObject] tick stats */
 	virtual TStatId GetStatId() const override;
diff --git a/Engine/Source/Runtime/AIModule/Private/BehaviorTree/BTCompositeNode.cpp b/Engine/Source/Runtime/AIModule/Private/BehaviorTree/BTCompositeNode.cpp
--- a/Engine/Source/Runtime/AIModule/Private/BehaviorTree/BTCompositeNode.cpp
+++ b/Engine/Source/Runtime/AIModule/Private/BehaviorTree/BTCompositeNode.cpp
@@ -10,6 +10,7 @@
 UBTCompositeNode::UBTCompositeNode(const FObjectInitializer& ObjectInitializer) : Super(ObjectInitializer)
 {
 	NodeName = "UnknownComposite";
+	bApplyDecoratorScope = false;
 	bUseChildExecutionNotify = false;
 	bUseNodeActivationNotify = false;
 	bUseNodeDeactivationNotify = false;
@@ -177,6 +178,35 @@ void UBTCompositeNode::OnNodeDeactivation(FBehaviorTreeSearchData& SearchData, E
 	{
 		SearchData.AddUniqueUpdate(FBehaviorTreeSearchUpdate(Services[ServiceIndex], SearchData.OwnerComp.GetActiveInstanceIdx(), EBTNodeUpdateMode::Remove));
 	}
+
+	// optional: remove all decorators if execution flow leaves this composite
+	if (bApplyDecoratorScope)
+	{
+		const uint16 InstanceIdx = SearchData.OwnerComp.GetActiveInstanceIdx();
+		const FBTNodeIndex FromIndex(InstanceIdx, GetExecutionIndex());
+		const FBTNodeIndex ToIndex(InstanceIdx, GetLastExecutionIndex());
+
+		SearchData.OwnerComp.UnregisterAuxNodesInRange(FromIndex, ToIndex);
+
+		// remove all pending updates "Add"
+		for (int32 Idx = SearchData.PendingUpdates.Num() - 1; Idx >= 0; Idx--)
+		{
+			const FBehaviorTreeSearchUpdate& UpdateInfo = SearchData.PendingUpdates[Idx];
+			if (UpdateInfo.Mode == EBTNodeUpdateMode::Add)
+			{
+				const uint16 UpdateNodeIdx = UpdateInfo.AuxNode ? UpdateInfo.AuxNode->GetExecutionIndex() : UpdateInfo.TaskNode->GetExecutionIndex();
+				const FBTNodeIndex UpdateIdx(UpdateInfo.InstanceIndex, UpdateNodeIdx);
+
+				if (FromIndex.TakesPriorityOver(UpdateIdx) && UpdateIdx.TakesPriorityOver(ToIndex))
+				{
+					UE_VLOG(SearchData.OwnerComp.GetOwner(), LogBehaviorTree, Verbose, TEXT("Search node update[canceled]: %s"),
+						*UBehaviorTreeTypes::DescribeNodeHelper(UpdateInfo.AuxNode ? (UBTNode*)UpdateInfo.AuxNode : (UBTNode*)UpdateInfo.TaskNode));
+
+					SearchData.PendingUpdates.RemoveAt(Idx);
+				}
+			}
+		}
+	}
 }
 
 void UBTCompositeNode::OnNodeRestart(FBehaviorTreeSearchData& SearchData) const
diff --git a/Engine/Source/Runtime/AIModule/Private/BehaviorTree/BehaviorTreeComponent.cpp b/Engine/Source/Runtime/AIModule/Private/BehaviorTree/BehaviorTreeComponent.cpp
--- a/Engine/Source/Runtime/AIModule/Private/BehaviorTree/BehaviorTreeComponent.cpp
+++ b/Engine/Source/Runtime/AIModule/Private/BehaviorTree/BehaviorTreeComponent.cpp
@@ -18,6 +18,29 @@
 int32 UBehaviorTreeComponent::ActiveDebuggerCounter = 0;
 #endif
 
+struct FScopedBehaviorTreeLock
+{
+	FScopedBehaviorTreeLock(UBehaviorTreeComponent& InOwnerComp, uint8 InLockFlag) : OwnerComp(InOwnerComp), LockFlag(InLockFlag)
+	{
+		OwnerComp.StopTreeLock |= LockFlag;
+	}
+
+	~FScopedBehaviorTreeLock()
+	{
+		OwnerComp.StopTreeLock &= ~LockFlag;
+	}
+
+	enum
+	{
+		LockTick = 1 << 0,
+		LockReentry = 1 << 1,
+	};
+
+private:
+	UBehaviorTreeComponent& OwnerComp;
+	uint8 LockFlag;
+};
+
 //----------------------------------------------------------------------//
 // UBehaviorTreeComponent
 //----------------------------------------------------------------------//
@@ -27,6 +50,8 @@ UBehaviorTreeComponent::UBehaviorTreeComponent(const FObjectInitializer& ObjectI
 	, SearchData(*this)
 {
 	ActiveInstanceIdx = 0;
+	StopTreeLock = 0;
+	bDeferredStopTree = false;
 	bLoopExecution = false;
 	bWaitingForAbortingTasks = false;
 	bRequestedFlowUpdate = false;
@@ -185,6 +210,13 @@ void UBehaviorTreeComponent::ProcessPendingInitialize()
 
 void UBehaviorTreeComponent::StopTree(EBTStopMode::Type StopMode)
 {
+	if (StopTreeLock)
+	{
+		bDeferredStopTree = true;
+		return;
+	}
+
+	FScopedBehaviorTreeLock(*this, FScopedBehaviorTreeLock::LockReentry);
 	if (!bRequestedStop)
 	{
 		bRequestedStop = true;
@@ -314,6 +346,7 @@ void UBehaviorTreeComponent::StopTree(EBTStopMode::Type StopMode)
 	bRequestedStop = false;
 	bIsRunning = false;
 	bWaitingForAbortingTasks = false;
+	bDeferredStopTree = false;
 }
 
 void UBehaviorTreeComponent::RestartTree()
@@ -416,6 +449,19 @@ void UBehaviorTreeComponent::OnTaskFinished(const UBTTaskNode* TaskNode, EBTNode
 			}
 			else
 			{
+				// force new search if there were any execution requests while waiting for aborting task
+				if (ExecutionRequest.ExecuteNode)
+				{
+					UE_VLOG(GetOwner(), LogBehaviorTree, Log, TEXT("> found valid ExecutionRequest, locking PendingExecution data to force new search!"));
+					PendingExecution.Lock();
+
+					if (ExecutionRequest.SearchEnd.IsSet())
+					{
+						UE_VLOG(GetOwner(), LogBehaviorTree, Log, TEXT("> removing limit from end of search range! [abort done]"));
+						ExecutionRequest.SearchEnd = FBTNodeIndex();
+					}
+				}
+
 				ScheduleExecutionUpdate();
 			}
 		}
@@ -890,7 +936,7 @@ void UBehaviorTreeComponent::RequestExecution(UBTCompositeNode* RequestedOn, int
 
 	ExecutionRequest.SearchStart = ExecutionIdx;
 	ExecutionRequest.ContinueWithResult = ContinueWithResult;
-	ExecutionRequest.bTryNextChild = !bSwitchToHigherPriority && !PendingExecution.bAborting;
+	ExecutionRequest.bTryNextChild = !bSwitchToHigherPriority;
 	ExecutionRequest.bIsRestart = (RequestedBy != GetActiveNode());
 	PendingExecution.Lock();
 	
@@ -902,7 +948,22 @@ void UBehaviorTreeComponent::RequestExecution(UBTCompositeNode* RequestedOn, int
 		SearchData.bPostponeSearch = true;
 	}
 
-	ScheduleExecutionUpdate();
+	// latent task abort:
+	// - don't search, just accumulate requests and run them when abort is done
+	// - rollback changes from search that caused abort to ensure proper state of tree
+	const bool bIsActiveNodeAborting = InstanceStack.Num() && InstanceStack.Last().ActiveNodeType == EBTActiveNode::AbortingTask;
+	const bool bInvalidateCurrentSearch = bWaitingForAbortingTasks || bIsActiveNodeAborting;
+	const bool bScheduleNewSearch = !bWaitingForAbortingTasks;
+
+	if (bInvalidateCurrentSearch)
+	{
+		RollbackSearchChanges();
+	}
+	
+	if (bScheduleNewSearch)
+	{
+		ScheduleExecutionUpdate();
+	}
 }
 
 void UBehaviorTreeComponent::ApplySearchUpdates(const TArray<FBehaviorTreeSearchUpdate>& UpdateList, int32 NewNodeExecutionIndex, bool bPostUpdate)
@@ -1015,11 +1076,32 @@ void UBehaviorTreeComponent::ApplySearchUpdates(const TArray<FBehaviorTreeSearch
 
 void UBehaviorTreeComponent::ApplySearchData(UBTNode* NewActiveNode)
 {
+	// search is finalized, can't rollback anymore at this point
+	SearchData.RollbackInstanceIdx = INDEX_NONE;
+
 	const int32 NewNodeExecutionIndex = NewActiveNode ? NewActiveNode->GetExecutionIndex() : 0;
 
 	ApplySearchUpdates(SearchData.PendingUpdates, NewNodeExecutionIndex);
 	ApplySearchUpdates(SearchData.PendingUpdates, NewNodeExecutionIndex, true);
 
+	// tick newly added aux nodes to compensate for tick-search order changes
+	UWorld* MyWorld = GetWorld();
+	const float CurrentFrameDeltaSeconds = MyWorld ? MyWorld->GetDeltaSeconds() : 0.0f;
+
+	for (int32 Idx = 0; Idx < SearchData.PendingUpdates.Num(); Idx++)
+	{
+		const FBehaviorTreeSearchUpdate& UpdateInfo = SearchData.PendingUpdates[Idx];
+		if (UpdateInfo.Mode == EBTNodeUpdateMode::Add && UpdateInfo.AuxNode && InstanceStack.IsValidIndex(UpdateInfo.InstanceIndex))
+		{
+			FBehaviorTreeInstance& InstanceInfo = InstanceStack[UpdateInfo.InstanceIndex];
+			uint8* NodeMemory = UpdateInfo.AuxNode->GetNodeMemory<uint8>(InstanceInfo);
+
+			UpdateInfo.AuxNode->WrappedTickNode(*this, NodeMemory, CurrentFrameDeltaSeconds);
+		}
+	}
+
+	// clear update list
+	// nothing should be added during application or tick - all changes are supposed to go to ExecutionRequest accumulator first
 	SearchData.PendingUpdates.Reset();
 }
 
@@ -1056,6 +1138,20 @@ void UBehaviorTreeComponent::TickComponent(float DeltaTime, enum ELevelTick Tick
 
 	check(this != nullptr && this->IsPendingKill() == false);
 
+	// tick active auxiliary nodes (in execution order, before task)
+	// do it before processing execution request to give BP driven logic chance to accumulate execution requests
+	// newly added aux nodes are ticked as part of SearchData application
+	for (int32 InstanceIndex = 0; InstanceIndex < InstanceStack.Num(); InstanceIndex++)
+	{
+		FBehaviorTreeInstance& InstanceInfo = InstanceStack[InstanceIndex];
+		for (int32 AuxIndex = 0; AuxIndex < InstanceInfo.ActiveAuxNodes.Num(); AuxIndex++)
+		{
+			const UBTAuxiliaryNode* AuxNode = InstanceInfo.ActiveAuxNodes[AuxIndex];
+			uint8* NodeMemory = AuxNode->GetNodeMemory<uint8>(InstanceInfo);
+			AuxNode->WrappedTickNode(*this, NodeMemory, DeltaTime);
+		}
+	}
+
 	if (bRequestedFlowUpdate)
 	{
 		ProcessExecutionRequest();
@@ -1066,48 +1162,50 @@ void UBehaviorTreeComponent::TickComponent(float DeltaTime, enum ELevelTick Tick
 		return;
 	}
 
-	// tick active auxiliary nodes and parallel tasks (in execution order, before task)
-	for (int32 InstanceIndex = 0; InstanceIndex < InstanceStack.Num(); InstanceIndex++)
 	{
-		FBehaviorTreeInstance& InstanceInfo = InstanceStack[InstanceIndex];
-		for (int32 AuxIndex = 0; AuxIndex < InstanceInfo.ActiveAuxNodes.Num(); AuxIndex++)
+		FScopedBehaviorTreeLock(*this, FScopedBehaviorTreeLock::LockTick);
+
+		// tick active parallel tasks (in execution order, before task)
+		for (int32 InstanceIndex = 0; InstanceIndex < InstanceStack.Num(); InstanceIndex++)
 		{
-			const UBTAuxiliaryNode* AuxNode = InstanceInfo.ActiveAuxNodes[AuxIndex];
-			uint8* NodeMemory = AuxNode->GetNodeMemory<uint8>(InstanceInfo);
-			AuxNode->WrappedTickNode(*this, NodeMemory, DeltaTime);
+			FBehaviorTreeInstance& InstanceInfo = InstanceStack[InstanceIndex];
+			for (int32 TaskIndex = 0; TaskIndex < InstanceInfo.ParallelTasks.Num(); TaskIndex++)
+			{
+				const UBTTaskNode* ParallelTask = InstanceInfo.ParallelTasks[TaskIndex].TaskNode;
+				uint8* NodeMemory = ParallelTask->GetNodeMemory<uint8>(InstanceInfo);
+				ParallelTask->WrappedTickTask(*this, NodeMemory, DeltaTime);
+			}
 		}
 
-		for (int32 TaskIndex = 0; TaskIndex < InstanceInfo.ParallelTasks.Num(); TaskIndex++)
+		// tick active task
+		if (InstanceStack.IsValidIndex(ActiveInstanceIdx))
 		{
-			const UBTTaskNode* ParallelTask = InstanceInfo.ParallelTasks[TaskIndex].TaskNode;
-			uint8* NodeMemory = ParallelTask->GetNodeMemory<uint8>(InstanceInfo);
-			ParallelTask->WrappedTickTask(*this, NodeMemory, DeltaTime);
+			FBehaviorTreeInstance& ActiveInstance = InstanceStack[ActiveInstanceIdx];
+			if (ActiveInstance.ActiveNodeType == EBTActiveNode::ActiveTask ||
+				ActiveInstance.ActiveNodeType == EBTActiveNode::AbortingTask)
+			{
+				UBTTaskNode* ActiveTask = (UBTTaskNode*)ActiveInstance.ActiveNode;
+				uint8* NodeMemory = ActiveTask->GetNodeMemory<uint8>(ActiveInstance);
+				ActiveTask->WrappedTickTask(*this, NodeMemory, DeltaTime);
+			}
 		}
-	}
 
-	// tick active task
-	if (InstanceStack.IsValidIndex(ActiveInstanceIdx))
-	{
-		FBehaviorTreeInstance& ActiveInstance = InstanceStack[ActiveInstanceIdx];
-		if (ActiveInstance.ActiveNodeType == EBTActiveNode::ActiveTask ||
-			ActiveInstance.ActiveNodeType == EBTActiveNode::AbortingTask)
+		// tick aborting task from abandoned subtree
+		if (InstanceStack.IsValidIndex(ActiveInstanceIdx + 1))
 		{
-			UBTTaskNode* ActiveTask = (UBTTaskNode*)ActiveInstance.ActiveNode;
-			uint8* NodeMemory = ActiveTask->GetNodeMemory<uint8>(ActiveInstance);
-			ActiveTask->WrappedTickTask(*this, NodeMemory, DeltaTime);
+			FBehaviorTreeInstance& LastInstance = InstanceStack.Last();
+			if (LastInstance.ActiveNodeType == EBTActiveNode::AbortingTask)
+			{
+				UBTTaskNode* ActiveTask = (UBTTaskNode*)LastInstance.ActiveNode;
+				uint8* NodeMemory = ActiveTask->GetNodeMemory<uint8>(LastInstance);
+				ActiveTask->WrappedTickTask(*this, NodeMemory, DeltaTime);
+			}
 		}
 	}
 
-	// tick aborting task from abandoned subtree
-	if (InstanceStack.IsValidIndex(ActiveInstanceIdx + 1))
+	if (bDeferredStopTree)
 	{
-		FBehaviorTreeInstance& LastInstance = InstanceStack.Last();
-		if (LastInstance.ActiveNodeType == EBTActiveNode::AbortingTask)
-		{
-			UBTTaskNode* ActiveTask = (UBTTaskNode*)LastInstance.ActiveNode;
-			uint8* NodeMemory = ActiveTask->GetNodeMemory<uint8>(LastInstance);
-			ActiveTask->WrappedTickTask(*this, NodeMemory, DeltaTime);
-		}
+		StopTree(EBTStopMode::Safe);
 	}
 }
 
@@ -1126,14 +1224,20 @@ void UBehaviorTreeComponent::ProcessExecutionRequest()
 		return;
 	}
 
+	if (bWaitingForAbortingTasks)
+	{
+		UE_VLOG(GetOwner(), LogBehaviorTree, Verbose, TEXT("Ignoring ProcessExecutionRequest call, aborting task must finish first"));
+		return;
+	}
+
 	if (PendingExecution.IsSet())
 	{
 		ProcessPendingExecution();
 		return;
 	}
 
-	const uint16 PrevActiveInstanceIdx = ActiveInstanceIdx;
 	bool bIsSearchValid = true;
+	SearchData.RollbackInstanceIdx = ActiveInstanceIdx;
 
 	EBTNodeResult::Type NodeResult = ExecutionRequest.ContinueWithResult;
 	UBTTaskNode* NextTask = NULL;
@@ -1273,64 +1377,52 @@ void UBehaviorTreeComponent::ProcessExecutionRequest()
 			}
 		}
 
-		if (bIsSearchValid)
-		{
-			// collect all aux nodes that have lower priority than new task
-			// occurs when normal execution is forced to revisit lower priority nodes (e.g. loop decorator)
-			const FBTNodeIndex NextTaskIdx = NextTask ? FBTNodeIndex(ActiveInstanceIdx, NextTask->GetExecutionIndex()) : FBTNodeIndex(0, 0);
-			UnregisterAuxNodesUpTo(NextTaskIdx);
+		// valid search - if search requires aborting current task and that abort happens to be latent
+		// try to keep current (before search) state of tree until everything is ready for next execution
+		// - observer changes will be applied just before starting new task (ProcessPendingExecution)
+		// - memory needs to be updated as well, but this requires keeping another copy
+		//   it's easier to just discard everything on first execution request and start new search when abort finishes
 
-			// change aux nodes
-			ApplySearchData(NextTask);
-		}
-		else
+		if (!bIsSearchValid || SearchData.bPostponeSearch)
 		{
-			// rollback search
-			ActiveInstanceIdx = PrevActiveInstanceIdx;
-			CopyInstanceMemoryFromPersistent();
+			RollbackSearchChanges();
 
-			// apply new observer changes
-			ApplyDiscardedSearch();
+			UE_VLOG(GetOwner(), LogBehaviorTree, Verbose, TEXT("Search %s, reverted all changes."), !bIsSearchValid ? TEXT("is not valid") : TEXT("will be retried"));
 		}
 
 		SearchData.bSearchInProgress = false;
 		// finish timer scope
 	}
 
-	if (SearchData.bPostponeSearch)
-	{
-		ScheduleExecutionUpdate();
-	}
-	else
+	if (!SearchData.bPostponeSearch)
 	{
+		// clear request accumulator
 		ExecutionRequest = FBTNodeExecutionInfo();
-	}
 
-	// unlock execution data, can get locked again if AbortCurrentTask starts any new requests
-	PendingExecution.Unlock();
+		// unlock execution data, can get locked again if AbortCurrentTask starts any new requests
+		PendingExecution.Unlock();
 
-	if (bIsSearchValid)
-	{
-		// abort task if needed
-		if (InstanceStack.Last().ActiveNodeType == EBTActiveNode::ActiveTask)
+		if (bIsSearchValid)
 		{
-			PendingExecution.OnAbortStart();
-
-			AbortCurrentTask();
+			// abort task if needed
+			if (InstanceStack.Last().ActiveNodeType == EBTActiveNode::ActiveTask)
+			{
+				AbortCurrentTask();
+			}
 
-			PendingExecution.OnAbortProcessed();
+			// set next task to execute
+			PendingExecution.NextTask = NextTask;
+			PendingExecution.bOutOfNodes = (NextTask == NULL);
 		}
 
-		// set next task to execute
-		PendingExecution.NextTask = NextTask;
-		PendingExecution.bOutOfNodes = (NextTask == NULL);
+		ProcessPendingExecution();
 	}
 	else
 	{
-		UE_VLOG(GetOwner(), LogBehaviorTree, Verbose, TEXT("Search result is not valid, reverted all changes."));
+		// more important execution request was found
+		// stop everything and search again in next tick
+		ScheduleExecutionUpdate();
 	}
-
-	ProcessPendingExecution();
 }
 
 void UBehaviorTreeComponent::ProcessPendingExecution()
@@ -1344,6 +1436,14 @@ void UBehaviorTreeComponent::ProcessPendingExecution()
 	FBTPendingExecutionInfo SavedInfo = PendingExecution;
 	PendingExecution = FBTPendingExecutionInfo();
 
+	// collect all aux nodes that have lower priority than new task
+	// occurs when normal execution is forced to revisit lower priority nodes (e.g. loop decorator)
+	const FBTNodeIndex NextTaskIdx = SavedInfo.NextTask ? FBTNodeIndex(ActiveInstanceIdx, SavedInfo.NextTask->GetExecutionIndex()) : FBTNodeIndex(0, 0);
+	UnregisterAuxNodesUpTo(NextTaskIdx);
+
+	// change aux nodes
+	ApplySearchData(SavedInfo.NextTask);
+
 	// make sure that we don't have any additional instances on stack
 	if (InstanceStack.Num() > (ActiveInstanceIdx + 1))
 	{
@@ -1367,6 +1467,43 @@ void UBehaviorTreeComponent::ProcessPendingExecution()
 	}
 }
 
+void UBehaviorTreeComponent::RollbackSearchChanges()
+{
+	if (SearchData.RollbackInstanceIdx >= 0)
+	{
+		ActiveInstanceIdx = SearchData.RollbackInstanceIdx;
+		SearchData.RollbackInstanceIdx = INDEX_NONE;
+
+		if (SearchData.bPreserveActiveNodeMemoryOnRollback)
+		{
+			for (int32 Idx = 0; Idx < InstanceStack.Num(); Idx++)
+			{
+				FBehaviorTreeInstance& InstanceData = InstanceStack[Idx];
+				FBehaviorTreeInstanceId& InstanceInfo = KnownInstances[InstanceData.InstanceIdIndex];
+
+				const uint16 NodeMemorySize = InstanceData.ActiveNode ? InstanceData.ActiveNode->GetInstanceMemorySize() : 0;
+				if (NodeMemorySize)
+				{
+					// copy over stored data in persistent, rollback is one time action and it won't be needed anymore
+					const uint8* NodeMemory = InstanceData.ActiveNode->GetNodeMemory<uint8>(InstanceData);
+					uint8* DestMemory = InstanceInfo.InstanceMemory.GetData() + InstanceData.ActiveNode->GetMemoryOffset();
+
+					FMemory::Memcpy(DestMemory, NodeMemory, NodeMemorySize);
+				}
+
+				InstanceData.InstanceMemory = InstanceInfo.InstanceMemory;
+			}
+		}
+		else
+		{
+			CopyInstanceMemoryFromPersistent();
+		}
+
+		// apply new observer changes
+		ApplyDiscardedSearch();
+	}
+}
+
 bool UBehaviorTreeComponent::DeactivateUpTo(UBTCompositeNode* Node, uint16 NodeInstanceIdx, EBTNodeResult::Type& NodeResult)
 {
 	UBTNode* DeactivatedChild = InstanceStack[ActiveInstanceIdx].ActiveNode;
@@ -1444,35 +1581,45 @@ void UBehaviorTreeComponent::UnregisterAuxNodesUpTo(const FBTNodeIndex& Index)
 	}
 }
 
-void UBehaviorTreeComponent::UnregisterAuxNodesInBranch(const UBTCompositeNode* Node)
+void UBehaviorTreeComponent::UnregisterAuxNodesInRange(const FBTNodeIndex& FromIndex, const FBTNodeIndex& ToIndex)
+{
+	for (int32 InstanceIndex = 0; InstanceIndex < InstanceStack.Num(); InstanceIndex++)
+	{
+		FBehaviorTreeInstance& InstanceInfo = InstanceStack[InstanceIndex];
+		for (int32 AuxIndex = 0; AuxIndex < InstanceInfo.ActiveAuxNodes.Num(); AuxIndex++)
+		{
+			FBTNodeIndex AuxIdx(InstanceIndex, InstanceInfo.ActiveAuxNodes[AuxIndex]->GetExecutionIndex());
+			if (FromIndex.TakesPriorityOver(AuxIdx) && AuxIdx.TakesPriorityOver(ToIndex))
+			{
+				SearchData.AddUniqueUpdate(FBehaviorTreeSearchUpdate(InstanceInfo.ActiveAuxNodes[AuxIndex], InstanceIndex, EBTNodeUpdateMode::Remove));
+			}
+		}
+	}
+}
+
+void UBehaviorTreeComponent::UnregisterAuxNodesInBranch(const UBTCompositeNode* Node, bool bApplyImmediately)
 {
 	const int32 InstanceIdx = FindInstanceContainingNode(Node);
 	if (InstanceIdx != INDEX_NONE)
 	{
 		check(Node);
 
-		TArray<FBehaviorTreeSearchUpdate> UpdateList;
-
-		for (int32 Idx = 0; Idx < InstanceStack[InstanceIdx].ActiveAuxNodes.Num(); Idx++)
+		TArray<FBehaviorTreeSearchUpdate> UpdateListCopy;
+		if (bApplyImmediately)
 		{
-			UBTAuxiliaryNode* AuxNode = InstanceStack[InstanceIdx].ActiveAuxNodes[Idx];
-			if (AuxNode->GetTreeDepth() > Node->GetTreeDepth())
-			{
-				// check if aux node belongs to given branch
-				UBTCompositeNode* TestNode = AuxNode->GetParentNode();
-				while (TestNode && TestNode != Node)
-				{
-					TestNode = TestNode->GetParentNode();
-				}
-
-				if (TestNode == Node)
-				{
-					UpdateList.Add(FBehaviorTreeSearchUpdate(AuxNode, InstanceIdx, EBTNodeUpdateMode::Remove));
-				}
-			}
+			UpdateListCopy = SearchData.PendingUpdates;
+			SearchData.PendingUpdates.Reset();
 		}
 
-		ApplySearchUpdates(UpdateList, 0);
+		const FBTNodeIndex FromIndex(InstanceIdx, Node->GetExecutionIndex());
+		const FBTNodeIndex ToIndex(InstanceIdx, Node->GetLastExecutionIndex());
+		UnregisterAuxNodesInRange(FromIndex, ToIndex);
+
+		if (bApplyImmediately)
+		{
+			ApplySearchUpdates(SearchData.PendingUpdates, 0);
+			SearchData.PendingUpdates = UpdateListCopy;
+		}
 	}
 }
 
@@ -1536,6 +1683,12 @@ void UBehaviorTreeComponent::AbortCurrentTask()
 	// remove all observers before requesting abort
 	UnregisterMessageObserversFrom(CurrentTask);
 
+	// protect memory of this task from rollbacks
+	// at this point, invalid search rollback already happened
+	// only reason to do the rollback is restoring tree state during abort for accumulated requests
+	// but this task needs to remain unchanged: it's still aborting and internal memory can be modified on AbortTask call
+	SearchData.bPreserveActiveNodeMemoryOnRollback = true;
+
 	UE_VLOG(GetOwner(), LogBehaviorTree, Log, TEXT("Abort task: %s"), *UBehaviorTreeTypes::DescribeNodeHelper(CurrentTask));
 
 	// abort task using current state of tree
diff --git a/Engine/Source/Runtime/AIModule/Private/BehaviorTree/BehaviorTreeTypes.cpp b/Engine/Source/Runtime/AIModule/Private/BehaviorTree/BehaviorTreeTypes.cpp
--- a/Engine/Source/Runtime/AIModule/Private/BehaviorTree/BehaviorTreeTypes.cpp
+++ b/Engine/Source/Runtime/AIModule/Private/BehaviorTree/BehaviorTreeTypes.cpp
@@ -282,8 +282,10 @@ void FBehaviorTreeSearchData::Reset()
 	PendingUpdates.Reset();
 	SearchStart = FBTNodeIndex();
 	SearchEnd = FBTNodeIndex();
+	RollbackInstanceIdx = INDEX_NONE;
 	bSearchInProgress = false;
 	bPostponeSearch = false;
+	bPreserveActiveNodeMemoryOnRollback = false;
 }
 
 //----------------------------------------------------------------------//
diff --git a/Engine/Source/Runtime/AIModule/Private/BehaviorTree/Composites/BTComposite_SimpleParallel.cpp b/Engine/Source/Runtime/AIModule/Private/BehaviorTree/Composites/BTComposite_SimpleParallel.cpp
--- a/Engine/Source/Runtime/AIModule/Private/BehaviorTree/Composites/BTComposite_SimpleParallel.cpp
+++ b/Engine/Source/Runtime/AIModule/Private/BehaviorTree/Composites/BTComposite_SimpleParallel.cpp
@@ -12,6 +12,7 @@ UBTComposite_SimpleParallel::UBTComposite_SimpleParallel(const FObjectInitialize
 	bUseChildExecutionNotify = true;
 	bUseNodeDeactivationNotify = true;
 	bUseDecoratorsDeactivationCheck = true;
+	bApplyDecoratorScope = true;
 	OnNextChild.BindUObject(this, &UBTComposite_SimpleParallel::GetNextChildHandler);
 }
 
@@ -121,30 +122,6 @@ void UBTComposite_SimpleParallel::NotifyNodeDeactivation(FBehaviorTreeSearchData
 	{
 		SearchData.AddUniqueUpdate(FBehaviorTreeSearchUpdate(Children[EBTParallelChild::MainTask].ChildTask, ActiveInstanceIdx, EBTNodeUpdateMode::Remove));
 	}
-
-	// remove all active nodes from background tree
-	const FBTNodeIndex FirstBackgroundIndex(ActiveInstanceIdx, GetChildExecutionIndex(EBTParallelChild::BackgroundTree, EBTChildIndex::FirstNode));
-	const FBTNodeIndex LastBackgroundIndex(ActiveInstanceIdx, GetLastExecutionIndex());
-	SearchData.OwnerComp.UnregisterAuxNodesUpTo(FirstBackgroundIndex);
-
-	// remove all pending updates "Add" from background tree 
-	// it doesn't make sense for decorators to reactivate themselves there
-	for (int32 Idx = SearchData.PendingUpdates.Num() - 1; Idx >= 0; Idx--)
-	{
-		const FBehaviorTreeSearchUpdate& UpdateInfo = SearchData.PendingUpdates[Idx];
-		if (UpdateInfo.Mode == EBTNodeUpdateMode::Add)
-		{
-			const uint16 UpdateNodeIdx = UpdateInfo.AuxNode ? UpdateInfo.AuxNode->GetExecutionIndex() : UpdateInfo.TaskNode->GetExecutionIndex();
-			const FBTNodeIndex UpdateIdx(UpdateInfo.InstanceIndex, UpdateNodeIdx);
-			if (FirstBackgroundIndex.TakesPriorityOver(UpdateIdx) && UpdateIdx.TakesPriorityOver(LastBackgroundIndex))
-			{
-				UE_VLOG(SearchData.OwnerComp.GetOwner(), LogBehaviorTree, Verbose, TEXT("Search node update[canceled]: %s"),
-					*UBehaviorTreeTypes::DescribeNodeHelper(UpdateInfo.AuxNode ? (UBTNode*)UpdateInfo.AuxNode : (UBTNode*)UpdateInfo.TaskNode));
-
-				SearchData.PendingUpdates.RemoveAt(Idx);
-			}
-		}
-	}
 }
 
 bool UBTComposite_SimpleParallel::CanNotifyDecoratorsOnDeactivation(FBehaviorTreeSearchData& SearchData, int32 ChildIdx, EBTNodeResult::Type& NodeResult) const
diff --git a/Engine/Source/Runtime/AIModule/Private/EnvironmentQuery/EQSTestingPawn.cpp b/Engine/Source/Runtime/AIModule/Private/EnvironmentQuery/EQSTestingPawn.cpp
--- a/Engine/Source/Runtime/AIModule/Private/EnvironmentQuery/EQSTestingPawn.cpp
+++ b/Engine/Source/Runtime/AIModule/Private/EnvironmentQuery/EQSTestingPawn.cpp
@@ -161,6 +161,15 @@ void AEQSTestingPawn::PostLoad()
 
 	if (QueryParams.Num() > 0)
 	{
+		// handle legacy props reading
+		for (FEnvNamedValue& NamedValue : QueryParams)
+		{
+			if (uint8(NamedValue.ParamType) >= uint8(EAIParamType::MAX))
+			{
+				NamedValue.ParamType = EAIParamType::Float;
+			}
+		}
+
 		FAIDynamicParam::GenerateConfigurableParamsFromNamedValues(*this, QueryConfig, QueryParams);
 		QueryParams.Empty();
 	}
diff --git a/Engine/Source/Runtime/AIModule/Private/EnvironmentQuery/EnvQueryDebugHelpers.cpp b/Engine/Source/Runtime/AIModule/Private/EnvironmentQuery/EnvQueryDebugHelpers.cpp
--- a/Engine/Source/Runtime/AIModule/Private/EnvironmentQuery/EnvQueryDebugHelpers.cpp
+++ b/Engine/Source/Runtime/AIModule/Private/EnvironmentQuery/EnvQueryDebugHelpers.cpp
@@ -175,7 +175,7 @@ namespace
 	}
 }
 
-void UEnvQueryDebugHelpers::LogQueryInternal(FEnvQueryInstance& Query, const FLogCategoryBase& Category, ELogVerbosity::Type Verbosity, float TimeSeconds, FVisualLogEntry *CurrentEntry)
+void UEnvQueryDebugHelpers::LogQueryInternal(FEnvQueryInstance& Query, const FName& CategoryName, ELogVerbosity::Type Verbosity, float TimeSeconds, FVisualLogEntry *CurrentEntry)
 {
 	const int32 UniqueId = FVisualLogger::Get().GetUniqueId(TimeSeconds);
 
@@ -186,11 +186,11 @@ void UEnvQueryDebugHelpers::LogQueryInternal(FEnvQueryInstance& Query, const FLo
 	DebugDataToBlobArray(EQSLocalData, BlobArray);
 
 	const FString AdditionalLogInfo = FString::Printf(TEXT("Executed EQS: \n - Name: '%s' (id=%d, option=%d),\n - All Items: %d,\n - ValidItems: %d"), *Query.QueryName, Query.QueryID, Query.OptionIndex, Query.ItemDetails.Num(), Query.NumValidItems);
-	FVisualLogLine Line(Category.GetCategoryName(), Verbosity, AdditionalLogInfo, Query.QueryID);
+	FVisualLogLine Line(CategoryName, Verbosity, AdditionalLogInfo, Query.QueryID);
 	Line.TagName = *EVisLogTags::TAG_EQS;
 	Line.UniqueId = UniqueId;
 
-	CurrentEntry->AddDataBlock(EVisLogTags::TAG_EQS, BlobArray, Category.GetCategoryName(), Verbosity).UniqueId = UniqueId;
+	CurrentEntry->AddDataBlock(EVisLogTags::TAG_EQS, BlobArray, CategoryName, Verbosity).UniqueId = UniqueId;
 
 
 	if (EQSLocalData.NumValidItems > 0)
diff --git a/Engine/Source/Runtime/Core/Public/UObject/FrameworkObjectVersion.h b/Engine/Source/Runtime/Core/Public/UObject/FrameworkObjectVersion.h
--- a/Engine/Source/Runtime/Core/Public/UObject/FrameworkObjectVersion.h
+++ b/Engine/Source/Runtime/Core/Public/UObject/FrameworkObjectVersion.h
@@ -107,6 +107,9 @@ struct CORE_API FFrameworkObjectVersion
 		// Fix User Defined structs so that all members are correct flagged blueprint visible 
 		UserDefinedStructsBlueprintVisible,
 
+		// FMaterialInput and FEdGraphPin store their name as FName instead of FString
+		PinsStoreFName,
+
 		// -----<new versions can be added above this line>-------------------------------------------------
 		VersionPlusOne,
 		LatestVersion = VersionPlusOne - 1
diff --git a/Engine/Source/Runtime/Core/Public/UObject/NameTypes.h b/Engine/Source/Runtime/Core/Public/UObject/NameTypes.h
--- a/Engine/Source/Runtime/Core/Public/UObject/NameTypes.h
+++ b/Engine/Source/Runtime/Core/Public/UObject/NameTypes.h
@@ -1,4 +1,4 @@
-// Copyright 1998-2017 Epic Games, Inc. All Rights Reserved.
+﻿// Copyright 1998-2017 Epic Games, Inc. All Rights Reserved.
 
 #pragma once
 
@@ -889,8 +889,13 @@ public:
 	template <typename CharType>
 	bool operator==(const CharType* Other) const
 	{
+		// Make NAME_None == TEXT("") or nullptr consistent with NAME_None == FName(TEXT("")) or FName(nullptr)
+		if (Other == nullptr || Other[0] == 0)
+		{
+			return (*this == NAME_None);
+		}
+
 		// Find name entry associated with this FName.
-		check(Other);
 		const FNameEntry* const Entry = GetComparisonNameEntry();
 
 		// Temporary buffer to hold split name in case passed in name is of Name_Number format.
diff --git a/Engine/Source/Runtime/CoreUObject/Private/UObject/Enum.cpp b/Engine/Source/Runtime/CoreUObject/Private/UObject/Enum.cpp
--- a/Engine/Source/Runtime/CoreUObject/Private/UObject/Enum.cpp
+++ b/Engine/Source/Runtime/CoreUObject/Private/UObject/Enum.cpp
@@ -447,18 +447,23 @@ int32 UEnum::GetIndexByNameString(const FString& InSearchString, EGetByNameFlags
 			SearchEnumEntryString = **FoundNewEnumEntry;
 
 			// Recompute modified name
-			DoubleColonIndex = SearchEnumEntryString.Find(TEXT("::"), ESearchCase::CaseSensitive);
-			if (DoubleColonIndex == INDEX_NONE)
+			int32 NewDoubleColonIndex = SearchEnumEntryString.Find(TEXT("::"), ESearchCase::CaseSensitive);
+			if (NewDoubleColonIndex == INDEX_NONE)
 			{
 				ModifiedEnumEntryString = GenerateFullEnumName(*SearchEnumEntryString);
 			}
 			else
 			{
-				ModifiedEnumEntryString = SearchEnumEntryString.RightChop(DoubleColonIndex + 2);
+				ModifiedEnumEntryString = SearchEnumEntryString.RightChop(NewDoubleColonIndex + 2);
 			}
 		}
 	}
-		
+	else if (DoubleColonIndex != INDEX_NONE)
+	{
+		// If we didn't find a value redirect and our original string was namespaced, we need to fix the namespace now as it may have changed due to enum type redirect
+		SearchEnumEntryString = GenerateFullEnumName(*ModifiedEnumEntryString);
+	}
+
 	// Search for names both with and without namespace
 	FName SearchName = FName(*SearchEnumEntryString);
 	FName ModifiedName = FName(*ModifiedEnumEntryString);
diff --git a/Engine/Source/Runtime/CoreUObject/Private/UObject/SoftObjectPath.cpp b/Engine/Source/Runtime/CoreUObject/Private/UObject/SoftObjectPath.cpp
--- a/Engine/Source/Runtime/CoreUObject/Private/UObject/SoftObjectPath.cpp
+++ b/Engine/Source/Runtime/CoreUObject/Private/UObject/SoftObjectPath.cpp
@@ -303,7 +303,22 @@ UObject* FSoftObjectPath::ResolveObject() const
 		return nullptr;
 	}
 
-	UObject* FoundObject = FindObject<UObject>(nullptr, *ToString());
+	FString PathString = ToString();
+#if WITH_EDITOR
+	if (GPlayInEditorID != INDEX_NONE)
+	{
+		// If we are in PIE and this hasn't already been fixed up, we need to fixup at resolution time. We cannot modify the path as it may be somewhere like a blueprint CDO
+		FSoftObjectPath FixupObjectPath = *this;
+		FixupObjectPath.FixupForPIE();
+
+		if (FixupObjectPath.AssetPathName != AssetPathName)
+		{
+			PathString = FixupObjectPath.ToString();
+		}
+	}
+#endif
+
+	UObject* FoundObject = FindObject<UObject>(nullptr, *PathString);
 	while (UObjectRedirector* Redirector = Cast<UObjectRedirector>(FoundObject))
 	{
 		FoundObject = Redirector->DestinationObject;
diff --git a/Engine/Source/Runtime/CoreUObject/Private/UObject/UObjectGlobals.cpp b/Engine/Source/Runtime/CoreUObject/Private/UObject/UObjectGlobals.cpp
--- a/Engine/Source/Runtime/CoreUObject/Private/UObject/UObjectGlobals.cpp
+++ b/Engine/Source/Runtime/CoreUObject/Private/UObject/UObjectGlobals.cpp
@@ -2649,8 +2649,7 @@ FObjectInitializer::~FObjectInitializer()
 
 				FLinkerLoad* SuperClassLinker = SuperClass->GetLinker();
 				const bool bSuperLoadPending = FDeferredObjInitializerTracker::IsCdoDeferred(SuperClass) ||
-					(SuperBpCDO && SuperBpCDO->HasAnyFlags(RF_NeedLoad)) ||
-					(SuperClassLinker && SuperClassLinker->IsBlueprintFinalizationPending());
+					(SuperBpCDO && (SuperBpCDO->HasAnyFlags(RF_NeedLoad) || (SuperBpCDO->HasAnyFlags(RF_WasLoaded) && !SuperBpCDO->HasAnyFlags(RF_LoadCompleted))));
 
 				FLinkerLoad* ObjLinker = BlueprintClass->GetLinker();
 				const bool bIsBpClassSerializing    = ObjLinker && (ObjLinker->LoadFlags & LOAD_DeferDependencyLoads);
diff --git a/Engine/Source/Runtime/CoreUObject/Public/UObject/SoftObjectPtr.h b/Engine/Source/Runtime/CoreUObject/Public/UObject/SoftObjectPtr.h
--- a/Engine/Source/Runtime/CoreUObject/Public/UObject/SoftObjectPtr.h
+++ b/Engine/Source/Runtime/CoreUObject/Public/UObject/SoftObjectPtr.h
@@ -89,6 +89,19 @@ public:
 		return ToSoftObjectPath().GetAssetName();
 	}
 
+#if WITH_EDITOR
+	/** Overridden to deal with PIE lookups */
+	FORCEINLINE UObject* Get() const
+	{
+		if (GPlayInEditorID != INDEX_NONE)
+		{
+			// Cannot use or set the cached value in PIE as it may affect other PIE instances or the editor
+			return GetUniqueID().ResolveObject();
+		}
+		return TPersistentObjectPtr<FSoftObjectPath>::Get();
+	}
+#endif
+
 	using TPersistentObjectPtr<FSoftObjectPath>::operator=;
 };
 
diff --git a/Engine/Source/Runtime/Engine/Classes/AI/Navigation/NavMeshBoundsVolume.h b/Engine/Source/Runtime/Engine/Classes/AI/Navigation/NavMeshBoundsVolume.h
--- a/Engine/Source/Runtime/Engine/Classes/AI/Navigation/NavMeshBoundsVolume.h
+++ b/Engine/Source/Runtime/Engine/Classes/AI/Navigation/NavMeshBoundsVolume.h
@@ -28,6 +28,7 @@ class ANavMeshBoundsVolume : public AVolume
 #if WITH_EDITOR
 	//~ Begin UObject Interface
 	virtual void PostEditChangeProperty( struct FPropertyChangedEvent& PropertyChangedEvent) override;
+	virtual void PostEditUndo() override;
 	//~ End UObject Interface
 #endif // WITH_EDITOR
 };
