added OnTaskFinished call when behavior tree's subtree is deactivated

diff --git a/Engine/Source/Runtime/AIModule/Classes/BehaviorTree/BehaviorTreeTypes.h b/Engine/Source/Runtime/AIModule/Classes/BehaviorTree/BehaviorTreeTypes.h
--- a/Engine/Source/Runtime/AIModule/Classes/BehaviorTree/BehaviorTreeTypes.h
+++ b/Engine/Source/Runtime/AIModule/Classes/BehaviorTree/BehaviorTreeTypes.h
@@ -287,6 +287,7 @@ struct FBehaviorTreeInstanceId
 };
 
 struct FBehaviorTreeSearchData;
+DECLARE_DELEGATE_TwoParams(FBTInstanceDeactivation, UBehaviorTreeComponent&, EBTNodeResult::Type);
 
 /** data required for instance of single subtree */
 struct FBehaviorTreeInstance
@@ -312,6 +313,9 @@ struct FBehaviorTreeInstance
 	/** active node type */
 	TEnumAsByte<EBTActiveNode::Type> ActiveNodeType;
 
+	/** delegate sending a notify when tree instance is removed from active stack */
+	FBTInstanceDeactivation DeactivationNotify;
+
 	FBehaviorTreeInstance() { IncMemoryStats(); }
 	FBehaviorTreeInstance(const FBehaviorTreeInstance& Other) { *this = Other; IncMemoryStats(); }
 	FBehaviorTreeInstance(int32 MemorySize) { InstanceMemory.AddZeroed(MemorySize); IncMemoryStats(); }
@@ -390,6 +394,16 @@ struct FBehaviorTreeSearchUpdate
 	{}
 };
 
+/** instance notify data */
+struct FBehaviorTreeSearchUpdateNotify
+{
+	uint16 InstanceIndex;
+	TEnumAsByte<EBTNodeResult::Type> NodeResult;
+
+	FBehaviorTreeSearchUpdateNotify() : InstanceIndex(0), NodeResult(EBTNodeResult::Succeeded) {}
+	FBehaviorTreeSearchUpdateNotify(uint16 InInstanceIndex, EBTNodeResult::Type InNodeResult) : InstanceIndex(InInstanceIndex), NodeResult(InNodeResult) {}
+};
+
 /** node search data */
 struct FBehaviorTreeSearchData
 {
@@ -400,6 +414,9 @@ struct FBehaviorTreeSearchData
 	 *  buffered during search to prevent instant add & remove pairs */
 	TArray<FBehaviorTreeSearchUpdate> PendingUpdates;
 
+	/** notifies for tree instances */
+	TArray<FBehaviorTreeSearchUpdateNotify> PendingNotifies;
+
 	/** first node allowed in search */
 	FBTNodeIndex SearchStart;
 
diff --git a/Engine/Source/Runtime/AIModule/Classes/BehaviorTree/Tasks/BTTask_RunBehavior.h b/Engine/Source/Runtime/AIModule/Classes/BehaviorTree/Tasks/BTTask_RunBehavior.h
--- a/Engine/Source/Runtime/AIModule/Classes/BehaviorTree/Tasks/BTTask_RunBehavior.h
+++ b/Engine/Source/Runtime/AIModule/Classes/BehaviorTree/Tasks/BTTask_RunBehavior.h
@@ -42,6 +42,9 @@ protected:
 	/** behavior to run */
 	UPROPERTY(Category = Node, EditAnywhere)
 	UBehaviorTree* BehaviorAsset;
+
+	/** called when subtree is removed from active stack */
+	virtual void OnSubtreeDeactivated(UBehaviorTreeComponent& OwnerComp, EBTNodeResult::Type NodeResult);
 };
 
 //////////////////////////////////////////////////////////////////////////
diff --git a/Engine/Source/Runtime/AIModule/Classes/BehaviorTree/Tasks/BTTask_RunBehaviorDynamic.h b/Engine/Source/Runtime/AIModule/Classes/BehaviorTree/Tasks/BTTask_RunBehaviorDynamic.h
--- a/Engine/Source/Runtime/AIModule/Classes/BehaviorTree/Tasks/BTTask_RunBehaviorDynamic.h
+++ b/Engine/Source/Runtime/AIModule/Classes/BehaviorTree/Tasks/BTTask_RunBehaviorDynamic.h
@@ -47,6 +47,9 @@ protected:
 	/** current subtree */
 	UPROPERTY()
 	UBehaviorTree* BehaviorAsset;
+
+	/** called when subtree is removed from active stack */
+	virtual void OnSubtreeDeactivated(UBehaviorTreeComponent& OwnerComp, EBTNodeResult::Type NodeResult);
 };
 
 //////////////////////////////////////////////////////////////////////////
diff --git a/Engine/Source/Runtime/AIModule/Private/BehaviorTree/BehaviorTreeComponent.cpp b/Engine/Source/Runtime/AIModule/Private/BehaviorTree/BehaviorTreeComponent.cpp
--- a/Engine/Source/Runtime/AIModule/Private/BehaviorTree/BehaviorTreeComponent.cpp
+++ b/Engine/Source/Runtime/AIModule/Private/BehaviorTree/BehaviorTreeComponent.cpp
@@ -1079,6 +1079,17 @@ void UBehaviorTreeComponent::ApplySearchData(UBTNode* NewActiveNode)
 	// search is finalized, can't rollback anymore at this point
 	SearchData.RollbackInstanceIdx = INDEX_NONE;
 
+	// send all deactivation notifies for bookkeeping
+	for (int32 Idx = 0; Idx < SearchData.PendingNotifies.Num(); Idx++)
+	{
+		const FBehaviorTreeSearchUpdateNotify& NotifyInfo = SearchData.PendingNotifies[Idx];
+		if (InstanceStack.IsValidIndex(NotifyInfo.InstanceIndex))
+		{
+			InstanceStack[NotifyInfo.InstanceIndex].DeactivationNotify.ExecuteIfBound(*this, NotifyInfo.NodeResult);
+		}	
+	}
+
+	// apply changes to aux nodes and parallel tasks
 	const int32 NewNodeExecutionIndex = NewActiveNode ? NewActiveNode->GetExecutionIndex() : 0;
 
 	ApplySearchUpdates(SearchData.PendingUpdates, NewNodeExecutionIndex);
@@ -1103,6 +1114,7 @@ void UBehaviorTreeComponent::ApplySearchData(UBTNode* NewActiveNode)
 	// clear update list
 	// nothing should be added during application or tick - all changes are supposed to go to ExecutionRequest accumulator first
 	SearchData.PendingUpdates.Reset();
+	SearchData.PendingNotifies.Reset();
 }
 
 void UBehaviorTreeComponent::ApplyDiscardedSearch()
@@ -1128,6 +1140,9 @@ void UBehaviorTreeComponent::ApplyDiscardedSearch()
 
 	// remove everything else
 	SearchData.PendingUpdates.Reset();
+
+	// don't send deactivation notifies
+	SearchData.PendingNotifies.Reset();
 }
 
 void UBehaviorTreeComponent::TickComponent(float DeltaTime, enum ELevelTick TickType, FActorComponentTickFunction *ThisTickFunction)
@@ -1254,6 +1269,10 @@ void UBehaviorTreeComponent::ProcessExecutionRequest()
 			const bool bDeactivated = DeactivateUpTo(ExecutionRequest.ExecuteNode, ExecutionRequest.ExecuteInstanceIdx, NodeResult);
 			if (!bDeactivated)
 			{
+				// error occurred and tree will restart, all pending deactivation notifies will be lost
+				// this is should happen
+				SearchData.PendingUpdates.Reset();
+
 				return;
 			}
 		}
@@ -1332,6 +1351,9 @@ void UBehaviorTreeComponent::ProcessExecutionRequest()
 						StoreDebuggerRemovedInstance(ActiveInstanceIdx);
 						InstanceStack[ActiveInstanceIdx].DeactivateNodes(SearchData, ActiveInstanceIdx);
 
+						// store notify for later use if search is not reverted
+						SearchData.PendingNotifies.Add(FBehaviorTreeSearchUpdateNotify(ActiveInstanceIdx, NodeResult));
+
 						// and leave subtree
 						ActiveInstanceIdx--;
 
@@ -1552,6 +1574,9 @@ bool UBehaviorTreeComponent::DeactivateUpTo(UBTCompositeNode* Node, uint16 NodeI
 				return false;
 			}
 
+			// store notify for later use if search is not reverted
+			SearchData.PendingNotifies.Add(FBehaviorTreeSearchUpdateNotify(ActiveInstanceIdx, NodeResult));
+
 			ActiveInstanceIdx--;
 			DeactivatedChild = InstanceStack[ActiveInstanceIdx].ActiveNode;
 		}
diff --git a/Engine/Source/Runtime/AIModule/Private/BehaviorTree/BehaviorTreeTypes.cpp b/Engine/Source/Runtime/AIModule/Private/BehaviorTree/BehaviorTreeTypes.cpp
--- a/Engine/Source/Runtime/AIModule/Private/BehaviorTree/BehaviorTreeTypes.cpp
+++ b/Engine/Source/Runtime/AIModule/Private/BehaviorTree/BehaviorTreeTypes.cpp
@@ -280,6 +280,7 @@ void FBehaviorTreeSearchData::AssignSearchId()
 void FBehaviorTreeSearchData::Reset()
 {
 	PendingUpdates.Reset();
+	PendingNotifies.Reset();
 	SearchStart = FBTNodeIndex();
 	SearchEnd = FBTNodeIndex();
 	RollbackInstanceIdx = INDEX_NONE;
diff --git a/Engine/Source/Runtime/AIModule/Private/BehaviorTree/Tasks/BTTask_RunBehavior.cpp b/Engine/Source/Runtime/AIModule/Private/BehaviorTree/Tasks/BTTask_RunBehavior.cpp
--- a/Engine/Source/Runtime/AIModule/Private/BehaviorTree/Tasks/BTTask_RunBehavior.cpp
+++ b/Engine/Source/Runtime/AIModule/Private/BehaviorTree/Tasks/BTTask_RunBehavior.cpp
@@ -14,7 +14,27 @@ EBTNodeResult::Type UBTTask_RunBehavior::ExecuteTask(UBehaviorTreeComponent& Own
 	UE_CVLOG(BehaviorAsset == nullptr, OwnerComp.GetAIOwner(), LogBehaviorTree, Error, TEXT("\'%s\' is missing BehaviorAsset!"), *GetNodeName());
 
 	const bool bPushed = BehaviorAsset != nullptr && OwnerComp.PushInstance(*BehaviorAsset);
-	return bPushed ? EBTNodeResult::InProgress : EBTNodeResult::Failed;
+	if (bPushed && OwnerComp.InstanceStack.Num() > 0)
+	{
+		FBehaviorTreeInstance& MyInstance = OwnerComp.InstanceStack[OwnerComp.InstanceStack.Num() - 1];
+		MyInstance.DeactivationNotify.BindUObject(this, &UBTTask_RunBehavior::OnSubtreeDeactivated);
+		// unbinding is not required, MyInstance will be destroyed after firing that delegate (usually by UBehaviorTreeComponent::ProcessPendingExecution) 
+
+		return EBTNodeResult::InProgress;
+	}
+
+	return EBTNodeResult::Failed;
+}
+
+void UBTTask_RunBehavior::OnSubtreeDeactivated(UBehaviorTreeComponent& OwnerComp, EBTNodeResult::Type NodeResult)
+{
+	const int32 MyInstanceIdx = OwnerComp.FindInstanceContainingNode(this);
+	uint8* NodeMemory = OwnerComp.GetNodeMemory(this, MyInstanceIdx);
+
+	UE_VLOG(OwnerComp.GetOwner(), LogBehaviorTree, Verbose, TEXT("OnSubtreeDeactivated: %s (result: %s)"),
+		*UBehaviorTreeTypes::DescribeNodeHelper(this), *UBehaviorTreeTypes::DescribeNodeResult(NodeResult));
+
+	OnTaskFinished(OwnerComp, NodeMemory, NodeResult);
 }
 
 FString UBTTask_RunBehavior::GetStaticDescription() const
diff --git a/Engine/Source/Runtime/AIModule/Private/BehaviorTree/Tasks/BTTask_RunBehaviorDynamic.cpp b/Engine/Source/Runtime/AIModule/Private/BehaviorTree/Tasks/BTTask_RunBehaviorDynamic.cpp
--- a/Engine/Source/Runtime/AIModule/Private/BehaviorTree/Tasks/BTTask_RunBehaviorDynamic.cpp
+++ b/Engine/Source/Runtime/AIModule/Private/BehaviorTree/Tasks/BTTask_RunBehaviorDynamic.cpp
@@ -1,6 +1,7 @@
 // Copyright 1998-2018 Epic Games, Inc. All Rights Reserved.
 
 #include "BehaviorTree/Tasks/BTTask_RunBehaviorDynamic.h"
+#include "VisualLogger/VisualLogger.h"
 #include "BehaviorTree/BehaviorTree.h"
 
 UBTTask_RunBehaviorDynamic::UBTTask_RunBehaviorDynamic(const FObjectInitializer& ObjectInitializer) : Super(ObjectInitializer)
@@ -12,7 +13,16 @@ UBTTask_RunBehaviorDynamic::UBTTask_RunBehaviorDynamic(const FObjectInitializer&
 EBTNodeResult::Type UBTTask_RunBehaviorDynamic::ExecuteTask(UBehaviorTreeComponent& OwnerComp, uint8* NodeMemory)
 {
 	const bool bPushed = BehaviorAsset != nullptr && OwnerComp.PushInstance(*BehaviorAsset);
-	return bPushed ? EBTNodeResult::InProgress : EBTNodeResult::Failed;
+	if (bPushed && OwnerComp.InstanceStack.Num() > 0)
+	{
+		FBehaviorTreeInstance& MyInstance = OwnerComp.InstanceStack[OwnerComp.InstanceStack.Num() - 1];
+		MyInstance.DeactivationNotify.BindUObject(this, &UBTTask_RunBehaviorDynamic::OnSubtreeDeactivated);
+		// unbinding is not required, MyInstance will be destroyed after firing that delegate (usually by UBehaviorTreeComponent::ProcessPendingExecution) 
+
+		return EBTNodeResult::InProgress;
+	}
+
+	return EBTNodeResult::Failed;
 }
 
 void UBTTask_RunBehaviorDynamic::OnInstanceCreated(UBehaviorTreeComponent& OwnerComp)
@@ -21,6 +31,17 @@ void UBTTask_RunBehaviorDynamic::OnInstanceCreated(UBehaviorTreeComponent& Owner
 	BehaviorAsset = DefaultBehaviorAsset;
 }
 
+void UBTTask_RunBehaviorDynamic::OnSubtreeDeactivated(UBehaviorTreeComponent& OwnerComp, EBTNodeResult::Type NodeResult)
+{
+	const int32 MyInstanceIdx = OwnerComp.FindInstanceContainingNode(this);
+	uint8* NodeMemory = OwnerComp.GetNodeMemory(this, MyInstanceIdx);
+
+	UE_VLOG(OwnerComp.GetOwner(), LogBehaviorTree, Verbose, TEXT("OnSubtreeDeactivated: %s (result: %s)"),
+		*UBehaviorTreeTypes::DescribeNodeHelper(this), *UBehaviorTreeTypes::DescribeNodeResult(NodeResult));
+
+	OnTaskFinished(OwnerComp, NodeMemory, NodeResult);
+}
+
 FString UBTTask_RunBehaviorDynamic::GetStaticDescription() const
 {
 	return FString::Printf(TEXT("%s: %s"), *Super::GetStaticDescription(), *InjectionTag.ToString());
