Added a comprehensive function to AIPerceptionSystem for registering an AActro as a stimuli source for all senses #UE4

diff --git a/Engine/Source/Runtime/AIModule/Classes/Perception/AIPerceptionSystem.h b/Engine/Source/Runtime/AIModule/Classes/Perception/AIPerceptionSystem.h
--- a/Engine/Source/Runtime/AIModule/Classes/Perception/AIPerceptionSystem.h
+++ b/Engine/Source/Runtime/AIModule/Classes/Perception/AIPerceptionSystem.h
@@ -149,13 +149,18 @@ public:
 	UFUNCTION(BlueprintCallable, Category = "AI|Perception", meta = (WorldContext="WorldContextObject"))
 	static void ReportPerceptionEvent(UObject* WorldContextObject, UAISenseEvent* PerceptionEvent);
 
+	/** Registers a source of given sense's stimuli */
 	template<typename FSenseClass>
 	void RegisterSource(AActor& SourceActor);
 
+	/** Registers given actor as a source for all registered senses */
+	void RegisterSource(AActor& SourceActor);
+
 	void RegisterSourceForSenseClass(TSubclassOf<UAISense> Sense, AActor& Target);
 
 	/** 
 	 *	unregisters given actor from the list of active stimuli sources
+	 *	@param Sense if null will result in removing SourceActor from all the senses
 	 */
 	void UnregisterSource(AActor& SourceActor, const TSubclassOf<UAISense> Sense = nullptr);
 
diff --git a/Engine/Source/Runtime/AIModule/Private/Perception/AIPerceptionSystem.cpp b/Engine/Source/Runtime/AIModule/Private/Perception/AIPerceptionSystem.cpp
--- a/Engine/Source/Runtime/AIModule/Private/Perception/AIPerceptionSystem.cpp
+++ b/Engine/Source/Runtime/AIModule/Private/Perception/AIPerceptionSystem.cpp
@@ -542,6 +542,23 @@ void UAIPerceptionSystem::OnNewPawn(APawn& Pawn)
 	}
 }
 
+void UAIPerceptionSystem::RegisterSource(AActor& SourceActor)
+{
+	for (UAISense* Sense : Senses)
+	{
+		if (Sense == nullptr)
+		{
+			continue;
+		}
+
+		const FAISenseID SenseID = Sense->GetSenseID();
+		if (IsSenseInstantiated(SenseID))
+		{
+			RegisterSource(SenseID, SourceActor);
+		}
+	}
+}
+
 void UAIPerceptionSystem::StartPlay()
 {
 	for (UAISense* Sense : Senses)
