Added a GetTypeHash override for FGenericTeamId so that it can be used in a UPROPERTY TMap #UE4

diff --git a/Engine/Source/Runtime/AIModule/Classes/GenericTeamAgentInterface.h b/Engine/Source/Runtime/AIModule/Classes/GenericTeamAgentInterface.h
--- a/Engine/Source/Runtime/AIModule/Classes/GenericTeamAgentInterface.h
+++ b/Engine/Source/Runtime/AIModule/Classes/GenericTeamAgentInterface.h
@@ -66,6 +66,11 @@ public:
 	static const FGenericTeamId NoTeam;
 };
 
+FORCEINLINE uint32 GetTypeHash(const FGenericTeamId TeamID)
+{
+	return TeamID.GetId();
+}
+
 UINTERFACE()
 class AIMODULE_API UGenericTeamAgentInterface : public UInterface
 {
