Made BT editor's comment on root-level decorators cleares and added error logging if BTTask_RunBehaviorDynamic is executed with a subtree containing root-level decorators #UE4

diff --git a/Engine/Source/Editor/BehaviorTreeEditor/Private/BehaviorTreeEditor.cpp b/Engine/Source/Editor/BehaviorTreeEditor/Private/BehaviorTreeEditor.cpp
--- a/Engine/Source/Editor/BehaviorTreeEditor/Private/BehaviorTreeEditor.cpp
+++ b/Engine/Source/Editor/BehaviorTreeEditor/Private/BehaviorTreeEditor.cpp
@@ -654,7 +654,10 @@ TSharedRef<SWidget> FBehaviorTreeEditor::SpawnProperties()
 				.Padding(FMargin(5.0f))
 				[
 					SNew(STextBlock)
-					.Text(LOCTEXT("RootLevelNode", "Root level decorators are not executed\nThey will be injected into a parent tree"))
+					.Text(LOCTEXT("RootLevelNode", "\
+Root-level decorators are only valid and will be executed if this BT is be used\n\
+as static a sub-tree (via \"Run Behavior\"). These decorators will be ignored if\n\
+dynamically injected with \"Run Dynamic Behavior\"."))
 				]
 			]
 			+SVerticalBox::Slot()
diff --git a/Engine/Source/Runtime/AIModule/Private/BehaviorTree/Tasks/BTTask_RunBehaviorDynamic.cpp b/Engine/Source/Runtime/AIModule/Private/BehaviorTree/Tasks/BTTask_RunBehaviorDynamic.cpp
--- a/Engine/Source/Runtime/AIModule/Private/BehaviorTree/Tasks/BTTask_RunBehaviorDynamic.cpp
+++ b/Engine/Source/Runtime/AIModule/Private/BehaviorTree/Tasks/BTTask_RunBehaviorDynamic.cpp
@@ -12,6 +12,12 @@ UBTTask_RunBehaviorDynamic::UBTTask_RunBehaviorDynamic(const FObjectInitializer&
 
 EBTNodeResult::Type UBTTask_RunBehaviorDynamic::ExecuteTask(UBehaviorTreeComponent& OwnerComp, uint8* NodeMemory)
 {
+	UE_CLOG(BehaviorAsset == nullptr, LogBehaviorTree, Warning, TEXT("BTTask_RunBehaviorDynamic node (\"%s\") executed with no BehaviorAsset")
+		, *GetNodeName());
+	UE_CLOG(BehaviorAsset != nullptr && BehaviorAsset->RootDecorators.Num() > 0, LogBehaviorTree
+		, Error, TEXT("BTTask_RunBehaviorDynamic node (\"%s\") executed with a BehaviorAsset (\"%s\") containing root level decorators. These decorators will be ignored by design.")
+		, *GetNodeName(), *BehaviorAsset->GetName());
+
 	const bool bPushed = BehaviorAsset != nullptr && OwnerComp.PushInstance(*BehaviorAsset);
 	if (bPushed && OwnerComp.InstanceStack.Num() > 0)
 	{
