Exporting EnvQueryGenerator_SimpleGrid with AIMODULE_API #UE4

diff --git a/Engine/Source/Runtime/AIModule/Classes/EnvironmentQuery/Generators/EnvQueryGenerator_SimpleGrid.h b/Engine/Source/Runtime/AIModule/Classes/EnvironmentQuery/Generators/EnvQueryGenerator_SimpleGrid.h
--- a/Engine/Source/Runtime/AIModule/Classes/EnvironmentQuery/Generators/EnvQueryGenerator_SimpleGrid.h
+++ b/Engine/Source/Runtime/AIModule/Classes/EnvironmentQuery/Generators/EnvQueryGenerator_SimpleGrid.h
@@ -15,7 +15,7 @@
  */
 
 UCLASS(meta = (DisplayName = "Points: Grid"))
-class UEnvQueryGenerator_SimpleGrid : public UEnvQueryGenerator_ProjectedPoints
+class AIMODULE_API UEnvQueryGenerator_SimpleGrid : public UEnvQueryGenerator_ProjectedPoints
 {
 	GENERATED_UCLASS_BODY()
 
