Clearing editor world's EnvironmentQueryManager instace's LocalContextMap whenever a BP finished compiling #UE4

diff --git a/Engine/Source/Runtime/AIModule/Classes/EnvironmentQuery/EnvQueryManager.h b/Engine/Source/Runtime/AIModule/Classes/EnvironmentQuery/EnvQueryManager.h
--- a/Engine/Source/Runtime/AIModule/Classes/EnvironmentQuery/EnvQueryManager.h
+++ b/Engine/Source/Runtime/AIModule/Classes/EnvironmentQuery/EnvQueryManager.h
@@ -161,6 +161,8 @@ class AIMODULE_API UEnvQueryManager : public UObject, public FTickableGameObject
 	// Any loaded instance will get marked as PendingKill
 	virtual void PostLoad() override;
 
+	virtual void PostInitProperties() override;
+
 	// We need to implement GetWorld() so that any EQS-related blueprints (such as blueprint contexts) can implement
 	// GetWorld() and so provide access to blueprint nodes using hidden WorldContextObject parameters.
 	virtual UWorld* GetWorld() const override;
@@ -319,4 +321,11 @@ private:
 	void CheckQueryCount() const;
 	void LogQueryCountWarning() const;
 #endif
+
+#if WITH_EDITOR
+	FDelegateHandle OnBlueprintCompiledHandle;
+
+	// used to reset LocalContextMap in case we had BP contexts stored
+	void OnBlueprintCompiled();
+#endif // WITH_EDITOR
 };
diff --git a/Engine/Source/Runtime/AIModule/Private/EnvironmentQuery/EnvQueryManager.cpp b/Engine/Source/Runtime/AIModule/Private/EnvironmentQuery/EnvQueryManager.cpp
--- a/Engine/Source/Runtime/AIModule/Private/EnvironmentQuery/EnvQueryManager.cpp
+++ b/Engine/Source/Runtime/AIModule/Private/EnvironmentQuery/EnvQueryManager.cpp
@@ -132,6 +132,17 @@ void UEnvQueryManager::PostLoad()
 	MarkPendingKill();
 }
 
+void UEnvQueryManager::PostInitProperties()
+{
+	Super::PostInitProperties();
+#if WITH_EDITOR
+	if (GEditor)
+	{
+		OnBlueprintCompiledHandle = GEditor->OnBlueprintCompiled().AddUObject(this, &UEnvQueryManager::OnBlueprintCompiled);
+	}
+#endif // WITH_EDITOR
+}
+
 UWorld* UEnvQueryManager::GetWorld() const
 {
 	return Cast<UWorld>(GetOuter());
@@ -140,6 +151,14 @@ UWorld* UEnvQueryManager::GetWorld() const
 void UEnvQueryManager::FinishDestroy()
 {
 	FCoreUObjectDelegates::PreLoadMap.RemoveAll(this);
+
+#if WITH_EDITOR
+	if (GEditor)
+	{
+		GEditor->OnBlueprintCompiled().Remove(OnBlueprintCompiledHandle);
+	}
+#endif // WITH_EDITOR
+
 	Super::FinishDestroy();
 }
 
@@ -948,6 +967,13 @@ TSharedPtr<FEnvQueryInstance> UEnvQueryManager::FindQueryInstance(const int32 Qu
 	return nullptr;
 }
 
+#if WITH_EDITOR
+void UEnvQueryManager::OnBlueprintCompiled()
+{
+	LocalContextMap.Reset();
+}
+#endif // WITH_EDITOR
+
 //----------------------------------------------------------------------//
 // Exec functions (i.e. console commands)
 //----------------------------------------------------------------------//
