Added a way to configure navmesh drawing GameplayDebugger category to mark up "forbidden" nav polys  (the ones that have

diff --git a/Engine/Source/Runtime/AIModule/Private/GameplayDebugger/GameplayDebuggerCategory_Navmesh.cpp b/Engine/Source/Runtime/AIModule/Private/GameplayDebugger/GameplayDebuggerCategory_Navmesh.cpp
--- a/Engine/Source/Runtime/AIModule/Private/GameplayDebugger/GameplayDebuggerCategory_Navmesh.cpp
+++ b/Engine/Source/Runtime/AIModule/Private/GameplayDebugger/GameplayDebuggerCategory_Navmesh.cpp
@@ -8,6 +8,13 @@
 #include "GameFramework/PlayerController.h"
 #include "NavMesh/RecastNavMesh.h"
 
+namespace
+{
+	int32 bDrawExcludedFlags = 0;
+	FAutoConsoleVariableRef CVar(TEXT("ai.debug.nav.DrawExcludedFlags"), bDrawExcludedFlags, TEXT("If we want to mark \"forbidden\" nav polys while debug-drawing."), ECVF_Default);
+}
+
+
 FGameplayDebuggerCategory_Navmesh::FGameplayDebuggerCategory_Navmesh()
 {
 	bShowOnlyWithDebugActor = false;
@@ -64,7 +71,8 @@ void FGameplayDebuggerCategory_Navmesh::CollectData(APlayerController* OwnerPC,
 		const int32 DetailFlags =
 			(1 << static_cast<int32>(ENavMeshDetailFlags::PolyEdges)) |
 			(1 << static_cast<int32>(ENavMeshDetailFlags::FilledPolys)) |
-			(1 << static_cast<int32>(ENavMeshDetailFlags::NavLinks));
+			(1 << static_cast<int32>(ENavMeshDetailFlags::NavLinks)) |
+			(bDrawExcludedFlags ? (1 << static_cast<int32>(ENavMeshDetailFlags::MarkForbiddenPolys)) : 0);
 
 		NavmeshRenderData.GatherData(NavData, DetailFlags, TileSet);
 	}
