Use TActorRange instead of FConstPawnIterator
#rnx

diff --git a/Engine/Source/Developer/GameplayDebugger/Private/GameplayDebuggerLocalController.cpp b/Engine/Source/Developer/GameplayDebugger/Private/GameplayDebuggerLocalController.cpp
--- a/Engine/Source/Developer/GameplayDebugger/Private/GameplayDebuggerLocalController.cpp
+++ b/Engine/Source/Developer/GameplayDebugger/Private/GameplayDebuggerLocalController.cpp
@@ -24,6 +24,7 @@
 #include "UnrealEngine.h"
 #include "Engine/LocalPlayer.h"
 #include "GameFramework/PlayerInput.h"
+#include "EngineUtils.h"
 
 #if WITH_EDITOR
 #include "Editor/GameplayDebuggerEdMode.h"
@@ -623,10 +624,9 @@ void UGameplayDebuggerLocalController::OnSelectActorTick()
 		float BestScore = MinViewDirDot;
 		
 		const FVector ViewDir = CameraRotation.Vector();
-		for (FConstPawnIterator It = OwnerPC->GetWorld()->GetPawnIterator(); It; ++It)
+		for (APawn* TestPawn  : TActorRange<APawn>(OwnerPC->GetWorld()))
 		{
-			APawn* TestPawn = It->Get();
-			if (TestPawn && !TestPawn->bHidden && TestPawn->GetActorEnableCollision() &&
+			if (!TestPawn->bHidden && TestPawn->GetActorEnableCollision() &&
 				!TestPawn->IsA(ASpectatorPawn::StaticClass()) &&
 				TestPawn != OwnerPC->GetPawn())
 			{
diff --git a/Engine/Source/Runtime/AIModule/Private/GameplayDebugger/GameplayDebuggerCategory_AI.cpp b/Engine/Source/Runtime/AIModule/Private/GameplayDebugger/GameplayDebuggerCategory_AI.cpp
--- a/Engine/Source/Runtime/AIModule/Private/GameplayDebugger/GameplayDebuggerCategory_AI.cpp
+++ b/Engine/Source/Runtime/AIModule/Private/GameplayDebugger/GameplayDebuggerCategory_AI.cpp
@@ -23,7 +23,7 @@
 #include "GameFramework/Character.h"
 #include "GameFramework/CharacterMovementComponent.h"
 #include "DrawDebugHelpers.h"
-
+#include "EngineUtils.h"
 
 namespace FGameplayDebuggerCategoryTweakables
 {
@@ -480,10 +480,9 @@ void FGameplayDebuggerCategory_AI::DrawOverheadInfo(AActor& DebugActor, FGamepla
 void FGameplayDebuggerCategory_AI::DrawPawnIcons(UWorld* World, AActor* DebugActor, APawn* SkipPawn, FGameplayDebuggerCanvasContext& CanvasContext)
 {
 	FString FailsafeIcon = TEXT("/Engine/EngineResources/AICON-Green.AICON-Green");
-	for (FConstPawnIterator It = World->GetPawnIterator(); It; ++It)
+	for (const APawn* ItPawn : TActorRange<APawn>(World))
 	{
-		const APawn* ItPawn = It->Get();
-		if (IsValid(ItPawn) && SkipPawn != ItPawn)
+		if (SkipPawn != ItPawn)
 		{
 			const FVector IconLocation = ItPawn->GetActorLocation() + FVector(0, 0, ItPawn->GetSimpleCollisionHalfHeight());
 			const AAIController* ItAI = Cast<const AAIController>(ItPawn->GetController());
diff --git a/Engine/Source/Runtime/AIModule/Private/Perception/PawnSensingComponent.cpp b/Engine/Source/Runtime/AIModule/Private/Perception/PawnSensingComponent.cpp
--- a/Engine/Source/Runtime/AIModule/Private/Perception/PawnSensingComponent.cpp
+++ b/Engine/Source/Runtime/AIModule/Private/Perception/PawnSensingComponent.cpp
@@ -8,6 +8,7 @@
 #include "GameFramework/Controller.h"
 #include "GameFramework/PlayerController.h"
 #include "Engine/Engine.h"
+#include "EngineUtils.h"
 #include "AIController.h"
 #include "Components/PawnNoiseEmitterComponent.h"
 
@@ -227,10 +228,9 @@ void UPawnSensingComponent::UpdateAISensing()
 	}
 	else
 	{
-		for (FConstPawnIterator Iterator = Owner->GetWorld()->GetPawnIterator(); Iterator; ++Iterator)
+		for (APawn* Pawn : TActorRange<APawn>(Owner->GetWorld()))
 		{
-			APawn* Pawn = Iterator->Get();
-			if (IsValid(Pawn) && !IsSensorActor(Pawn))
+			if (!IsSensorActor(Pawn))
 			{
 				SensePawn(*Pawn);
 			}
diff --git a/Engine/Source/Runtime/Engine/Private/GameModeBase.cpp b/Engine/Source/Runtime/Engine/Private/GameModeBase.cpp
--- a/Engine/Source/Runtime/Engine/Private/GameModeBase.cpp
+++ b/Engine/Source/Runtime/Engine/Private/GameModeBase.cpp
@@ -1366,10 +1366,9 @@ bool AGameModeBase::SpawnPlayerFromSimulate(const FVector& NewLocation, const FR
 		if (PC->GetPawn() == nullptr)
 		{
 			// Use the "auto-possess" pawn in the world, if there is one.
-			for (FConstPawnIterator Iterator = GetWorld()->GetPawnIterator(); Iterator; ++Iterator)
+			for (APawn* Pawn : TActorRange<APawn>(GetWorld()))
 			{
-				APawn* Pawn = Iterator->Get();
-				if (Pawn && Pawn->AutoPossessPlayer == EAutoReceiveInput::Player0)
+				if (Pawn->AutoPossessPlayer == EAutoReceiveInput::Player0)
 				{
 					if (Pawn->Controller == nullptr)
 					{
