Added AISenseConfig const array iterator accessor to perception component.

#rnx

diff --git a/Engine/Source/Runtime/AIModule/Classes/Perception/AIPerceptionComponent.h b/Engine/Source/Runtime/AIModule/Classes/Perception/AIPerceptionComponent.h
--- a/Engine/Source/Runtime/AIModule/Classes/Perception/AIPerceptionComponent.h
+++ b/Engine/Source/Runtime/AIModule/Classes/Perception/AIPerceptionComponent.h
@@ -255,6 +255,9 @@ public:
 	const UAISenseConfig* GetSenseConfig(const FAISenseID& SenseID) const;
 	void ConfigureSense(UAISenseConfig& SenseConfig);
 
+	typedef TArray<UAISenseConfig*>::TConstIterator TAISenseConfigConstIterator;
+	TAISenseConfigConstIterator GetSensesConfigIterator() const;
+
 	/** Notifies AIPerceptionSystem to update properties for this "stimuli listener" */
 	UFUNCTION(BlueprintCallable, Category="AI|Perception")
 	void RequestStimuliListenerUpdate();
diff --git a/Engine/Source/Runtime/AIModule/Private/Perception/AIPerceptionComponent.cpp b/Engine/Source/Runtime/AIModule/Private/Perception/AIPerceptionComponent.cpp
--- a/Engine/Source/Runtime/AIModule/Private/Perception/AIPerceptionComponent.cpp
+++ b/Engine/Source/Runtime/AIModule/Private/Perception/AIPerceptionComponent.cpp
@@ -131,6 +131,11 @@ void UAIPerceptionComponent::ConfigureSense(UAISenseConfig& Config)
 	// else the sense will be auto-configured during OnRegister
 }
 
+UAIPerceptionComponent::TAISenseConfigConstIterator UAIPerceptionComponent::GetSensesConfigIterator() const
+{
+	return SensesConfig.CreateConstIterator();
+}
+
 void UAIPerceptionComponent::SetMaxStimulusAge(int32 ConfigIndex, float MaxAge)
 {
 	if (MaxActiveAge.IsValidIndex(ConfigIndex) == false)
