Fixed a bug in EQS test sorting that was resulting in undesired sorting results.

Yoan.StAmant

diff --git a/Engine/Source/Runtime/AIModule/Private/EnvironmentQuery/EnvQueryManager.cpp b/Engine/Source/Runtime/AIModule/Private/EnvironmentQuery/EnvQueryManager.cpp
--- a/Engine/Source/Runtime/AIModule/Private/EnvironmentQuery/EnvQueryManager.cpp
+++ b/Engine/Source/Runtime/AIModule/Private/EnvironmentQuery/EnvQueryManager.cpp
@@ -579,64 +579,19 @@ namespace EnvQueryTestSort
 {
 	struct FAllMatching
 	{
-		FORCEINLINE bool operator()(const UEnvQueryTest& TestA, const UEnvQueryTest& TestB) const
+		FORCEINLINE_DEBUGGABLE bool operator()(const UEnvQueryTest& TestA, const UEnvQueryTest& TestB) const
 		{
-			// cheaper tests go first
-			if (TestB.Cost > TestA.Cost)
-			{
-				return true;
-			}
-
 			// conditions go first
 			const bool bConditionA = (TestA.TestPurpose != EEnvTestPurpose::Score); // Is Test A Filtering?
 			const bool bConditionB = (TestB.TestPurpose != EEnvTestPurpose::Score); // Is Test B Filtering?
-			if (bConditionA && !bConditionB)
+			if (bConditionA != bConditionB)
 			{
-				return true;
+				return bConditionA;
 			}
 
-			// keep connection order (sort stability)
-			return (TestB.TestOrder > TestA.TestOrder);
-		}
-	};
-
-	struct FSingleResult
-	{
-		FSingleResult(EEnvTestCost::Type InHighestCost) : HighestCost(InHighestCost) {}
-
-		FORCEINLINE bool operator()(const UEnvQueryTest& TestA, const UEnvQueryTest& TestB) const
-		{
 			// cheaper tests go first
-			if (TestB.Cost > TestA.Cost)
-			{
-				return true;
-			}
-
-			const bool bConditionA = (TestA.TestPurpose != EEnvTestPurpose::Score); // Is Test A Filtering?
-			const bool bConditionB = (TestB.TestPurpose != EEnvTestPurpose::Score); // Is Test B Filtering?
-			if (TestA.Cost == HighestCost)
-			{
-				// highest cost: weights go first, conditions later (first match will return result)
-				if (!bConditionA && bConditionB)
-				{
-					return true;
-				}
-			}
-			else
-			{
-				// lower costs: conditions go first to reduce amount of items
-				if (bConditionA && !bConditionB)
-				{
-					return true;
-				}
-			}
-
-			// keep connection order (sort stability)
-			return (TestB.TestOrder > TestA.TestOrder);
+			return (TestA.Cost < TestB.Cost);
 		}
-
-	protected:
-		TEnumAsByte<EEnvTestCost::Type> HighestCost;
 	};
 }
 
@@ -734,6 +689,13 @@ TSharedPtr<FEnvQueryInstance> UEnvQueryManager::CreateQueryInstance(const UEnvQu
 			LocalTemplate->Options[OptionIndex] = LocalOption;
 			LocalOption->Generator = LocalGenerator;
 
+			if (MyOption->Tests.Num() == 0)
+			{
+				// no further work required here
+				CreateOptionInstance(LocalOption, SourceOptionIndex, TArray<UEnvQueryTest*>(), *InstanceTemplate);
+				continue;
+			}
+
 			// check if TestOrder property is set correctly in asset, try to recreate it if not
 			// don't use editor graph, so any disabled tests in the middle will make it look weird
 			if (MyOption->Tests.Num() > 1 && MyOption->Tests.Last() && MyOption->Tests.Last()->TestOrder == 0)
@@ -747,7 +709,10 @@ TSharedPtr<FEnvQueryInstance> UEnvQueryManager::CreateQueryInstance(const UEnvQu
 				}
 			}
 
-			EEnvTestCost::Type HighestCost(EEnvTestCost::Low);
+			const bool bOptionSingleResultSearch = (RunMode == EEnvQueryRunMode::SingleResult) && LocalGenerator->bAutoSortTests;
+			EEnvTestCost::Type HighestFilterCost(EEnvTestCost::Low);
+			UEnvQueryTest* MostExpensiveFilter = nullptr;
+			int32 MostExpensiveFilterIndex = INDEX_NONE;
 			TArray<UEnvQueryTest*> SortedTests = MyOption->Tests;
 			TSubclassOf<UEnvQueryItemType> GeneratedType = MyOption->Generator->ItemType;
 			for (int32 TestIndex = SortedTests.Num() - 1; TestIndex >= 0; TestIndex--)
@@ -760,45 +725,56 @@ TSharedPtr<FEnvQueryInstance> UEnvQueryManager::CreateQueryInstance(const UEnvQu
 
 					SortedTests.RemoveAt(TestIndex, 1, false);
 				}
-				else if (HighestCost < TestOb->Cost)
+				else if (bOptionSingleResultSearch
+					&& TestOb->TestPurpose == EEnvTestPurpose::Filter
+					&& (HighestFilterCost < TestOb->Cost || MostExpensiveFilterIndex == INDEX_NONE))
 				{
-					HighestCost = TestOb->Cost;
+					HighestFilterCost = TestOb->Cost;					
+					MostExpensiveFilterIndex = TestIndex;
 				}
 			}
 
 			LocalOption->Tests.Reset(SortedTests.Num());
-			for (int32 TestIdx = 0; TestIdx < SortedTests.Num(); TestIdx++)
+			// MostExpensiveFilterIndex can be INDEX_NONE if there are no filtering 
+			// tests done (i.e. just scoring tests).
+			if (bOptionSingleResultSearch && MostExpensiveFilterIndex != INDEX_NONE)
 			{
-				UEnvQueryTest* LocalTest = (UEnvQueryTest*)StaticDuplicateObject(SortedTests[TestIdx], this);
-				LocalOption->Tests.Add(LocalTest);
+				// we're going to remove the most expensive test so that 
+				// it can be added last after everything else gets sorted
+				MostExpensiveFilter = (UEnvQueryTest*)StaticDuplicateObject(SortedTests[MostExpensiveFilterIndex], this);
+				for (int32 TestIdx = 0; TestIdx < SortedTests.Num(); TestIdx++)
+				{
+					if (TestIdx == MostExpensiveFilterIndex)
+					{
+						continue;
+					}
+					UEnvQueryTest* LocalTest = (UEnvQueryTest*)StaticDuplicateObject(SortedTests[TestIdx], this);
+					LocalOption->Tests.Add(LocalTest);
+				}
+			}
+			else
+			{				
+				for (int32 TestIdx = 0; TestIdx < SortedTests.Num(); TestIdx++)
+				{
+					UEnvQueryTest* LocalTest = (UEnvQueryTest*)StaticDuplicateObject(SortedTests[TestIdx], this);
+					LocalOption->Tests.Add(LocalTest);
+				}
 			}
-
-			// use locally referenced duplicates
-			SortedTests = LocalOption->Tests;
 
 			if (SortedTests.Num() && LocalGenerator->bAutoSortTests)
 			{
-				switch (RunMode)
+				LocalOption->Tests.StableSort(EnvQueryTestSort::FAllMatching());
+				if (bOptionSingleResultSearch && MostExpensiveFilter)
 				{
-				case EEnvQueryRunMode::SingleResult:
-					SortedTests.Sort(EnvQueryTestSort::FSingleResult(HighestCost));
-					break;
-
-				case EEnvQueryRunMode::RandomBest5Pct:
-				case EEnvQueryRunMode::RandomBest25Pct:
-				case EEnvQueryRunMode::AllMatching:
-					SortedTests.Sort(EnvQueryTestSort::FAllMatching());
-					break;
-
-				default:
-					{
-						UE_LOG(LogEQS, Warning, TEXT("Query [%s] can't be sorted for RunMode: %d [%s]"),
-							*GetNameSafe(LocalTemplate), (int32)RunMode, RunModeEnum ? *RunModeEnum->GetNameStringByValue(RunMode) : TEXT("??"));
-					}
+					// the only difference between running for a single result is that 
+					// we want to perform a single most expensive test as the last test
+					// to accept the first one that passes the test.
+					// so here we're adding that final, most expensive test
+					LocalOption->Tests.Add(MostExpensiveFilter);
 				}
 			}
 
-			CreateOptionInstance(LocalOption, SourceOptionIndex, SortedTests, *InstanceTemplate);
+			CreateOptionInstance(LocalOption, SourceOptionIndex, LocalOption->Tests, *InstanceTemplate);
 		}
 	}
 
@@ -818,13 +794,7 @@ void UEnvQueryManager::CreateOptionInstance(UEnvQueryOption* OptionTemplate, int
 	OptionInstance.Generator = OptionTemplate->Generator;
 	OptionInstance.ItemType = OptionTemplate->Generator->ItemType;
 	OptionInstance.SourceOptionIndex = SourceOptionIndex;
-
-	OptionInstance.Tests.AddZeroed(SortedTests.Num());
-	for (int32 TestIndex = 0; TestIndex < SortedTests.Num(); TestIndex++)
-	{
-		UEnvQueryTest* TestOb = SortedTests[TestIndex];
-		OptionInstance.Tests[TestIndex] = TestOb;
-	}
+	OptionInstance.Tests = SortedTests;
 
 	DEC_MEMORY_STAT_BY(STAT_AI_EQS_InstanceMemory, Instance.Options.GetAllocatedSize());
 
