Engine-level code updated to go along with NavigationSystem refactor

diff --git a/Engine/Source/Runtime/AIModule/Private/AIController.cpp b/Engine/Source/Runtime/AIModule/Private/AIController.cpp
--- a/Engine/Source/Runtime/AIModule/Private/AIController.cpp
+++ b/Engine/Source/Runtime/AIModule/Private/AIController.cpp
@@ -743,7 +743,7 @@ bool AAIController::BuildPathfindingQuery(const FAIMoveRequest& MoveRequest, FPa
 
 	UNavigationSystemV1* NavSys = FNavigationSystem::GetCurrent<UNavigationSystemV1>(GetWorld());
 	const ANavigationData* NavData = (NavSys == nullptr) ? nullptr :
-		MoveRequest.IsUsingPathfinding() ? NavSys->GetNavDataForProps(GetNavAgentPropertiesRef()) :
+		MoveRequest.IsUsingPathfinding() ? NavSys->GetNavDataForProps(GetNavAgentPropertiesRef(), GetNavAgentLocation()) :
 		NavSys->GetAbstractNavData();
 
 	if (NavData)
diff --git a/Engine/Source/Runtime/AIModule/Private/BehaviorTree/Decorators/BTDecorator_DoesPathExist.cpp b/Engine/Source/Runtime/AIModule/Private/BehaviorTree/Decorators/BTDecorator_DoesPathExist.cpp
--- a/Engine/Source/Runtime/AIModule/Private/BehaviorTree/Decorators/BTDecorator_DoesPathExist.cpp
+++ b/Engine/Source/Runtime/AIModule/Private/BehaviorTree/Decorators/BTDecorator_DoesPathExist.cpp
@@ -64,7 +64,7 @@ bool UBTDecorator_DoesPathExist::CalculateRawConditionValue(UBehaviorTreeCompone
 	if (NavSys && bHasPointA && bHasPointB)
 	{
 		const AAIController* AIOwner = OwnerComp.GetAIOwner();
-		const ANavigationData* NavData = AIOwner ? NavSys->GetNavDataForProps(AIOwner->GetNavAgentPropertiesRef()) : NULL;
+		const ANavigationData* NavData = AIOwner ? NavSys->GetNavDataForProps(AIOwner->GetNavAgentPropertiesRef(), AIOwner->GetNavAgentLocation()) : NULL;
 		if (NavData)
 		{
 			FSharedConstNavQueryFilter QueryFilter = UNavigationQueryFilter::GetQueryFilter(*NavData, AIOwner, FilterClass);
diff --git a/Engine/Source/Runtime/AIModule/Private/Blueprint/AIBlueprintHelperLibrary.cpp b/Engine/Source/Runtime/AIModule/Private/Blueprint/AIBlueprintHelperLibrary.cpp
--- a/Engine/Source/Runtime/AIModule/Private/Blueprint/AIBlueprintHelperLibrary.cpp
+++ b/Engine/Source/Runtime/AIModule/Private/Blueprint/AIBlueprintHelperLibrary.cpp
@@ -415,10 +415,11 @@ void UAIBlueprintHelperLibrary::SimpleMoveToActor(AController* Controller, const
 	}
 	else
 	{
-		const ANavigationData* NavData = NavSys->GetNavDataForProps(Controller->GetNavAgentPropertiesRef());
+		const FVector AgentNavLocation = Controller->GetNavAgentLocation();
+		const ANavigationData* NavData = NavSys->GetNavDataForProps(Controller->GetNavAgentPropertiesRef(), AgentNavLocation);
 		if (NavData)
 		{
-			FPathFindingQuery Query(Controller, *NavData, Controller->GetNavAgentLocation(), Goal->GetActorLocation());
+			FPathFindingQuery Query(Controller, *NavData, AgentNavLocation, Goal->GetActorLocation());
 			FPathFindingResult Result = NavSys->FindPathSync(Query);
 			if (Result.IsSuccessful())
 			{
@@ -484,10 +485,11 @@ void UAIBlueprintHelperLibrary::SimpleMoveToLocation(AController* Controller, co
 	}
 	else
 	{
-		const ANavigationData* NavData = NavSys->GetNavDataForProps(Controller->GetNavAgentPropertiesRef());
+		const FVector AgentNavLocation = Controller->GetNavAgentLocation();
+		const ANavigationData* NavData = NavSys->GetNavDataForProps(Controller->GetNavAgentPropertiesRef(), AgentNavLocation);
 		if (NavData)
 		{
-			FPathFindingQuery Query(Controller, *NavData, Controller->GetNavAgentLocation(), GoalLocation);
+			FPathFindingQuery Query(Controller, *NavData, AgentNavLocation, GoalLocation);
 			FPathFindingResult Result = NavSys->FindPathSync(Query);
 			if (Result.IsSuccessful())
 			{
diff --git a/Engine/Source/Runtime/AIModule/Private/EnvironmentQuery/EnvQueryTypes.cpp b/Engine/Source/Runtime/AIModule/Private/EnvironmentQuery/EnvQueryTypes.cpp
--- a/Engine/Source/Runtime/AIModule/Private/EnvironmentQuery/EnvQueryTypes.cpp
+++ b/Engine/Source/Runtime/AIModule/Private/EnvironmentQuery/EnvQueryTypes.cpp
@@ -292,7 +292,7 @@ namespace FEQSHelpers
 		if (NavAgent)
 		{
 			const FNavAgentProperties& NavAgentProps = NavAgent->GetNavAgentPropertiesRef();
-			return NavSys->GetNavDataForProps(NavAgentProps);
+			return NavSys->GetNavDataForProps(NavAgentProps, NavAgent->GetNavAgentLocation());
 		}
 
 		return NavSys->GetDefaultNavDataInstance();
diff --git a/Engine/Source/Runtime/AIModule/Private/EnvironmentQuery/Tests/EnvQueryTest_Pathfinding.cpp b/Engine/Source/Runtime/AIModule/Private/EnvironmentQuery/Tests/EnvQueryTest_Pathfinding.cpp
--- a/Engine/Source/Runtime/AIModule/Private/EnvironmentQuery/Tests/EnvQueryTest_Pathfinding.cpp
+++ b/Engine/Source/Runtime/AIModule/Private/EnvironmentQuery/Tests/EnvQueryTest_Pathfinding.cpp
@@ -225,7 +225,7 @@ ANavigationData* UEnvQueryTest_Pathfinding::FindNavigationData(UNavigationSystem
 	INavAgentInterface* NavAgent = Cast<INavAgentInterface>(Owner);
 	if (NavAgent)
 	{
-		return NavSys.GetNavDataForProps(NavAgent->GetNavAgentPropertiesRef());
+		return NavSys.GetNavDataForProps(NavAgent->GetNavAgentPropertiesRef(), NavAgent->GetNavAgentLocation());
 	}
 
 	return NavSys.GetDefaultNavDataInstance(FNavigationSystem::DontCreate);
diff --git a/Engine/Source/Runtime/AIModule/Private/GameplayDebugger/GameplayDebuggerCategory_AI.cpp b/Engine/Source/Runtime/AIModule/Private/GameplayDebugger/GameplayDebuggerCategory_AI.cpp
--- a/Engine/Source/Runtime/AIModule/Private/GameplayDebugger/GameplayDebuggerCategory_AI.cpp
+++ b/Engine/Source/Runtime/AIModule/Private/GameplayDebugger/GameplayDebuggerCategory_AI.cpp
@@ -193,7 +193,9 @@ void FGameplayDebuggerCategory_AI::CollectData(APlayerController* OwnerPC, AActo
 		DataPack.MontageInfo = MyChar ? GetNameSafe(MyChar->GetCurrentMontage()) : FString();
 
 		UNavigationSystemV1* NavSys = FNavigationSystem::GetCurrent<UNavigationSystemV1>(MyPawn->GetWorld());
-		const ANavigationData* NavData = MyController && NavSys ? NavSys->GetNavDataForProps(MyController->GetNavAgentPropertiesRef()) : nullptr;
+		const ANavigationData* NavData = MyController && NavSys 
+			? NavSys->GetNavDataForProps(MyController->GetNavAgentPropertiesRef(), MyController->GetNavAgentLocation()) 
+			: nullptr;
 		DataPack.NavDataInfo = NavData ? NavData->GetConfig().Name.ToString() : FString();
 
 		CollectPathData(MyController);
diff --git a/Engine/Source/Runtime/AIModule/Private/GameplayDebugger/GameplayDebuggerCategory_Navmesh.cpp b/Engine/Source/Runtime/AIModule/Private/GameplayDebugger/GameplayDebuggerCategory_Navmesh.cpp
--- a/Engine/Source/Runtime/AIModule/Private/GameplayDebugger/GameplayDebuggerCategory_Navmesh.cpp
+++ b/Engine/Source/Runtime/AIModule/Private/GameplayDebugger/GameplayDebuggerCategory_Navmesh.cpp
@@ -43,7 +43,7 @@ void FGameplayDebuggerCategory_Navmesh::CollectData(APlayerController* OwnerPC,
 		if (NavSys) 
 		{
 			const FNavAgentProperties& NavAgentProperties = DestPawn->GetNavAgentPropertiesRef();
-			NavData = Cast<const ARecastNavMesh>(NavSys->GetNavDataForProps(NavAgentProperties));
+			NavData = Cast<const ARecastNavMesh>(NavSys->GetNavDataForProps(NavAgentProperties, DestPawn->GetNavAgentLocation()));
 		}
 	}
 
diff --git a/Engine/Source/Runtime/AIModule/Private/Navigation/MetaNavMeshPath.cpp b/Engine/Source/Runtime/AIModule/Private/Navigation/MetaNavMeshPath.cpp
--- a/Engine/Source/Runtime/AIModule/Private/Navigation/MetaNavMeshPath.cpp
+++ b/Engine/Source/Runtime/AIModule/Private/Navigation/MetaNavMeshPath.cpp
@@ -37,7 +37,7 @@ FMetaNavMeshPath::FMetaNavMeshPath(const TArray<FMetaPathWayPoint>& InWaypoints,
 	: FMetaNavMeshPath()
 {
 	UNavigationSystemV1* NavSys = FNavigationSystem::GetCurrent<UNavigationSystemV1>(Owner.GetWorld());
-	const ANavigationData* NavData = NavSys ? NavSys->GetNavDataForProps(Owner.GetNavAgentPropertiesRef()) : nullptr;
+	const ANavigationData* NavData = NavSys ? NavSys->GetNavDataForProps(Owner.GetNavAgentPropertiesRef(), Owner.GetNavAgentLocation()) : nullptr;
 
 	if (ensure(NavData))
 	{
@@ -59,7 +59,7 @@ FMetaNavMeshPath::FMetaNavMeshPath(const TArray<FVector>& InWaypoints, const ANa
 FMetaNavMeshPath::FMetaNavMeshPath(const TArray<FVector>& InWaypoints, const AController& Owner) : FMetaNavMeshPath()
 {
 	UNavigationSystemV1* NavSys = FNavigationSystem::GetCurrent<UNavigationSystemV1>(Owner.GetWorld());
-	const ANavigationData* NavData = NavSys ? NavSys->GetNavDataForProps(Owner.GetNavAgentPropertiesRef()) : nullptr;
+	const ANavigationData* NavData = NavSys ? NavSys->GetNavDataForProps(Owner.GetNavAgentPropertiesRef(), Owner.GetNavAgentLocation()) : nullptr;
 
 	if (ensure(NavData))
 	{
diff --git a/Engine/Source/Runtime/AIModule/Private/Navigation/PathFollowingComponent.cpp b/Engine/Source/Runtime/AIModule/Private/Navigation/PathFollowingComponent.cpp
--- a/Engine/Source/Runtime/AIModule/Private/Navigation/PathFollowingComponent.cpp
+++ b/Engine/Source/Runtime/AIModule/Private/Navigation/PathFollowingComponent.cpp
@@ -669,7 +669,7 @@ void UPathFollowingComponent::SetMovementComponent(UNavMovementComponent* MoveCo
 		
 		if (NavSys)
 		{
-			MyNavData = NavSys->GetNavDataForProps(NavAgentProps);
+			MyNavData = NavSys->GetNavDataForProps(NavAgentProps, MoveComp->GetActorLocation());
 			if (MyNavData == nullptr)
 			{
 				if (NavSys->IsInitialized() == false)
@@ -693,7 +693,7 @@ void UPathFollowingComponent::OnNavigationInitDone()
 		UNavigationSystemV1* NavSys = FNavigationSystem::GetCurrent<UNavigationSystemV1>(MyWorld);
 		check(NavSys);
 		const FNavAgentProperties& NavAgentProps = MovementComp->GetNavAgentPropertiesRef();
-		MyNavData = NavSys->GetNavDataForProps(NavAgentProps);
+		MyNavData = NavSys->GetNavDataForProps(NavAgentProps, MovementComp->GetActorLocation());
 		NavSys->OnNavigationInitDone.RemoveAll(this);
 	}
 }
diff --git a/Engine/Source/Runtime/NavigationSystem/Private/NavigationTestingActor.cpp b/Engine/Source/Runtime/NavigationSystem/Private/NavigationTestingActor.cpp
--- a/Engine/Source/Runtime/NavigationSystem/Private/NavigationTestingActor.cpp
+++ b/Engine/Source/Runtime/NavigationSystem/Private/NavigationTestingActor.cpp
@@ -286,7 +286,7 @@ void ANavigationTestingActor::UpdateNavData()
 		UNavigationSystemV1* NavSys = FNavigationSystem::GetCurrent<UNavigationSystemV1>(GetWorld());
 		if (NavSys)
 		{
-			MyNavData = NavSys->GetNavDataForProps(NavAgentProps);
+			MyNavData = NavSys->GetNavDataForProps(NavAgentProps, GetActorLocation());
 		}
 	}
 }
