Manual merge of CL#4554906 over from the original stream

Original description:
Fixed a bug in FAISenseAffiliationFilter::ShouldDetectAll. Was using binary '&' rathen than logical

diff --git a/Engine/Source/Runtime/AIModule/Classes/Perception/AIPerceptionTypes.h b/Engine/Source/Runtime/AIModule/Classes/Perception/AIPerceptionTypes.h
--- a/Engine/Source/Runtime/AIModule/Classes/Perception/AIPerceptionTypes.h
+++ b/Engine/Source/Runtime/AIModule/Classes/Perception/AIPerceptionTypes.h
@@ -216,7 +216,7 @@ struct AIMODULE_API FAISenseAffiliationFilter
 	uint32 bDetectFriendlies : 1;
 	
 	uint8 GetAsFlags() const { return (bDetectEnemies << ETeamAttitude::Hostile) | (bDetectNeutrals << ETeamAttitude::Neutral) | (bDetectFriendlies << ETeamAttitude::Friendly); }
-	FORCEINLINE bool ShouldDetectAll() const { return (bDetectEnemies & bDetectNeutrals & bDetectFriendlies); }
+	FORCEINLINE bool ShouldDetectAll() const { return (bDetectEnemies && bDetectNeutrals && bDetectFriendlies); }
 
 	static FORCEINLINE uint8 DetectAllFlags() { return (1 << ETeamAttitude::Hostile) | (1 << ETeamAttitude::Neutral) | (1 << ETeamAttitude::Friendly); }
 
