Manual merge of CL#4554902 over from the original stream

Original description:
Added an BBKey accessor type that optimizes and simplifies BB key accessing #UE4

diff --git a/Engine/Source/Runtime/AIModule/Classes/BehaviorTree/BlackboardComponent.h b/Engine/Source/Runtime/AIModule/Classes/BehaviorTree/BlackboardComponent.h
--- a/Engine/Source/Runtime/AIModule/Classes/BehaviorTree/BlackboardComponent.h
+++ b/Engine/Source/Runtime/AIModule/Classes/BehaviorTree/BlackboardComponent.h
@@ -37,6 +37,88 @@ namespace EBlackboardDescription
 	};
 }
 
+/**
+ *	A helper type that improved performance of reading data from BB
+ *	It's meant for a specific use-case:
+
+ *		1.	you have a logical property you want to use both in C++ code
+ *			as well as being reflected in the BB
+ *		2.	you only ever set this property in native code
+ *		
+ *		If those two are true then add a member variable of type FBlackboardCachedAccessor
+ *		like so:
+
+ *			FBlackboardCachedAccessor<UBlackboardKeyType_Bool> BBEnemyInMeleeRangeKey;
+
+ *		and from this point on whenever you set or read the value use this variable.
+ *		This will make reading almost free.
+ *
+ *		Before you use the variable you need to initialize it with appropriate BB asset.
+ *		This is best done in AAIController::InitializeBlackboard override, like so:
+ *
+ *			const FBlackboard::FKey EnemyInMeleeRangeKey = BlackboardAsset.GetKeyID(TEXT("EnemyInMeleeRange"));
+ *			BBEnemyInMeleeRangeKey = FBlackboardCachedAccessor<UBlackboardKeyType_Bool>(BlackboardComp, EnemyInMeleeRangeKey);
+ *
+ *		Best used with numerical and boolean types. No guarantees made when using pointer types.
+
+ *	@note does not automatically support BB component or asset change */
+template<typename TBlackboardKey>
+struct FBBKeyCachedAccessor
+{
+private:
+	FBlackboard::FKey BBKey;
+	typedef typename TBlackboardKey::FDataType FStoredType;
+	FStoredType CachedValue;
+public:
+	FBBKeyCachedAccessor() : BBKey(FBlackboard::InvalidKey), CachedValue(TBlackboardKey::InvalidValue)
+	{}
+
+	FBBKeyCachedAccessor(UBlackboardComponent& BBComponent, FBlackboard::FKey InBBKey)
+	{
+		ensure(InBBKey != FBlackboard::InvalidKey);
+		if (ensure(BBComponent.IsKeyOfType<TBlackboardKey>(InBBKey)))
+		{
+			BBKey = InBBKey;
+			CachedValue = BBComponent.GetValue<TBlackboardKey>(InBBKey);
+		}
+	}
+
+	template<typename T2>
+	FORCEINLINE bool SetValue(UBlackboardComponent& BBComponent, const T2 InValue)
+	{
+		return SetValue(BBComponent, FStoredType(InValue));
+	}
+
+	/** @return True is value has changed*/
+	FORCEINLINE bool SetValue(UBlackboardComponent& BBComponent, const FStoredType InValue)
+	{
+		ensure(BBKey != FBlackboard::InvalidKey);
+		if (InValue != CachedValue)
+		{
+			CachedValue = InValue;
+			BBComponent.SetValue<TBlackboardKey>(BBKey, InValue);
+			return true;
+		}
+		return false;
+	}
+
+	FORCEINLINE const FStoredType& Get() const
+	{
+		ensure(BBKey != FBlackboard::InvalidKey);
+		return CachedValue;
+	}
+
+	template<typename T2>
+	FORCEINLINE T2 Get() const
+	{
+		ensure(BBKey != FBlackboard::InvalidKey);
+		return (T2)CachedValue;
+	}
+
+	bool IsValid() const { return BBKey != FBlackboard::InvalidKey; }
+};
+
+
 UCLASS(ClassGroup = AI, meta = (BlueprintSpawnableComponent))
 class AIMODULE_API UBlackboardComponent : public UActorComponent
 {
