Manual merge of CL#4554914 over from the original stream

Original description:
Added an option to BTTask_BlueprintBase to control how often it ticks

diff --git a/Engine/Source/Runtime/AIModule/Classes/AITypes.h b/Engine/Source/Runtime/AIModule/Classes/AITypes.h
--- a/Engine/Source/Runtime/AIModule/Classes/AITypes.h
+++ b/Engine/Source/Runtime/AIModule/Classes/AITypes.h
@@ -597,4 +597,49 @@ enum class EGenericAICheck : uint8
 	IsTrue,
 
 	MAX UMETA(Hidden)
-};
\ No newline at end of file
+};
+
+USTRUCT()
+struct FIntervalCountdown
+{
+	GENERATED_BODY()
+
+	UPROPERTY(EditAnywhere, Category=Time)
+	float Interval;
+
+	float TimeLeft;
+
+	explicit FIntervalCountdown(const float InInterval = -1.f) : Interval(InInterval), TimeLeft(0) {}
+
+	void Set(const float InNewTimeLeft)
+	{
+		TimeLeft = InNewTimeLeft;
+	}
+
+	/**	@return True if time's up */
+	bool Tick(const float TimeDelta)
+	{
+		TimeLeft -= TimeDelta;
+		return TimeLeft <= 0;
+	}
+
+	/** resets the countdown back to Interval */
+	void Reset()
+	{
+		TimeLeft = Interval;
+	}
+
+	/** Note that this function returns meaningful results only if Interval > 0*/
+	float GetElapsedTime() const
+	{
+		ensure(Interval > 0);
+		// using -TimeLeft because at this point TimeLeft is negative
+		return (Interval - TimeLeft);
+	}
+
+	/** @return If Interval > 0 returns time accumulated since resetting. Oterwise returns FallbackValue */
+	float GetElapsedTimeWithFallback(const float FallbackValue) const
+	{
+		return Interval > 0 ? GetElapsedTime() : FallbackValue;
+	}
+};
diff --git a/Engine/Source/Runtime/AIModule/Classes/BehaviorTree/Tasks/BTTask_BlueprintBase.h b/Engine/Source/Runtime/AIModule/Classes/BehaviorTree/Tasks/BTTask_BlueprintBase.h
--- a/Engine/Source/Runtime/AIModule/Classes/BehaviorTree/Tasks/BTTask_BlueprintBase.h
+++ b/Engine/Source/Runtime/AIModule/Classes/BehaviorTree/Tasks/BTTask_BlueprintBase.h
@@ -49,6 +49,11 @@ protected:
 	UPROPERTY(Transient)
 	AActor* ActorOwner;
 
+	/** If any of the Tick functions is implemented, how ofter should they be ticked.
+	 *	Values < 0 mean 'every tick'. */
+	UPROPERTY(EditAnywhere, Category = Task)
+	FIntervalCountdown TickInterval;
+	
 	/** temporary variable for ReceiveExecute(Abort)-FinishExecute(Abort) chain */
 	mutable TEnumAsByte<EBTNodeResult::Type> CurrentCallResult;
 
diff --git a/Engine/Source/Runtime/AIModule/Private/BehaviorTree/Tasks/BTTask_BlueprintBase.cpp b/Engine/Source/Runtime/AIModule/Private/BehaviorTree/Tasks/BTTask_BlueprintBase.cpp
--- a/Engine/Source/Runtime/AIModule/Private/BehaviorTree/Tasks/BTTask_BlueprintBase.cpp
+++ b/Engine/Source/Runtime/AIModule/Private/BehaviorTree/Tasks/BTTask_BlueprintBase.cpp
@@ -97,20 +97,29 @@ EBTNodeResult::Type UBTTask_BlueprintBase::AbortTask(UBehaviorTreeComponent& Own
 
 void UBTTask_BlueprintBase::TickTask(UBehaviorTreeComponent& OwnerComp, uint8* NodeMemory, float DeltaSeconds)
 {
-	if (AIOwner != nullptr && (ReceiveTickImplementations & FBTNodeBPImplementationHelper::AISpecific))
+	if (TickInterval.Tick(DeltaSeconds))
 	{
-		ReceiveTickAI(AIOwner, AIOwner->GetPawn(), DeltaSeconds);
-	}
-	else if (ReceiveTickImplementations & FBTNodeBPImplementationHelper::Generic)
-	{
-		ReceiveTick(ActorOwner, DeltaSeconds);
+		DeltaSeconds = TickInterval.GetElapsedTimeWithFallback(DeltaSeconds);
+
+		if (AIOwner != nullptr && (ReceiveTickImplementations & FBTNodeBPImplementationHelper::AISpecific))
+		{
+			ReceiveTickAI(AIOwner, AIOwner->GetPawn(), DeltaSeconds);
+		}
+		else if (ReceiveTickImplementations & FBTNodeBPImplementationHelper::Generic)
+		{
+			ReceiveTick(ActorOwner, DeltaSeconds);
+		}
+
+		TickInterval.Reset();
 	}
 }
 
 void UBTTask_BlueprintBase::OnTaskFinished(UBehaviorTreeComponent& OwnerComp, uint8* NodeMemory, EBTNodeResult::Type TaskResult)
 {
 	Super::OnTaskFinished(OwnerComp, NodeMemory, TaskResult);
-	
+
+	TickInterval.Set(0); // so that we tick as soon as enabled back
+
 	if (TaskResult != EBTNodeResult::InProgress)
 	{
 		BlueprintNodeHelpers::AbortLatentActions(OwnerComp, *this);
