Added number of dirty areas to GameplayDebugger category navmesh

#rnx

diff --git a/Engine/Source/Runtime/AIModule/Private/GameplayDebugger/GameplayDebuggerCategory_Navmesh.cpp b/Engine/Source/Runtime/AIModule/Private/GameplayDebugger/GameplayDebuggerCategory_Navmesh.cpp
--- a/Engine/Source/Runtime/AIModule/Private/GameplayDebugger/GameplayDebuggerCategory_Navmesh.cpp
+++ b/Engine/Source/Runtime/AIModule/Private/GameplayDebugger/GameplayDebuggerCategory_Navmesh.cpp
@@ -18,10 +18,10 @@ namespace
 FGameplayDebuggerCategory_Navmesh::FGameplayDebuggerCategory_Navmesh()
 {
 	bShowOnlyWithDebugActor = false;
-	bShowCategoryName = false;
 	bShowDataPackReplication = true;
 	CollectDataInterval = 5.0f;
 	SetDataPackReplication<FNavMeshSceneProxyData>(&NavmeshRenderData);
+	SetDataPackReplication<FRepData>(&DataPack);
 }
 
 TSharedRef<FGameplayDebuggerCategory> FGameplayDebuggerCategory_Navmesh::MakeInstance()
@@ -29,6 +29,11 @@ TSharedRef<FGameplayDebuggerCategory> FGameplayDebuggerCategory_Navmesh::MakeIns
 	return MakeShareable(new FGameplayDebuggerCategory_Navmesh());
 }
 
+void FGameplayDebuggerCategory_Navmesh::FRepData::Serialize(FArchive& Ar)
+{
+	Ar << NumDirtyAreas;
+}
+
 void FGameplayDebuggerCategory_Navmesh::CollectData(APlayerController* OwnerPC, AActor* DebugActor)
 {
 #if WITH_RECAST
@@ -37,13 +42,18 @@ void FGameplayDebuggerCategory_Navmesh::CollectData(APlayerController* OwnerPC,
 	APawn* PlayerPawn = OwnerPC ? OwnerPC->GetPawnOrSpectator() : nullptr;
 	APawn* DebugActorAsPawn = Cast<APawn>(DebugActor);
 	APawn* DestPawn = DebugActorAsPawn ? DebugActorAsPawn : PlayerPawn;
-	if (OwnerPC && DestPawn)
+	if (OwnerPC != nullptr)
 	{
 		UNavigationSystemV1* NavSys = FNavigationSystem::GetCurrent<UNavigationSystemV1>(OwnerPC->GetWorld());
 		if (NavSys) 
 		{
-			const FNavAgentProperties& NavAgentProperties = DestPawn->GetNavAgentPropertiesRef();
-			NavData = Cast<const ARecastNavMesh>(NavSys->GetNavDataForProps(NavAgentProperties, DestPawn->GetNavAgentLocation()));
+			DataPack.NumDirtyAreas = NavSys->GetNumDirtyAreas();
+
+			if (DestPawn != nullptr)
+			{
+				const FNavAgentProperties& NavAgentProperties = DestPawn->GetNavAgentPropertiesRef();
+				NavData = Cast<const ARecastNavMesh>(NavSys->GetNavDataForProps(NavAgentProperties, DestPawn->GetNavAgentLocation()));
+			}
 		}
 	}
 
@@ -79,6 +89,11 @@ void FGameplayDebuggerCategory_Navmesh::CollectData(APlayerController* OwnerPC,
 #endif // WITH_RECAST
 }
 
+void FGameplayDebuggerCategory_Navmesh::DrawData(APlayerController* OwnerPC, FGameplayDebuggerCanvasContext& CanvasContext)
+{
+	CanvasContext.Printf(TEXT("Num dirty areas: {%s}%d"), DataPack.NumDirtyAreas > 0 ? TEXT("red") : TEXT("green"), DataPack.NumDirtyAreas);
+}
+
 void FGameplayDebuggerCategory_Navmesh::OnDataPackReplicated(int32 DataPackId)
 {
 	MarkRenderStateDirty();
diff --git a/Engine/Source/Runtime/AIModule/Private/GameplayDebugger/GameplayDebuggerCategory_Navmesh.h b/Engine/Source/Runtime/AIModule/Private/GameplayDebugger/GameplayDebuggerCategory_Navmesh.h
--- a/Engine/Source/Runtime/AIModule/Private/GameplayDebugger/GameplayDebuggerCategory_Navmesh.h
+++ b/Engine/Source/Runtime/AIModule/Private/GameplayDebugger/GameplayDebuggerCategory_Navmesh.h
@@ -16,13 +16,22 @@ public:
 	FGameplayDebuggerCategory_Navmesh();
 
 	virtual void CollectData(APlayerController* OwnerPC, AActor* DebugActor) override;
+	virtual void DrawData(APlayerController* OwnerPC, FGameplayDebuggerCanvasContext& CanvasContext) override;
 	virtual FDebugRenderSceneProxy* CreateDebugSceneProxy(const UPrimitiveComponent* InComponent, FDebugDrawDelegateHelper*& OutDelegateHelper) override;
 	virtual void OnDataPackReplicated(int32 DataPackId) override;
 
 	static TSharedRef<FGameplayDebuggerCategory> MakeInstance();
 
 protected:
+
+	struct FRepData
+	{
+		int32 NumDirtyAreas = 0;
+		void Serialize(FArchive& Ar);
+	};
+
 	FNavMeshSceneProxyData NavmeshRenderData;
+	FRepData DataPack;
 };
 
 #endif // WITH_GAMEPLAY_DEBUGGER
diff --git a/Engine/Source/Runtime/NavigationSystem/Private/NavigationSystem.cpp b/Engine/Source/Runtime/NavigationSystem/Private/NavigationSystem.cpp
--- a/Engine/Source/Runtime/NavigationSystem/Private/NavigationSystem.cpp
+++ b/Engine/Source/Runtime/NavigationSystem/Private/NavigationSystem.cpp
@@ -2330,6 +2330,11 @@ void UNavigationSystemV1::AddDirtyAreas(const TArray<FBox>& NewAreas, int32 Flag
 	}
 }
 
+int32 UNavigationSystemV1::GetNumDirtyAreas() const
+{
+	return DefaultDirtyAreasController.GetNumDirtyAreas();
+}
+
 bool UNavigationSystemV1::HasDirtyAreasQueued() const
 {
 	return DefaultDirtyAreasController.IsDirty();
@@ -2524,7 +2529,7 @@ void UNavigationSystemV1::UpdateAttachedActorsInNavOctree(AActor& RootActor)
 	if (GetAllAttachedActors(RootActor, UniqueAttachedActors) > 0)
 	{
 		for (AActor* AttachedActor : UniqueAttachedActors)
-	{
+		{
 			UpdateActorAndComponentsInNavOctree(*AttachedActor, /*bUpdateAttachedActors=*/false);
 		}
 	}
@@ -2569,7 +2574,7 @@ void UNavigationSystemV1::UpdateNavOctreeElement(UObject* ElementOwner, INavRele
 void UNavigationSystemV1::UpdateNavOctreeParentChain(UObject* ElementOwner, bool bSkipElementOwnerUpdate)
 {
 	if (ElementOwner)
-		{
+	{
 		FNavigationDataHandler(DefaultOctreeController, DefaultDirtyAreasController).UpdateNavOctreeParentChain(*ElementOwner, bSkipElementOwnerUpdate);
 	}
 }
diff --git a/Engine/Source/Runtime/NavigationSystem/Public/NavigationDirtyAreasController.h b/Engine/Source/Runtime/NavigationSystem/Public/NavigationDirtyAreasController.h
--- a/Engine/Source/Runtime/NavigationSystem/Public/NavigationDirtyAreasController.h
+++ b/Engine/Source/Runtime/NavigationSystem/Public/NavigationDirtyAreasController.h
@@ -30,7 +30,8 @@ struct NAVIGATIONSYSTEM_API FNavigationDirtyAreasController
 	void Tick(float DeltaSeconds, const TArray<ANavigationData*>& NavDataSet, bool bForceRebuilding = false);
 	void AddArea(const FBox& NewArea, int32 Flags);
 	
-	bool IsDirty() const { return DirtyAreas.Num() > 0; }
+	bool IsDirty() const { return GetNumDirtyAreas() > 0; }
+	int32 GetNumDirtyAreas() const { return DirtyAreas.Num(); }
 
 #if !UE_BUILD_SHIPPING
 	bool HadDirtyAreasReportedWhileAccumulationLocked() const { return bCanAccumulateDirtyAreas == false && bDirtyAreasReportedWhileAccumulationLocked; }
diff --git a/Engine/Source/Runtime/NavigationSystem/Public/NavigationSystem.h b/Engine/Source/Runtime/NavigationSystem/Public/NavigationSystem.h
--- a/Engine/Source/Runtime/NavigationSystem/Public/NavigationSystem.h
+++ b/Engine/Source/Runtime/NavigationSystem/Public/NavigationSystem.h
@@ -609,6 +609,7 @@ public:
 	void AddDirtyArea(const FBox& NewArea, int32 Flags);
 	void AddDirtyAreas(const TArray<FBox>& NewAreas, int32 Flags);
 	bool HasDirtyAreasQueued() const;
+	int32 GetNumDirtyAreas() const;
 
 	const FNavigationOctree* GetNavOctree() const { return DefaultOctreeController.GetOctree(); }
 	FNavigationOctree* GetMutableNavOctree() { return DefaultOctreeController.GetMutableOctree(); }
