Resaved Odin maps containing rogue EQS manager instances #Odin

Also, changed how we report detecting rogue AISubsystems instances. Removed the ensure.

diff --git a/Engine/Source/Runtime/AIModule/Private/AISubsystem.cpp b/Engine/Source/Runtime/AIModule/Private/AISubsystem.cpp
--- a/Engine/Source/Runtime/AIModule/Private/AISubsystem.cpp
+++ b/Engine/Source/Runtime/AIModule/Private/AISubsystem.cpp
@@ -13,10 +13,19 @@ UAISubsystem::UAISubsystem(const FObjectInitializer& ObjectInitializer)
 	if (HasAnyFlags(RF_ClassDefaultObject) == false)
 	{
 		AISystem = Cast<UAISystem>(GetOuter());
-		ensure(AISystem);
 		UE_CLOG(AISystem == nullptr, LogAISub, Error, TEXT("%s is an invalid outer for UAISubsystem instance %s")
 			, *GetName(), *GetNameSafe(GetOuter()));
+#if !(UE_BUILD_SHIPPING || UE_BUILD_TEST)
+		if (AISystem == nullptr 
+			&& GetOuter()
+			&& (IsRunningCommandlet() 
+				|| (GIsEditor && GetWorld() && GetWorld()->WorldType == EWorldType::Editor)))
+		{
+			// not calling MarkPackageDirty on this because it might be marked as transient 
+			GetOuter()->MarkPackageDirty();
+		}		
 	}
+#endif // !(UE_BUILD_SHIPPING || UE_BUILD_TEST)
 }
 
 UWorld* UAISubsystem::GetWorld() const
