Made UPathFollowingComponent::Cleanup clear out the MovementComp it points to to ensure full hand-off when AIController::Unpossess is called #UE4

diff --git a/Engine/Source/Runtime/AIModule/Private/Navigation/PathFollowingComponent.cpp b/Engine/Source/Runtime/AIModule/Private/Navigation/PathFollowingComponent.cpp
--- a/Engine/Source/Runtime/AIModule/Private/Navigation/PathFollowingComponent.cpp
+++ b/Engine/Source/Runtime/AIModule/Private/Navigation/PathFollowingComponent.cpp
@@ -612,7 +612,8 @@ void UPathFollowingComponent::Initialize()
 
 void UPathFollowingComponent::Cleanup()
 {
-	// empty in base class
+	SetMovementComponent(nullptr);
+	Reset();
 }
 
 void UPathFollowingComponent::UpdateCachedComponents()
