PR #5413: Make EQ Instance BP Wrapper const-correct (Contributed by [REDACTED_PERSON])

diff --git a/Engine/Source/Runtime/AIModule/Classes/EnvironmentQuery/EnvQueryInstanceBlueprintWrapper.h b/Engine/Source/Runtime/AIModule/Classes/EnvironmentQuery/EnvQueryInstanceBlueprintWrapper.h
--- a/Engine/Source/Runtime/AIModule/Classes/EnvironmentQuery/EnvQueryInstanceBlueprintWrapper.h
+++ b/Engine/Source/Runtime/AIModule/Classes/EnvironmentQuery/EnvQueryInstanceBlueprintWrapper.h
@@ -52,15 +52,15 @@ public:
 	void RunQuery(const EEnvQueryRunMode::Type InRunMode, FEnvQueryRequest& QueryRequest);
 
 	UFUNCTION(BlueprintPure, Category = "AI|EQS")
-	float GetItemScore(int32 ItemIndex);
+	float GetItemScore(int32 ItemIndex) const;
 
 	/** return an array filled with resulting actors. Note that it makes sense only if ItemType is a EnvQueryItemType_ActorBase-derived type */
 	UFUNCTION(BlueprintPure, Category = "AI|EQS")
-	TArray<AActor*> GetResultsAsActors();
+	TArray<AActor*> GetResultsAsActors() const;
 
 	/** returns an array of locations generated by the query. If the query generated Actors the the array is filled with their locations */
 	UFUNCTION(BlueprintPure, Category = "AI|EQS")
-	TArray<FVector> GetResultsAsLocations();
+	TArray<FVector> GetResultsAsLocations() const;
 
 	UFUNCTION(BlueprintCallable, Category = "AI|EQS")
 	void SetNamedParam(FName ParamName, float Value);
diff --git a/Engine/Source/Runtime/AIModule/Private/EnvironmentQuery/EnvQueryInstanceBlueprintWrapper.cpp b/Engine/Source/Runtime/AIModule/Private/EnvironmentQuery/EnvQueryInstanceBlueprintWrapper.cpp
--- a/Engine/Source/Runtime/AIModule/Private/EnvironmentQuery/EnvQueryInstanceBlueprintWrapper.cpp
+++ b/Engine/Source/Runtime/AIModule/Private/EnvironmentQuery/EnvQueryInstanceBlueprintWrapper.cpp
@@ -31,12 +31,12 @@ void UEnvQueryInstanceBlueprintWrapper::OnQueryFinished(TSharedPtr<FEnvQueryResu
 	}
 }
 
-float UEnvQueryInstanceBlueprintWrapper::GetItemScore(int32 ItemIndex)
+float UEnvQueryInstanceBlueprintWrapper::GetItemScore(int32 ItemIndex) const
 {
 	return QueryResult.IsValid() ? QueryResult->GetItemScore(ItemIndex) : -1.f;
 }
 
-TArray<AActor*> UEnvQueryInstanceBlueprintWrapper::GetResultsAsActors()
+TArray<AActor*> UEnvQueryInstanceBlueprintWrapper::GetResultsAsActors() const
 {
 	TArray<AActor*> Results;
 
@@ -55,7 +55,7 @@ TArray<AActor*> UEnvQueryInstanceBlueprintWrapper::GetResultsAsActors()
 	return Results;
 }
 
-TArray<FVector> UEnvQueryInstanceBlueprintWrapper::GetResultsAsLocations()
+TArray<FVector> UEnvQueryInstanceBlueprintWrapper::GetResultsAsLocations() const
 {
 	TArray<FVector> Results;
 
