Made the key for TActorPerceptionContainer an int64 as it's not a managed pointer and the code is clearer this way.

diff --git a/Engine/Source/Runtime/AIModule/Classes/Perception/AIPerceptionComponent.h b/Engine/Source/Runtime/AIModule/Classes/Perception/AIPerceptionComponent.h
--- a/Engine/Source/Runtime/AIModule/Classes/Perception/AIPerceptionComponent.h
+++ b/Engine/Source/Runtime/AIModule/Classes/Perception/AIPerceptionComponent.h
@@ -165,7 +165,10 @@ class AIMODULE_API UAIPerceptionComponent : public UActorComponent
 	
 	static const int32 InitialStimuliToProcessArraySize;
 
-	typedef TMap<AActor*, FActorPerceptionInfo> TActorPerceptionContainer;
+	/** The uint64 is the address of the Actor. As we are using the address raw as a key to the TMap there's potential the actor it 
+	 *  points to will be GCed in future (we aren't tagging PerceptualData UPROPERTY intentionally for optimization purposes).
+	 */
+	typedef TMap<uint64, FActorPerceptionInfo> TActorPerceptionContainer;
 	typedef TActorPerceptionContainer FActorPerceptionContainer;
 
 protected:
@@ -188,6 +191,11 @@ protected:
 private:
 	FPerceptionListenerID PerceptionListenerId;
 
+	/**@TODO there is a rare but possible issue here. Actors could be set to Endplay() and GCed between calls to RemoveDeadData, 
+	 * infact EndPlay and GC can occur in same frame. Either we need to take the hit and make PerceptualData a UPROPERTY or we need to come up
+	 * with a different indexing scheme. Currently we could add an new Actor (at the same address as a GCed one), if we then add a new 
+	 * PerceptualData record we could end up merging the new results with the GCed Actor's record instead of making a fresh one.
+	 */
 	FActorPerceptionContainer PerceptualData;
 		
 protected:	
@@ -231,7 +239,7 @@ public:
 	FORCEINLINE FPerceptionListenerID GetListenerId() const { return PerceptionListenerId; }
 
 	FVector GetActorLocation(const AActor& Actor) const;
-	FORCEINLINE const FActorPerceptionInfo* GetActorInfo(const AActor& Actor) const { return PerceptualData.Find(&Actor); }
+	FORCEINLINE const FActorPerceptionInfo* GetActorInfo(const AActor& Actor) const { return PerceptualData.Find(reinterpret_cast<uint64>(&Actor)); }
 	FORCEINLINE FActorPerceptionContainer::TIterator GetPerceptualDataIterator() { return FActorPerceptionContainer::TIterator(PerceptualData); }
 	FORCEINLINE FActorPerceptionContainer::TConstIterator GetPerceptualDataConstIterator() const { return FActorPerceptionContainer::TConstIterator(PerceptualData); }
 
diff --git a/Engine/Source/Runtime/AIModule/Private/Perception/AIPerceptionComponent.cpp b/Engine/Source/Runtime/AIModule/Private/Perception/AIPerceptionComponent.cpp
--- a/Engine/Source/Runtime/AIModule/Private/Perception/AIPerceptionComponent.cpp
+++ b/Engine/Source/Runtime/AIModule/Private/Perception/AIPerceptionComponent.cpp
@@ -220,6 +220,8 @@ void UAIPerceptionComponent::CleanUp()
 {
 	if (bCleanedUp == false)
 	{
+		ForgetAll();
+
 		UAIPerceptionSystem* AIPerceptionSys = UAIPerceptionSystem::GetCurrent(GetWorld());
 		if (AIPerceptionSys != nullptr)
 		{
@@ -270,11 +272,13 @@ void UAIPerceptionComponent::GetHostileActors(TArray<AActor*>& OutActors) const
 	OutActors.Reserve(PerceptualData.Num());
 	for (FActorPerceptionContainer::TConstIterator DataIt = GetPerceptualDataConstIterator(); DataIt; ++DataIt)
 	{
-		if (DataIt->Value.bIsHostile && DataIt->Value.HasAnyKnownStimulus())
+		const FActorPerceptionInfo& ActorPerceptionInfo = DataIt->Value;
+
+		if (ActorPerceptionInfo.bIsHostile && ActorPerceptionInfo.HasAnyKnownStimulus())
 		{
-			if (DataIt->Value.Target.IsValid())
+			if (ActorPerceptionInfo.Target.IsValid())
 			{
-				OutActors.Add(DataIt->Value.Target.Get());
+				OutActors.Add(ActorPerceptionInfo.Target.Get());
 			}
 			else
 			{
@@ -408,7 +412,9 @@ void UAIPerceptionComponent::ProcessStimuli()
 
 	for (FStimulusToProcess& SourcedStimulus : StimuliToProcess)
 	{
-		FActorPerceptionInfo* PerceptualInfo = PerceptualData.Find(SourcedStimulus.Source);
+		const uint64 SourceAddr = reinterpret_cast<uint64>(SourcedStimulus.Source);
+
+		FActorPerceptionInfo* PerceptualInfo = PerceptualData.Find(SourceAddr);
 
 		if (PerceptualInfo == NULL)
 		{
@@ -421,7 +427,7 @@ void UAIPerceptionComponent::ProcessStimuli()
 			else
 			{
 				// create an entry
-				PerceptualInfo = &PerceptualData.Add(SourcedStimulus.Source, FActorPerceptionInfo(SourcedStimulus.Source));
+				PerceptualInfo = &PerceptualData.Add(SourceAddr, FActorPerceptionInfo(SourcedStimulus.Source));
 				// tell it what's our dominant sense
 				PerceptualInfo->DominantSense = DominantSenseID;
 
@@ -512,7 +518,7 @@ bool UAIPerceptionComponent::AgeStimuli(const float ConstPerceptionAgingRate)
 		for (FAIStimulus& Stimulus : ActorPerceptionInfo.LastSensedStimuli)
 		{
 			// Age the stimulus. If it is active but has just expired, mark it as such
-			if (Stimulus.AgeStimulus(ConstPerceptionAgingRate) == false 
+			if (Stimulus.AgeStimulus(ConstPerceptionAgingRate) == false
 				&& (Stimulus.IsActive() || Stimulus.WantsToNotifyOnlyOnPerceptionChange())
 				&& Stimulus.IsExpired() == false)
 			{
@@ -540,7 +546,7 @@ void UAIPerceptionComponent::ForgetActor(AActor* ActorToForget)
 			AIPerceptionSys->OnListenerForgetsActor(*this, *ActorToForget);
 		}
 
-		PerceptualData.Remove(ActorToForget);
+		const int32 NumRemoved = PerceptualData.Remove(reinterpret_cast<uint64>(ActorToForget));
 	}
 }
 
