#core Removed redundant Casts

diff --git a/Engine/Plugins/Experimental/ControlRig/Source/ControlRigEditor/Private/Graph/SControlRigGraphNode.cpp b/Engine/Plugins/Experimental/ControlRig/Source/ControlRigEditor/Private/Graph/SControlRigGraphNode.cpp
--- a/Engine/Plugins/Experimental/ControlRig/Source/ControlRigEditor/Private/Graph/SControlRigGraphNode.cpp
+++ b/Engine/Plugins/Experimental/ControlRig/Source/ControlRigEditor/Private/Graph/SControlRigGraphNode.cpp
@@ -18,11 +18,12 @@
 
 void SControlRigGraphNode::Construct( const FArguments& InArgs )
 {
+	check(InArgs._GraphNodeObj);
 	this->GraphNode = InArgs._GraphNodeObj;
 
 	this->SetCursor( EMouseCursor::CardinalCross );
 
-	UControlRigGraphNode* ControlRigGraphNode = CastChecked<UControlRigGraphNode>(InArgs._GraphNodeObj);
+ 	UControlRigGraphNode* ControlRigGraphNode = InArgs._GraphNodeObj;
 
 	// Re-cache variable info here (unit structure could have changed since last reconstruction, e.g. array add/remove)
 	ControlRigGraphNode->CacheVariableInfo();
diff --git a/Engine/Plugins/Experimental/PythonScriptPlugin/Source/PythonScriptPlugin/Private/PyUtil.cpp b/Engine/Plugins/Experimental/PythonScriptPlugin/Source/PythonScriptPlugin/Private/PyUtil.cpp
--- a/Engine/Plugins/Experimental/PythonScriptPlugin/Source/PythonScriptPlugin/Private/PyUtil.cpp
+++ b/Engine/Plugins/Experimental/PythonScriptPlugin/Source/PythonScriptPlugin/Private/PyUtil.cpp
@@ -983,8 +983,8 @@ FString GetFriendlyStructValue(const UScriptStruct* InStruct, const void* InStru
 	}
 	else
 	{
-		const UScriptStruct* ScriptStruct = CastChecked<UScriptStruct>(InStruct);
-		ScriptStruct->ExportText(FriendlyStructValue, InStructValue, InStructValue, nullptr, InPortFlags, nullptr);
+		check(InStruct);
+		InStruct->ExportText(FriendlyStructValue, InStructValue, InStructValue, nullptr, InPortFlags, nullptr);
 	}
 
 	return FriendlyStructValue;
diff --git a/Engine/Plugins/Runtime/ResonanceAudio/Source/ResonanceAudioEditor/Private/ResonanceAudioReverbPluginPresetFactory.cpp b/Engine/Plugins/Runtime/ResonanceAudio/Source/ResonanceAudioEditor/Private/ResonanceAudioReverbPluginPresetFactory.cpp
--- a/Engine/Plugins/Runtime/ResonanceAudio/Source/ResonanceAudioEditor/Private/ResonanceAudioReverbPluginPresetFactory.cpp
+++ b/Engine/Plugins/Runtime/ResonanceAudio/Source/ResonanceAudioEditor/Private/ResonanceAudioReverbPluginPresetFactory.cpp
@@ -39,7 +39,7 @@ UResonanceAudioReverbPluginPresetFactory::UResonanceAudioReverbPluginPresetFacto
 
 UObject* UResonanceAudioReverbPluginPresetFactory::FactoryCreateNew(UClass* Class, UObject* InParent, FName InName, EObjectFlags Flags, UObject* Context, FFeedbackContext* Warn)
 {
-	return Cast<UObject>(NewObject<UResonanceAudioReverbPluginPreset>(InParent, InName, Flags));
+	return NewObject<UResonanceAudioReverbPluginPreset>(InParent, InName, Flags);
 }
 
 uint32 UResonanceAudioReverbPluginPresetFactory::GetMenuCategories() const
diff --git a/Engine/Source/Developer/AssetTools/Private/AssetTypeActions/AssetTypeActions_Texture2D.cpp b/Engine/Source/Developer/AssetTools/Private/AssetTypeActions/AssetTypeActions_Texture2D.cpp
--- a/Engine/Source/Developer/AssetTools/Private/AssetTypeActions/AssetTypeActions_Texture2D.cpp
+++ b/Engine/Source/Developer/AssetTools/Private/AssetTypeActions/AssetTypeActions_Texture2D.cpp
@@ -103,7 +103,7 @@ void FAssetTypeActions_Texture2D::ExecuteCreateVolumeTexture(TArray<TWeakObjectP
 
 	if( Objects.Num() == 1 )
 	{
-		auto Object = Objects[0].Get();
+		UTexture2D* Object = Objects[0].Get();
 
 		if( Object )
 		{
@@ -114,7 +114,7 @@ void FAssetTypeActions_Texture2D::ExecuteCreateVolumeTexture(TArray<TWeakObjectP
 
 			// Create the factory used to generate the asset
 			UVolumeTextureFactory* Factory = NewObject<UVolumeTextureFactory>();
-			Factory->InitialTexture = CastChecked<UTexture2D>(Object);
+			Factory->InitialTexture = Object;
 			FContentBrowserModule& ContentBrowserModule = FModuleManager::LoadModuleChecked<FContentBrowserModule>("ContentBrowser");
 			ContentBrowserModule.Get().CreateNewAsset(Name, FPackageName::GetLongPackagePath(PackagePath), UVolumeTexture::StaticClass(), Factory);
 		}
diff --git a/Engine/Source/Editor/Persona/Private/AnimationEditorViewportClient.cpp b/Engine/Source/Editor/Persona/Private/AnimationEditorViewportClient.cpp
--- a/Engine/Source/Editor/Persona/Private/AnimationEditorViewportClient.cpp
+++ b/Engine/Source/Editor/Persona/Private/AnimationEditorViewportClient.cpp
@@ -1241,11 +1241,7 @@ void FAnimationViewportClient::DrawMeshBones(UDebugSkelMeshComponent * MeshCompo
 		TArray<FLinearColor> BoneColours;
 		BoneColours.AddUninitialized(MeshComponent->GetNumDrawTransform());
 
-		TArray<int32> SelectedBones;
-		if(UDebugSkelMeshComponent* DebugMeshComponent = Cast<UDebugSkelMeshComponent>(MeshComponent))
-		{
-			SelectedBones = DebugMeshComponent->BonesOfInterest;
-		}
+		TArray<int32> SelectedBones = MeshComponent->BonesOfInterest;
 
 		// we could cache parent bones as we calculate, but right now I'm not worried about perf issue of this
 		const TArray<FBoneIndexType>& DrawBoneIndices = MeshComponent->GetDrawBoneIndices();
diff --git a/Engine/Source/Editor/UMGEditor/Private/WidgetBlueprintEditorUtils.cpp b/Engine/Source/Editor/UMGEditor/Private/WidgetBlueprintEditorUtils.cpp
--- a/Engine/Source/Editor/UMGEditor/Private/WidgetBlueprintEditorUtils.cpp
+++ b/Engine/Source/Editor/UMGEditor/Private/WidgetBlueprintEditorUtils.cpp
@@ -1029,7 +1029,7 @@ TArray<UWidget*> FWidgetBlueprintEditorUtils::PasteWidgets(TSharedRef<FWidgetBlu
 			// If the widget isn't a panel, we'll try it's parent to see if the pasted widget can be a sibling
 			if (!ParentWidget)
 			{
-				ParentWidget = Cast<UPanelWidget>(ParentWidgetRef.GetTemplate()->GetParent());
+				ParentWidget = ParentWidgetRef.GetTemplate()->GetParent();
 			}
 		}
 
diff --git a/Engine/Source/Runtime/AIModule/Private/EnvironmentQuery/EnvQueryTypes.cpp b/Engine/Source/Runtime/AIModule/Private/EnvironmentQuery/EnvQueryTypes.cpp
--- a/Engine/Source/Runtime/AIModule/Private/EnvironmentQuery/EnvQueryTypes.cpp
+++ b/Engine/Source/Runtime/AIModule/Private/EnvironmentQuery/EnvQueryTypes.cpp
@@ -280,10 +280,10 @@ namespace FEQSHelpers
 		if (NavAgent)
 		{
 			const FNavAgentProperties& NavAgentProps = NavAgent->GetNavAgentPropertiesRef();
-			return Cast<const ANavigationData>(NavSys->GetNavDataForProps(NavAgentProps));
+			return NavSys->GetNavDataForProps(NavAgentProps);
 		}
 
-		return Cast<const ANavigationData>(NavSys->GetDefaultNavDataInstance());
+		return NavSys->GetDefaultNavDataInstance();
 	}
 }
 
diff --git a/Engine/Source/Runtime/AIModule/Private/Navigation/MetaNavMeshPath.cpp b/Engine/Source/Runtime/AIModule/Private/Navigation/MetaNavMeshPath.cpp
--- a/Engine/Source/Runtime/AIModule/Private/Navigation/MetaNavMeshPath.cpp
+++ b/Engine/Source/Runtime/AIModule/Private/Navigation/MetaNavMeshPath.cpp
@@ -37,7 +37,7 @@ FMetaNavMeshPath::FMetaNavMeshPath(const TArray<FMetaPathWayPoint>& InWaypoints,
 	: FMetaNavMeshPath()
 {
 	UNavigationSystemV1* NavSys = FNavigationSystem::GetCurrent<UNavigationSystemV1>(Owner.GetWorld());
-	const ANavigationData* NavData = NavSys ? Cast<const ANavigationData>(NavSys->GetNavDataForProps(Owner.GetNavAgentPropertiesRef())) : (const ANavigationData*)nullptr;
+	const ANavigationData* NavData = NavSys ? NavSys->GetNavDataForProps(Owner.GetNavAgentPropertiesRef()) : nullptr;
 
 	if (ensure(NavData))
 	{
@@ -59,7 +59,7 @@ FMetaNavMeshPath::FMetaNavMeshPath(const TArray<FVector>& InWaypoints, const ANa
 FMetaNavMeshPath::FMetaNavMeshPath(const TArray<FVector>& InWaypoints, const AController& Owner) : FMetaNavMeshPath()
 {
 	UNavigationSystemV1* NavSys = FNavigationSystem::GetCurrent<UNavigationSystemV1>(Owner.GetWorld());
-	const ANavigationData* NavData = Cast<const ANavigationData>(NavSys ? NavSys->GetNavDataForProps(Owner.GetNavAgentPropertiesRef()) : nullptr);
+	const ANavigationData* NavData = NavSys ? NavSys->GetNavDataForProps(Owner.GetNavAgentPropertiesRef()) : nullptr;
 
 	if (ensure(NavData))
 	{
diff --git a/Engine/Source/Runtime/ClothingSystemRuntime/Private/ClothingSimulationFactory.cpp b/Engine/Source/Runtime/ClothingSystemRuntime/Private/ClothingSimulationFactory.cpp
--- a/Engine/Source/Runtime/ClothingSystemRuntime/Private/ClothingSimulationFactory.cpp
+++ b/Engine/Source/Runtime/ClothingSystemRuntime/Private/ClothingSimulationFactory.cpp
@@ -41,5 +41,5 @@ bool UClothingSimulationFactoryNv::SupportsRuntimeInteraction()
 
 UClothingSimulationInteractor* UClothingSimulationFactoryNv::CreateInteractor()
 {
-	return CastChecked<UClothingSimulationInteractor>(NewObject<UClothingSimulationInteractorNv>(GetTransientPackage()));
+	return NewObject<UClothingSimulationInteractorNv>(GetTransientPackage());
 }
diff --git a/Engine/Source/Runtime/Engine/Private/Components/SkeletalMeshComponent.cpp b/Engine/Source/Runtime/Engine/Private/Components/SkeletalMeshComponent.cpp
--- a/Engine/Source/Runtime/Engine/Private/Components/SkeletalMeshComponent.cpp
+++ b/Engine/Source/Runtime/Engine/Private/Components/SkeletalMeshComponent.cpp
@@ -1841,7 +1841,7 @@ void USkeletalMeshComponent::WritebackClothingSimulationData()
 		USkinnedMeshComponent* OverrideComponent = nullptr;
 		if(MasterPoseComponent.IsValid())
 		{
-			OverrideComponent = Cast<USkinnedMeshComponent>(MasterPoseComponent.Get());
+			OverrideComponent = MasterPoseComponent.Get();
 
 			// Check if our bone map is actually valid, if not there is no clothing data to build
 			if(MasterBoneMap.Num() == 0)
diff --git a/Engine/Source/Runtime/Engine/Private/DemoNetDriver.cpp b/Engine/Source/Runtime/Engine/Private/DemoNetDriver.cpp
--- a/Engine/Source/Runtime/Engine/Private/DemoNetDriver.cpp
+++ b/Engine/Source/Runtime/Engine/Private/DemoNetDriver.cpp
@@ -1745,7 +1745,7 @@ void UDemoNetDriver::TickCheckpoint()
 		NumLevelsAddedThisFrame = AllLevelStatuses.Num();
 
 		// Write out all of the queued up packets generated while saving the checkpoint
-		WriteDemoFrameFromQueuedDemoPackets( *CheckpointArchive, CastChecked< UDemoNetConnection>( ClientConnection )->QueuedCheckpointPackets );
+		WriteDemoFrameFromQueuedDemoPackets( *CheckpointArchive, ClientConnection->QueuedCheckpointPackets );
 
 		// Get the total checkpoint size
 		const int32 TotalCheckpointSize = CheckpointArchive->TotalSize();
diff --git a/Engine/Source/Runtime/Engine/Private/NetConnection.cpp b/Engine/Source/Runtime/Engine/Private/NetConnection.cpp
--- a/Engine/Source/Runtime/Engine/Private/NetConnection.cpp
+++ b/Engine/Source/Runtime/Engine/Private/NetConnection.cpp
@@ -2783,7 +2783,7 @@ static void	AddSimulatedNetConnections(const TArray<FString>& Args, UWorld* Worl
 	APlayerController* PC = nullptr;
 	for (auto Iterator = BestNetDriver->GetWorld()->GetPlayerControllerIterator(); Iterator; ++Iterator)
 	{
-		PC = Cast<APlayerController>(*Iterator);
+		PC = Iterator->Get();
 		if (PC)
 		{
 			DefaultViewTarget = PC->GetViewTarget();
