EQS Volume Test - Simple EQS Test to check in a point is in one AVolume (volumes list generated either by a context or a AVolume class). The test can do either do only simple bounding box checks or a full physic test for more complex volume shapes.

diff --git a/Engine/Source/Runtime/AIModule/Classes/EnvironmentQuery/Tests/EnvQueryTest_Volume.h b/Engine/Source/Runtime/AIModule/Classes/EnvironmentQuery/Tests/EnvQueryTest_Volume.h
new file mode 100644
--- /dev/null
+++ b/Engine/Source/Runtime/AIModule/Classes/EnvironmentQuery/Tests/EnvQueryTest_Volume.h
@@ -0,0 +1,42 @@
+// Copyright 1998-2019 Epic Games, Inc. All Rights Reserved.
+
+#pragma once
+
+#include "CoreMinimal.h"
+#include "UObject/ObjectMacros.h"
+#include "EnvironmentQuery/EnvQueryTypes.h"
+#include "EnvironmentQuery/EnvQueryTest.h"
+#include "EnvQueryTest_Volume.generated.h"
+
+class AVolume;
+
+UCLASS(MinimalAPI)
+class UEnvQueryTest_Volume : public UEnvQueryTest
+{
+	GENERATED_BODY()
+
+public:
+
+	UEnvQueryTest_Volume(const FObjectInitializer& ObjectInitializer = FObjectInitializer::Get());
+#if WITH_EDITOR
+	virtual void PostEditChangeProperty(FPropertyChangedEvent& PropertyChangedEvent) override;
+#endif
+	virtual void RunTest(FEnvQueryInstance& QueryInstance) const override;
+
+	virtual FText GetDescriptionTitle() const override;
+	virtual FText GetDescriptionDetails() const override;
+
+private:
+
+	/** Context that populates a list of Actors derived from AVolume to test against */
+	UPROPERTY(EditAnywhere, Category = Volumes, meta = (AllowPrivateAccess = "true"))
+	TSubclassOf<class UEnvQueryContext> VolumeContext;
+
+	/** If VolumeContext is null, AVolume Class will be used to populate a list of AVolume to test against */
+	UPROPERTY(EditAnywhere, Category = Volumes, meta = (AllowPrivateAccess = "true"))
+	TSubclassOf<AVolume> VolumeClass;
+
+	/** If bDoComplexVolumeTest is set, it will use a full volume physic test (not only a bounding box test)  */
+	UPROPERTY(EditAnywhere, Category = Volumes, meta = (AllowPrivateAccess = "true"))
+	uint32 bDoComplexVolumeTest : 1;
+};
diff --git a/Engine/Source/Runtime/AIModule/Private/EnvironmentQuery/Tests/EnvQueryTest_Volume.cpp b/Engine/Source/Runtime/AIModule/Private/EnvironmentQuery/Tests/EnvQueryTest_Volume.cpp
new file mode 100644
--- /dev/null
+++ b/Engine/Source/Runtime/AIModule/Private/EnvironmentQuery/Tests/EnvQueryTest_Volume.cpp
@@ -0,0 +1,124 @@
+// Copyright 1998-2019 Epic Games, Inc. All Rights Reserved.
+#include "EnvironmentQuery/Tests/EnvQueryTest_Volume.h"
+
+#include "EngineUtils.h"
+#include "EnvironmentQuery/Items/EnvQueryItemType_VectorBase.h"
+#include "GameFramework/Volume.h"
+
+#define LOCTEXT_NAMESPACE "EnvQueryGenerator"
+
+UEnvQueryTest_Volume::UEnvQueryTest_Volume(const FObjectInitializer& ObjectInitializer) : Super(ObjectInitializer)
+{
+	Cost = EEnvTestCost::Medium;
+	ValidItemType = UEnvQueryItemType_VectorBase::StaticClass();
+	SetWorkOnFloatValues(false);
+
+	bDoComplexVolumeTest = false;
+}
+
+#if WITH_EDITOR
+void UEnvQueryTest_Volume::PostEditChangeProperty(FPropertyChangedEvent& PropertyChangedEvent)
+{
+	Super::PostEditChangeProperty(PropertyChangedEvent);
+	Cost = VolumeContext.Get() ? EEnvTestCost::Medium : EEnvTestCost::High;
+}
+#endif
+
+void UEnvQueryTest_Volume::RunTest(FEnvQueryInstance& QueryInstance) const
+{
+	const UObject* DataOwner = QueryInstance.Owner.Get();
+	BoolValue.BindData(DataOwner, QueryInstance.QueryID);
+
+	struct FVolumeBoundsCache
+	{
+		public:
+			FVolumeBoundsCache(AVolume& InVolume) : Volume(InVolume)
+			{
+				BoundingBox = Volume.GetBounds().GetBox();
+			}
+
+			bool IsPointInsideBoundingBox(const FVector TestPoint) const
+			{
+				return BoundingBox.IsInside(TestPoint);
+			}
+			bool EncompassesPoint(const FVector TestPoint) const
+			{
+				return Volume.EncompassesPoint(TestPoint);
+			}
+
+		private:
+			AVolume& Volume;
+			FBox BoundingBox;
+	};
+
+	TArray<AActor*> ContextVolumeActors;
+	TArray<FVolumeBoundsCache> ContextVolumes;
+	if (VolumeContext)
+	{
+		QueryInstance.PrepareContext(VolumeContext, ContextVolumeActors);
+		ContextVolumes.Reserve(ContextVolumeActors.Num());
+		for (AActor* VolumeActor : ContextVolumeActors)
+		{
+			AVolume* Volume = Cast<AVolume>(VolumeActor);
+			if (Volume)
+			{			
+				ContextVolumes.Add(FVolumeBoundsCache(*Volume));
+			}
+		}
+
+	}
+	else if (VolumeClass)
+	{
+		for (TActorIterator<AVolume> It(QueryInstance.World, VolumeClass); It; ++It)
+		{
+			AVolume* Volume = *It;
+			if (Volume)
+			{
+				ContextVolumes.Add(FVolumeBoundsCache(*Volume));
+			}
+		}
+	}
+
+	const bool bWantsInside = BoolValue.GetValue();	
+
+	for (FEnvQueryInstance::ItemIterator It(this, QueryInstance); It; ++It)
+	{
+		const FVector ItemLocation = GetItemLocation(QueryInstance, It.GetIndex());
+		bool bPointIsInsideAVolume = false;
+		for (const FVolumeBoundsCache& VolumeAndBounds : ContextVolumes)
+		{
+			if (VolumeAndBounds.IsPointInsideBoundingBox(ItemLocation))
+			{
+				if (!bDoComplexVolumeTest || VolumeAndBounds.EncompassesPoint(ItemLocation))
+				{
+					bPointIsInsideAVolume = true;
+					break;
+				}
+			}			
+		}	
+		It.SetScore(TestPurpose, FilterType, bPointIsInsideAVolume, bWantsInside);
+	}
+}
+
+FText UEnvQueryTest_Volume::GetDescriptionTitle() const
+{	
+	FString VolumeDesc(TEXT("No Source Volume Context/Class set"));
+	if (VolumeContext)
+	{
+		VolumeDesc = UEnvQueryTypes::DescribeContext(VolumeContext).ToString();
+	}
+	else if (VolumeClass)
+	{
+		VolumeDesc = VolumeClass->GetDescription();
+	}
+
+	return FText::FromString(FString::Printf(TEXT("%s : %s (%s)"), 
+		*Super::GetDescriptionTitle().ToString(), *VolumeDesc, bDoComplexVolumeTest ? TEXT("Complex Physic Volume Test") : TEXT("Simple Bounding Box Test")));
+}
+
+FText UEnvQueryTest_Volume::GetDescriptionDetails() const
+{
+	return GetDescriptionTitle();
+}
+
+#undef LOCTEXT_NAMESPACE
