No longer trying to resolve FEQSParametrizedQueryExecutionRequest.EQSQueryBlackboardKey BB key selector if the request is not configured to use it (i.e. FEQSParametrizedQueryExecutionRequest.bUseBBKeyForQueryTemplate == false) #UE4

This avoids logging out a confusing warning at runtime.

diff --git a/Engine/Source/Runtime/AIModule/Private/EnvironmentQuery/EnvQueryTypes.cpp b/Engine/Source/Runtime/AIModule/Private/EnvironmentQuery/EnvQueryTypes.cpp
--- a/Engine/Source/Runtime/AIModule/Private/EnvironmentQuery/EnvQueryTypes.cpp
+++ b/Engine/Source/Runtime/AIModule/Private/EnvironmentQuery/EnvQueryTypes.cpp
@@ -394,7 +394,7 @@ void FEQSParametrizedQueryExecutionRequest::InitForOwnerAndBlackboard(UObject& O
 
 	EQSQueryBlackboardKey.AddObjectFilter(&Owner, NAME_None, UEnvQuery::StaticClass());
 
-	if ((QueryConfig.Num() > 0 || bUseBBKeyForQueryTemplate) && BBAsset)
+	if ((bUseBBKeyForQueryTemplate || QueryConfig.Num() > 0) && BBAsset)
 	{
 		for (FAIDynamicParam& RuntimeParam : QueryConfig)
 		{
@@ -404,7 +404,10 @@ void FEQSParametrizedQueryExecutionRequest::InitForOwnerAndBlackboard(UObject& O
 			}
 		}
 
-		EQSQueryBlackboardKey.ResolveSelectedKey(*BBAsset);
+		if (bUseBBKeyForQueryTemplate)
+		{
+			EQSQueryBlackboardKey.ResolveSelectedKey(*BBAsset);
+		}
 	}
 }
 
