Fixed an edge-case in EQS results sorting where we'd sort the results multiple times if doing a time-sliced expensive filter test at the end #UE4

diff --git a/Engine/Source/Runtime/AIModule/Private/EnvironmentQuery/EnvQueryInstance.cpp b/Engine/Source/Runtime/AIModule/Private/EnvironmentQuery/EnvQueryInstance.cpp
--- a/Engine/Source/Runtime/AIModule/Private/EnvironmentQuery/EnvQueryInstance.cpp
+++ b/Engine/Source/Runtime/AIModule/Private/EnvironmentQuery/EnvQueryInstance.cpp
@@ -403,7 +403,7 @@ void FEnvQueryInstance::ExecuteOneStep(float TimeLimit)
 		// item generator uses this flag to alter the scoring behavior
 		bPassOnSingleResult = (bDoingLastTest && Mode == EEnvQueryRunMode::SingleResult && TestObject->CanRunAsFinalCondition());
 
-		if (bPassOnSingleResult)
+		if (bPassOnSingleResult && (CurrentTestStartingItem == 0))
 		{
 			// Since we know we're the last test that is a final condition, if we were scoring previously we should sort the tests now before we test them
 			bool bSortTests = false;
