PR #5703: Flip EQS query score sign when using inverse linear (Contributed by [REDACTED_PERSON])

diff --git a/Engine/Source/Runtime/AIModule/Private/EnvironmentQuery/EnvQueryTest.cpp b/Engine/Source/Runtime/AIModule/Private/EnvironmentQuery/EnvQueryTest.cpp
--- a/Engine/Source/Runtime/AIModule/Private/EnvironmentQuery/EnvQueryTest.cpp
+++ b/Engine/Source/Runtime/AIModule/Private/EnvironmentQuery/EnvQueryTest.cpp
@@ -250,8 +250,11 @@ FText UEnvQueryTest::DescribeFloatTestParams() const
 	}
 	else
 	{
-		FText ScoreSignDesc = (ScoringFactor.DefaultValue > 0) ? LOCTEXT("Greater", "greater") : LOCTEXT("Lesser", "lesser");
-		FText ScoreValueDesc = FText::AsNumber(FMath::Abs(ScoringFactor.DefaultValue), &NumberFormattingOptions);
+		const bool bScalesUp = (ScoringEquation == EEnvTestScoreEquation::InverseLinear) 
+			? (ScoringFactor.DefaultValue < 0) 
+			: (ScoringFactor.DefaultValue > 0);
+		const FText ScoreSignDesc = bScalesUp ? LOCTEXT("Greater", "greater") : LOCTEXT("Lesser", "lesser");
+		const FText ScoreValueDesc = FText::AsNumber(FMath::Abs(ScoringFactor.DefaultValue), &NumberFormattingOptions);
 		ScoreDesc = FText::Format(FText::FromString("{0} {1} [x{2}]"), LOCTEXT("ScorePrefer", "prefer"), ScoreSignDesc, ScoreValueDesc);
 	}
 
