PR #5293: UE-51364 Bug fix (Contributed by [REDACTED_PERSON])
Added a missing parameter to AIMoveTo function to allow the user to disable navigation usage entirely.

diff --git a/Engine/Source/Runtime/AIModule/Classes/Tasks/AITask_MoveTo.h b/Engine/Source/Runtime/AIModule/Classes/Tasks/AITask_MoveTo.h
--- a/Engine/Source/Runtime/AIModule/Classes/Tasks/AITask_MoveTo.h
+++ b/Engine/Source/Runtime/AIModule/Classes/Tasks/AITask_MoveTo.h
@@ -30,10 +30,10 @@ public:
 	EPathFollowingResult::Type GetMoveResult() const { return MoveResult; }
 	bool WasMoveSuccessful() const { return MoveResult == EPathFollowingResult::Success; }
 
-	UFUNCTION(BlueprintCallable, Category = "AI|Tasks", meta = (AdvancedDisplay = "AcceptanceRadius,StopOnOverlap,AcceptPartialPath,bUsePathfinding,bUseContinuosGoalTracking", DefaultToSelf = "Controller", BlueprintInternalUseOnly = "TRUE", DisplayName = "Move To Location or Actor"))
+	UFUNCTION(BlueprintCallable, Category = "AI|Tasks", meta = (AdvancedDisplay = "AcceptanceRadius,StopOnOverlap,AcceptPartialPath,bUsePathfinding,bUseContinuosGoalTracking,ProjectGoalOnNavigation", DefaultToSelf = "Controller", BlueprintInternalUseOnly = "TRUE", DisplayName = "Move To Location or Actor"))
 	static UAITask_MoveTo* AIMoveTo(AAIController* Controller, FVector GoalLocation, AActor* GoalActor = nullptr,
 		float AcceptanceRadius = -1.f, EAIOptionFlag::Type StopOnOverlap = EAIOptionFlag::Default, EAIOptionFlag::Type AcceptPartialPath = EAIOptionFlag::Default,
-		bool bUsePathfinding = true, bool bLockAILogic = true, bool bUseContinuosGoalTracking = false);
+		bool bUsePathfinding = true, bool bLockAILogic = true, bool bUseContinuosGoalTracking = false, EAIOptionFlag::Type ProjectGoalOnNavigation = EAIOptionFlag::Default);
 
 	UE_DEPRECATED(4.12, "This function is now depreacted, please use version with FAIMoveRequest parameter")
 	void SetUp(AAIController* Controller, FVector GoalLocation, AActor* GoalActor = nullptr, float AcceptanceRadius = -1.f, bool bUsePathfinding = true, EAIOptionFlag::Type StopOnOverlap = EAIOptionFlag::Default, EAIOptionFlag::Type AcceptPartialPath = EAIOptionFlag::Default);
diff --git a/Engine/Source/Runtime/AIModule/Private/AIController.cpp b/Engine/Source/Runtime/AIModule/Private/AIController.cpp
--- a/Engine/Source/Runtime/AIModule/Private/AIController.cpp
+++ b/Engine/Source/Runtime/AIModule/Private/AIController.cpp
@@ -633,7 +633,15 @@ FPathFollowingRequestResult AAIController::MoveTo(const FAIMoveRequest& MoveRequ
 
 			if (NavSys && !NavSys->ProjectPointToNavigation(MoveRequest.GetGoalLocation(), ProjectedLocation, INVALID_NAVEXTENT, &AgentProps))
 			{
-				UE_VLOG_LOCATION(this, LogAINavigation, Error, MoveRequest.GetGoalLocation(), 30.f, FColor::Red, TEXT("AAIController::MoveTo failed to project destination location to navmesh"));
+				if (MoveRequest.IsUsingPathfinding())
+				{
+					UE_VLOG_LOCATION(this, LogAINavigation, Error, MoveRequest.GetGoalLocation(), 30.f, FColor::Red, TEXT("AAIController::MoveTo failed to project destination location to navmesh"));
+				}
+				else
+				{
+					UE_VLOG_LOCATION(this, LogAINavigation, Error, MoveRequest.GetGoalLocation(), 30.f, FColor::Red, TEXT("AAIController::MoveTo failed to project destination location to navmesh, path finding is disabled perhaps disable goal projection ?"));
+				}
+
 				bCanRequestMove = false;
 			}
 
diff --git a/Engine/Source/Runtime/AIModule/Private/Tasks/AITask_MoveTo.cpp b/Engine/Source/Runtime/AIModule/Private/Tasks/AITask_MoveTo.cpp
--- a/Engine/Source/Runtime/AIModule/Private/Tasks/AITask_MoveTo.cpp
+++ b/Engine/Source/Runtime/AIModule/Private/Tasks/AITask_MoveTo.cpp
@@ -29,7 +29,7 @@ UAITask_MoveTo::UAITask_MoveTo(const FObjectInitializer& ObjectInitializer)
 
 UAITask_MoveTo* UAITask_MoveTo::AIMoveTo(AAIController* Controller, FVector InGoalLocation, AActor* InGoalActor,
 	float AcceptanceRadius, EAIOptionFlag::Type StopOnOverlap, EAIOptionFlag::Type AcceptPartialPath,
-	bool bUsePathfinding, bool bLockAILogic, bool bUseContinuosGoalTracking)
+	bool bUsePathfinding, bool bLockAILogic, bool bUseContinuosGoalTracking, EAIOptionFlag::Type ProjectGoalOnNavigation)
 {
 	UAITask_MoveTo* MyTask = Controller ? UAITask::NewAITask<UAITask_MoveTo>(*Controller, EAITaskPriority::High) : nullptr;
 	if (MyTask)
@@ -48,6 +48,7 @@ UAITask_MoveTo* UAITask_MoveTo::AIMoveTo(AAIController* Controller, FVector InGo
 		MoveReq.SetReachTestIncludesAgentRadius(FAISystem::PickAIOption(StopOnOverlap, MoveReq.IsReachTestIncludingAgentRadius()));
 		MoveReq.SetAllowPartialPath(FAISystem::PickAIOption(AcceptPartialPath, MoveReq.IsUsingPartialPaths()));
 		MoveReq.SetUsePathfinding(bUsePathfinding);
+		MoveReq.SetProjectGoalLocation(FAISystem::PickAIOption(ProjectGoalOnNavigation, MoveReq.IsProjectingGoal()));
 		if (Controller)
 		{
 			MoveReq.SetNavigationFilter(Controller->GetDefaultNavigationFilterClass());
