Removed c++ comment from EAIOptionFlag since an Editor tooltip gets created from it

diff --git a/Engine/Source/Runtime/AIModule/Classes/AITypes.h b/Engine/Source/Runtime/AIModule/Classes/AITypes.h
--- a/Engine/Source/Runtime/AIModule/Classes/AITypes.h
+++ b/Engine/Source/Runtime/AIModule/Classes/AITypes.h
@@ -57,8 +57,7 @@ namespace EAIOptionFlag
 	enum Type
 	{
 		Default,
-		Enable UMETA(DisplayName = "Yes"),	// UHT was complaining when tried to use True as value instead of Enable
-
+		Enable UMETA(DisplayName = "Yes"),
 		Disable UMETA(DisplayName = "No"),
 
 		MAX UMETA(Hidden)
