Removed restriction on UAISenseConfig_Sight.AutoSuccessRangeFromLastSeenLocation to allow setting it back to -1 #UE4

diff --git a/Engine/Source/Runtime/AIModule/Classes/Perception/AISenseConfig_Sight.h b/Engine/Source/Runtime/AIModule/Classes/Perception/AISenseConfig_Sight.h
--- a/Engine/Source/Runtime/AIModule/Classes/Perception/AISenseConfig_Sight.h
+++ b/Engine/Source/Runtime/AIModule/Classes/Perception/AISenseConfig_Sight.h
@@ -41,11 +41,15 @@ public:
 	FAISenseAffiliationFilter DetectionByAffiliation;
 
 	/** If not an InvalidRange (which is the default), we will always be able to see the target that has already been seen if they are within this range of their last seen location. */
-	UPROPERTY(EditDefaultsOnly, BlueprintReadOnly, Category = "Sense", config, meta = (UIMin = 0.0, ClampMin = 0.0))
+	UPROPERTY(EditDefaultsOnly, BlueprintReadOnly, Category = "Sense", config)
 	float AutoSuccessRangeFromLastSeenLocation;
 		
 	virtual TSubclassOf<UAISense> GetSenseImplementation() const override;
 
+#if WITH_EDITOR
+	virtual void PostEditChangeChainProperty(FPropertyChangedChainEvent& PropertyChangedEvent) override;
+#endif // WITH_EDITOR
+
 #if WITH_GAMEPLAY_DEBUGGER
 	virtual void DescribeSelfToGameplayDebugger(const UAIPerceptionComponent* PerceptionComponent, FGameplayDebuggerCategory* DebuggerCategory) const;
 #endif // WITH_GAMEPLAY_DEBUGGER
diff --git a/Engine/Source/Runtime/AIModule/Private/Perception/AISense.cpp b/Engine/Source/Runtime/AIModule/Private/Perception/AISense.cpp
--- a/Engine/Source/Runtime/AIModule/Private/Perception/AISense.cpp
+++ b/Engine/Source/Runtime/AIModule/Private/Perception/AISense.cpp
@@ -173,7 +173,28 @@ UAISenseConfig_Sight::UAISenseConfig_Sight(const FObjectInitializer& ObjectIniti
 TSubclassOf<UAISense> UAISenseConfig_Sight::GetSenseImplementation() const 
 { 
 	return *Implementation; 
+}	
+
+#if WITH_EDITOR
+void UAISenseConfig_Sight::PostEditChangeChainProperty(FPropertyChangedChainEvent& PropertyChangedEvent)
+{
+	static const FName NAME_AutoSuccessRangeFromLastSeenLocation = GET_MEMBER_NAME_CHECKED(UAISenseConfig_Sight, AutoSuccessRangeFromLastSeenLocation);
+
+	Super::PostEditChangeProperty(PropertyChangedEvent);
+
+	if (PropertyChangedEvent.Property)
+	{
+		const FName PropName = PropertyChangedEvent.Property->GetFName();
+		if (PropName == NAME_AutoSuccessRangeFromLastSeenLocation)
+		{
+			if (AutoSuccessRangeFromLastSeenLocation < 0)
+			{
+				AutoSuccessRangeFromLastSeenLocation = FAISystem::InvalidRange;
+			}
+		}
+	}
 }
+#endif // WITH_EDITOR
 
 #if WITH_GAMEPLAY_DEBUGGER
 void UAISenseConfig_Sight::DescribeSelfToGameplayDebugger(const UAIPerceptionComponent* PerceptionComponent, FGameplayDebuggerCategory* DebuggerCategory) const
