PR #5718: Remove UAISense::DefaultExpirationAge (Contributed by [REDACTED_PERSON])
Small changes to PR: value has been mark as deprecated instead of complety removed.

diff --git a/Engine/Source/Runtime/AIModule/Classes/Perception/AISense.h b/Engine/Source/Runtime/AIModule/Classes/Perception/AISense.h
--- a/Engine/Source/Runtime/AIModule/Classes/Perception/AISense.h
+++ b/Engine/Source/Runtime/AIModule/Classes/Perception/AISense.h
@@ -23,8 +23,9 @@ class AIMODULE_API UAISense : public UObject
 	static const float SuspendNextUpdate;
 
 protected:
-	/** age past which stimulus of this sense are "forgotten"*/
-	UPROPERTY(EditAnywhere, BlueprintReadOnly, Category = "AI Perception", config)
+	UE_DEPRECATED(4.23, "This property will be removed in future versions. Use AISenseConfig::MaxAge instead.")
+	/** age past which stimulus of this sense are "forgotten". (DEPRECATED: This property will be removed in future versions. Use AISenseConfig::MaxAge instead.)*/
+	UPROPERTY(VisibleAnywhere, BlueprintReadOnly, Category = "AI Perception", config)
 	float DefaultExpirationAge;
 
 	UPROPERTY(EditDefaultsOnly, BlueprintReadOnly, Category = "AI Perception", config)
@@ -115,7 +116,13 @@ public:
 	FORCEINLINE void OnListenerUpdate(const FPerceptionListener& NewListener) { OnListenerUpdateDelegate.ExecuteIfBound(NewListener); }
 	FORCEINLINE void OnListenerRemoved(const FPerceptionListener& NewListener) { OnListenerRemovedDelegate.ExecuteIfBound(NewListener); }
 
-	FORCEINLINE float GetDefaultExpirationAge() const { return DefaultExpirationAge; }
+	UE_DEPRECATED(4.23, "This method will be removed in future versions. Perception relies on AISenseConfig::MaxAge so the value returned is no longer used by the perception system.")
+	FORCEINLINE float GetDefaultExpirationAge() const
+	{
+		PRAGMA_DISABLE_DEPRECATION_WARNINGS
+		return DefaultExpirationAge; 
+		PRAGMA_ENABLE_DEPRECATION_WARNINGS
+	}
 
 	bool WantsNewPawnNotification() const { return bWantsNewPawnNotification; }
 	bool ShouldAutoRegisterAllPawnsAsSources() const { return bAutoRegisterAllPawnsAsSources; }
diff --git a/Engine/Source/Runtime/AIModule/Private/Perception/AIPerceptionTypes.cpp b/Engine/Source/Runtime/AIModule/Private/Perception/AIPerceptionTypes.cpp
--- a/Engine/Source/Runtime/AIModule/Private/Perception/AIPerceptionTypes.cpp
+++ b/Engine/Source/Runtime/AIModule/Private/Perception/AIPerceptionTypes.cpp
@@ -18,7 +18,9 @@ FPerceptionListenerCounter FAIGenericID<FPerceptionListenerCounter>::Counter = F
 const float FAIStimulus::NeverHappenedAge = FLT_MAX;
 
 FAIStimulus::FAIStimulus(const UAISense& Sense, float StimulusStrength, const FVector& InStimulusLocation, const FVector& InReceiverLocation, FResult Result, FName InStimulusTag)
-	: Age(0.f), Strength(Result == SensingSucceeded ? StimulusStrength : -1.f)
+	: Age(0.f)
+	, ExpirationAge(FAIStimulus::NeverHappenedAge)
+	, Strength(Result == SensingSucceeded ? StimulusStrength : -1.f)
 	, StimulusLocation(InStimulusLocation)
 	, ReceiverLocation(InReceiverLocation)
 	, Tag(InStimulusTag)
@@ -27,7 +29,6 @@ FAIStimulus::FAIStimulus(const UAISense& Sense, float StimulusStrength, const FV
 	, bExpired(false)
 {
 	Type = Sense.GetSenseID();
-	ExpirationAge = Sense.GetDefaultExpirationAge();
 }
 
 #if !(UE_BUILD_SHIPPING || UE_BUILD_TEST)
diff --git a/Engine/Source/Runtime/AIModule/Private/Perception/AISense.cpp b/Engine/Source/Runtime/AIModule/Private/Perception/AISense.cpp
--- a/Engine/Source/Runtime/AIModule/Private/Perception/AISense.cpp
+++ b/Engine/Source/Runtime/AIModule/Private/Perception/AISense.cpp
@@ -27,7 +27,9 @@ UAISense::UAISense(const FObjectInitializer& ObjectInitializer)
 	, TimeUntilNextUpdate(SuspendNextUpdate)
 	, SenseID(FAISenseID::InvalidID())
 {
+	PRAGMA_DISABLE_DEPRECATION_WARNINGS
 	DefaultExpirationAge = FAIStimulus::NeverHappenedAge;
+	PRAGMA_ENABLE_DEPRECATION_WARNINGS
 
 	bNeedsForgettingNotification = false;
 
