[AISense_Sight]
Restore LastResult when rebuilding queries on a Listener update (RequestStimuliListenerUpdate).

diff --git a/Engine/Source/Runtime/AIModule/Classes/Perception/AISense_Sight.h b/Engine/Source/Runtime/AIModule/Classes/Perception/AISense_Sight.h
--- a/Engine/Source/Runtime/AIModule/Classes/Perception/AISense_Sight.h
+++ b/Engine/Source/Runtime/AIModule/Classes/Perception/AISense_Sight.h
@@ -179,6 +179,7 @@ protected:
 	void OnListenerRemovedImpl(const FPerceptionListener& UpdatedListener);	
 	
 	void GenerateQueriesForListener(const FPerceptionListener& Listener, const FDigestedSightProperties& PropertyDigest);
+	void GenerateQueriesForListener(const FPerceptionListener& Listener, const FDigestedSightProperties& PropertyDigest, TFunctionRef<void(FAISightQuery&)> OnAddedFunc);
 
 	enum FQueriesOperationPostProcess
 	{
@@ -186,10 +187,13 @@ protected:
 		Sort
 	};
 	void RemoveAllQueriesByListener(const FPerceptionListener& Listener, FQueriesOperationPostProcess PostProcess);
+	void RemoveAllQueriesByListener(const FPerceptionListener& Listener, FQueriesOperationPostProcess PostProcess, TFunctionRef<void(const FAISightQuery&)> OnRemoveFunc);
 	void RemoveAllQueriesToTarget(const FAISightTarget::FTargetId& TargetId, FQueriesOperationPostProcess PostProcess);
+	void RemoveAllQueriesToTarget(const FAISightTarget::FTargetId& TargetId, FQueriesOperationPostProcess PostProcess, TFunctionRef<void(const FAISightQuery&)> OnRemoveFunc);
 
 	/** returns information whether new LoS queries have been added */
 	bool RegisterTarget(AActor& TargetActor, FQueriesOperationPostProcess PostProcess);
+	bool RegisterTarget(AActor& TargetActor, FQueriesOperationPostProcess PostProcess, TFunctionRef<void(FAISightQuery&)> OnAddedFunc);
 
 	FORCEINLINE void SortQueries() { SightQueryQueue.Sort(FAISightQuery::FSortPredicate()); }
 
diff --git a/Engine/Source/Runtime/AIModule/Private/Perception/AISense_Sight.cpp b/Engine/Source/Runtime/AIModule/Private/Perception/AISense_Sight.cpp
--- a/Engine/Source/Runtime/AIModule/Private/Perception/AISense_Sight.cpp
+++ b/Engine/Source/Runtime/AIModule/Private/Perception/AISense_Sight.cpp
@@ -425,6 +425,11 @@ void UAISense_Sight::CleanseInvalidSources()
 }
 
 bool UAISense_Sight::RegisterTarget(AActor& TargetActor, FQueriesOperationPostProcess PostProcess)
+{
+	return RegisterTarget(TargetActor, PostProcess, [](FAISightQuery& Query) {});
+}
+
+bool UAISense_Sight::RegisterTarget(AActor& TargetActor, FQueriesOperationPostProcess PostProcess, TFunctionRef<void(FAISightQuery&)> OnAddedFunc)
 {
 	SCOPE_CYCLE_COUNTER(STAT_AI_Sense_Sight_RegisterTarget);
 	
@@ -465,10 +470,12 @@ bool UAISense_Sight::RegisterTarget(AActor& TargetActor, FQueriesOperationPostPr
 			if (FAISenseAffiliationFilter::ShouldSenseTeam(ListenersTeamAgent, TargetActor, PropDigest.AffiliationFlags))
 			{
 				// create a sight query		
-				FAISightQuery SightQuery(ItListener->Key, SightTarget->TargetId);
-				SightQuery.Importance = CalcQueryImportance(ItListener->Value, TargetLocation, PropDigest.SightRadiusSq);
-
-				SightQueryQueue.Add(SightQuery);
+				FAISightQuery& AddedQuery = SightQueryQueue.AddDefaulted_GetRef();
+				AddedQuery.ObserverId = ItListener->Key;
+				AddedQuery.TargetId = SightTarget->TargetId;
+				AddedQuery.Importance = CalcQueryImportance(ItListener->Value, TargetLocation, PropDigest.SightRadiusSq);
+				
+				OnAddedFunc(AddedQuery);
 				bNewQueriesAdded = true;
 			}
 		}
@@ -497,6 +504,11 @@ void UAISense_Sight::OnNewListenerImpl(const FPerceptionListener& NewListener)
 }
 
 void UAISense_Sight::GenerateQueriesForListener(const FPerceptionListener& Listener, const FDigestedSightProperties& PropertyDigest)
+{
+	GenerateQueriesForListener(Listener, PropertyDigest, [](FAISightQuery& Query) {});
+}
+
+void UAISense_Sight::GenerateQueriesForListener(const FPerceptionListener& Listener, const FDigestedSightProperties& PropertyDigest, TFunctionRef<void(FAISightQuery&)> OnAddedFunc)
 {
 	bool bNewQueriesAdded = false;
 	const IGenericTeamAgentInterface* ListenersTeamAgent = Listener.GetTeamAgent();
@@ -514,10 +526,12 @@ void UAISense_Sight::GenerateQueriesForListener(const FPerceptionListener& Liste
 		if (FAISenseAffiliationFilter::ShouldSenseTeam(ListenersTeamAgent, *TargetActor, PropertyDigest.AffiliationFlags))
 		{
 			// create a sight query		
-			FAISightQuery SightQuery(Listener.GetListenerID(), ItTarget->Key);
-			SightQuery.Importance = CalcQueryImportance(Listener, ItTarget->Value.GetLocationSimple(), PropertyDigest.SightRadiusSq);
+			FAISightQuery& AddedQuery = SightQueryQueue.AddDefaulted_GetRef();
+			AddedQuery.ObserverId = Listener.GetListenerID();
+			AddedQuery.TargetId = ItTarget->Key;
+			AddedQuery.Importance = CalcQueryImportance(Listener, ItTarget->Value.GetLocationSimple(), PropertyDigest.SightRadiusSq);
 
-			SightQueryQueue.Add(SightQuery);
+			OnAddedFunc(AddedQuery);
 			bNewQueriesAdded = true;
 		}
 	}
@@ -538,18 +552,31 @@ void UAISense_Sight::OnListenerUpdateImpl(const FPerceptionListener& UpdatedList
 	// 1. remove all queries by this listener
 	// 2. proceed as if it was a new listener
 
-	// remove all queries
-	RemoveAllQueriesByListener(UpdatedListener, DontSort);
-	
 	// see if this listener is a Target as well
 	const FAISightTarget::FTargetId AsTargetId = UpdatedListener.GetBodyActorUniqueID();
 	FAISightTarget* AsTarget = ObservedTargets.Find(AsTargetId);
 	if (AsTarget != NULL)
 	{
-		RemoveAllQueriesToTarget(AsTargetId, DontSort);
 		if (AsTarget->Target.IsValid())
 		{
-			RegisterTarget(*(AsTarget->Target.Get()), DontSort);
+			// if still a valid target then backup list of observers for which the listener was visible to restore in the newly created queries
+			TSet<FPerceptionListenerID> LastVisibleObservers;
+			RemoveAllQueriesToTarget(AsTargetId, DontSort, [&LastVisibleObservers](const FAISightQuery& Query)
+			{
+				if (Query.bLastResult)
+				{
+					LastVisibleObservers.Add(Query.ObserverId);
+				}
+			});
+
+			RegisterTarget(*(AsTarget->Target.Get()), DontSort, [&LastVisibleObservers](FAISightQuery& Query)
+			{
+				Query.bLastResult = LastVisibleObservers.Contains(Query.ObserverId);
+			});
+		}
+		else
+		{
+			RemoveAllQueriesToTarget(AsTargetId, DontSort);
 		}
 	}
 
@@ -557,15 +584,31 @@ void UAISense_Sight::OnListenerUpdateImpl(const FPerceptionListener& UpdatedList
 
 	if (UpdatedListener.HasSense(GetSenseID()))
 	{
+		// if still a valid sense then backup list of targets that were visible by the listener to restore in the newly created queries
+		TSet<FAISightTarget::FTargetId> LastVisibleTargets;
+		RemoveAllQueriesByListener(UpdatedListener, DontSort, [&LastVisibleTargets](const FAISightQuery& Query)
+		{
+			if (Query.bLastResult)
+			{
+				LastVisibleTargets.Add(Query.TargetId);
+			}			
+		});
+
 		const UAISenseConfig_Sight* SenseConfig = Cast<const UAISenseConfig_Sight>(UpdatedListener.Listener->GetSenseConfig(GetSenseID()));
 		check(SenseConfig);
 		FDigestedSightProperties& PropertiesDigest = DigestedProperties.FindOrAdd(ListenerID);
 		PropertiesDigest = FDigestedSightProperties(*SenseConfig);
 
-		GenerateQueriesForListener(UpdatedListener, PropertiesDigest);
+		GenerateQueriesForListener(UpdatedListener, PropertiesDigest, [&LastVisibleTargets](FAISightQuery& Query)
+		{
+			Query.bLastResult = LastVisibleTargets.Contains(Query.TargetId);
+		});
 	}
 	else
 	{
+		// remove all queries
+		RemoveAllQueriesByListener(UpdatedListener, DontSort);
+
 		DigestedProperties.Remove(ListenerID);
 	}
 }
@@ -582,6 +625,11 @@ void UAISense_Sight::OnListenerRemovedImpl(const FPerceptionListener& UpdatedLis
 }
 
 void UAISense_Sight::RemoveAllQueriesByListener(const FPerceptionListener& Listener, FQueriesOperationPostProcess PostProcess)
+{
+	RemoveAllQueriesByListener(Listener, PostProcess, [](const FAISightQuery& Query) {});
+}
+
+void UAISense_Sight::RemoveAllQueriesByListener(const FPerceptionListener& Listener, FQueriesOperationPostProcess PostProcess, TFunctionRef<void(const FAISightQuery&)> OnRemoveFunc)
 {
 	SCOPE_CYCLE_COUNTER(STAT_AI_Sense_Sight_RemoveByListener);
 
@@ -599,6 +647,7 @@ void UAISense_Sight::RemoveAllQueriesByListener(const FPerceptionListener& Liste
 
 		if (SightQuery.ObserverId == ListenerId)
 		{
+			OnRemoveFunc(SightQuery);
 			SightQueryQueue.RemoveAt(QueryIndex, 1, /*bAllowShrinking=*/false);
 			bQueriesRemoved = true;
 		}
@@ -611,6 +660,11 @@ void UAISense_Sight::RemoveAllQueriesByListener(const FPerceptionListener& Liste
 }
 
 void UAISense_Sight::RemoveAllQueriesToTarget(const FAISightTarget::FTargetId& TargetId, FQueriesOperationPostProcess PostProcess)
+{
+	RemoveAllQueriesToTarget(TargetId, PostProcess, [](const FAISightQuery& Query) {});
+}
+
+void UAISense_Sight::RemoveAllQueriesToTarget(const FAISightTarget::FTargetId& TargetId, FQueriesOperationPostProcess PostProcess, TFunctionRef<void(const FAISightQuery&)> OnRemoveFunc)
 {
 	SCOPE_CYCLE_COUNTER(STAT_AI_Sense_Sight_RemoveToTarget);
 
@@ -627,6 +681,7 @@ void UAISense_Sight::RemoveAllQueriesToTarget(const FAISightTarget::FTargetId& T
 
 		if (SightQuery.TargetId == TargetId)
 		{
+			OnRemoveFunc(SightQuery);
 			SightQueryQueue.RemoveAt(QueryIndex, 1, /*bAllowShrinking=*/false);
 			bQueriesRemoved = true;
 		}
