NavLinkProxy: we are no longer updating NavigationRelevancy from the PostLoad callback since it may trigger RefreshNavigationModifiers but the component is not properly Registered in the world yet (NavParent is not cached yet, Actor transform will not be available, etc.)
The fix was from 2936499 for UE-28812 but we can now rely on ANavLinkProxy::PostRegisterAllComponents() to update NavigationRelevency.

diff --git a/Engine/Source/Runtime/AIModule/Private/Navigation/NavLinkProxy.cpp b/Engine/Source/Runtime/AIModule/Private/Navigation/NavLinkProxy.cpp
--- a/Engine/Source/Runtime/AIModule/Private/Navigation/NavLinkProxy.cpp
+++ b/Engine/Source/Runtime/AIModule/Private/Navigation/NavLinkProxy.cpp
@@ -155,11 +155,6 @@ void ANavLinkProxy::PostLoad()
 	{
 		Link.InitializeAreaClass();
 	}
-	
-	if (SmartLinkComp)
-	{
-		SmartLinkComp->SetNavigationRelevancy(bSmartLinkIsRelevant);
-	}
 }
 
 #if ENABLE_VISUAL_LOG
