PR #5170: Add BlueprintType to ETeamAttitude UENUM, to improve blueprint use (switch on enum, etc.) (Contributed by [REDACTED_PERSON])

diff --git a/Engine/Source/Runtime/AIModule/Classes/GenericTeamAgentInterface.h b/Engine/Source/Runtime/AIModule/Classes/GenericTeamAgentInterface.h
--- a/Engine/Source/Runtime/AIModule/Classes/GenericTeamAgentInterface.h
+++ b/Engine/Source/Runtime/AIModule/Classes/GenericTeamAgentInterface.h
@@ -8,7 +8,7 @@
 #include "GameFramework/Actor.h"
 #include "GenericTeamAgentInterface.generated.h"
 
-UENUM()
+UENUM(BlueprintType)
 namespace ETeamAttitude
 {
 	enum Type
