BehaviorTreeComponent: Fixed usages of FScopedBehaviorTreeLock by using a variable to have the proper scope. This will prevent a reentrancy issue in StopTree.

diff --git a/Engine/Source/Runtime/AIModule/Private/BehaviorTree/BehaviorTreeComponent.cpp b/Engine/Source/Runtime/AIModule/Private/BehaviorTree/BehaviorTreeComponent.cpp
--- a/Engine/Source/Runtime/AIModule/Private/BehaviorTree/BehaviorTreeComponent.cpp
+++ b/Engine/Source/Runtime/AIModule/Private/BehaviorTree/BehaviorTreeComponent.cpp
@@ -239,7 +239,7 @@ void UBehaviorTreeComponent::StopTree(EBTStopMode::Type StopMode)
 		return;
 	}
 
-	FScopedBehaviorTreeLock(*this, FScopedBehaviorTreeLock::LockReentry);
+	FScopedBehaviorTreeLock ScopedLock(*this, FScopedBehaviorTreeLock::LockReentry);
 	if (!bRequestedStop)
 	{
 		bRequestedStop = true;
@@ -1211,7 +1211,7 @@ void UBehaviorTreeComponent::TickComponent(float DeltaTime, enum ELevelTick Tick
 	}
 
 	{
-		FScopedBehaviorTreeLock(*this, FScopedBehaviorTreeLock::LockTick);
+		FScopedBehaviorTreeLock ScopedLock(*this, FScopedBehaviorTreeLock::LockTick);
 
 		// tick active parallel tasks (in execution order, before task)
 		for (int32 InstanceIndex = 0; InstanceIndex < InstanceStack.Num(); InstanceIndex++)
