AIPerceptionComponent: HasKnownStimulusOfSense no longer considers expired Stimulus. This makes it consistent with HasAnyKnownStimulus and will ensure that GetKnownPerceivedActors respects its description ("all actors ever perceived in any way (and not forgotten yet) will get fetched")
Also deprecated IsSenseRegistered since it is identical to IsSenseActive.

diff --git a/Engine/Source/Runtime/AIModule/Classes/Perception/AIPerceptionComponent.h b/Engine/Source/Runtime/AIModule/Classes/Perception/AIPerceptionComponent.h
--- a/Engine/Source/Runtime/AIModule/Classes/Perception/AIPerceptionComponent.h
+++ b/Engine/Source/Runtime/AIModule/Classes/Perception/AIPerceptionComponent.h
@@ -91,6 +91,7 @@ struct AIMODULE_API FActorPerceptionInfo
 		return false;
 	}
 
+	/** Indicates currently live (visible) stimulus from any sense */
 	bool HasAnyCurrentStimulus() const
 	{
 		for (const FAIStimulus& Stimulus : LastSensedStimuli)
@@ -105,27 +106,47 @@ struct AIMODULE_API FActorPerceptionInfo
 		return false;
 	}
 
-	/** @note will return FAISystem::InvalidLocation if given sense has never registered related Target actor */
+	/** Retrieves location of the last sensed stimuli for a given sense
+	* @param Sense	The AISenseID of the sense
+	*
+	* @return Location of the last sensed stimuli or FAISystem::InvalidLocation if given sense has never registered related Target actor or if last stimuli has expired.
+	*/
 	FORCEINLINE FVector GetStimulusLocation(FAISenseID Sense) const
 	{
-		return LastSensedStimuli.IsValidIndex(Sense) && LastSensedStimuli[Sense].GetAge() < FAIStimulus::NeverHappenedAge ? LastSensedStimuli[Sense].StimulusLocation : FAISystem::InvalidLocation;
+		return LastSensedStimuli.IsValidIndex(Sense) && (LastSensedStimuli[Sense].IsValid() && (LastSensedStimuli[Sense].IsExpired() == false)) ? LastSensedStimuli[Sense].StimulusLocation : FAISystem::InvalidLocation;
 	}
 
+	/** Retrieves receiver location of the last sense stimuli for a given sense
+	* @param Sense	The AISenseID of the sense
+	*
+	* @return Location of the receiver for the last sensed stimuli or FAISystem::InvalidLocation if given sense has never registered related Target actor or last stimuli has expired.
+	*/
 	FORCEINLINE FVector GetReceiverLocation(FAISenseID Sense) const
 	{
-		return LastSensedStimuli.IsValidIndex(Sense) && LastSensedStimuli[Sense].GetAge() < FAIStimulus::NeverHappenedAge ? LastSensedStimuli[Sense].ReceiverLocation : FAISystem::InvalidLocation;
+		return LastSensedStimuli.IsValidIndex(Sense) && (LastSensedStimuli[Sense].IsValid() && (LastSensedStimuli[Sense].IsExpired() == false)) ? LastSensedStimuli[Sense].ReceiverLocation : FAISystem::InvalidLocation;
 	}
 
+	UE_DEPRECATED(4.23, "This method is identical to IsSenseActive and will be removed in future versions. Please use IsSenseActive to check for a currently active stimuli or HasKnownStimulusOfSense for an active or remembered stimuli.")
 	FORCEINLINE bool IsSenseRegistered(FAISenseID Sense) const
 	{
-		return LastSensedStimuli.IsValidIndex(Sense) && LastSensedStimuli[Sense].WasSuccessfullySensed() && (LastSensedStimuli[Sense].GetAge() < FAIStimulus::NeverHappenedAge);
+		return LastSensedStimuli.IsValidIndex(Sense) && LastSensedStimuli[Sense].IsActive();
 	}
 
+	/** Indicates a currently active or "remembered" stimuli for a given sense
+	* @param Sense	The AISenseID of the sense
+	*
+	* @return True if a target has been registered (even if not currently sensed) for the given sense and the stimuli is not expired.
+	*/
 	FORCEINLINE bool HasKnownStimulusOfSense(FAISenseID Sense) const
 	{
-		return LastSensedStimuli.IsValidIndex(Sense) && (LastSensedStimuli[Sense].GetAge() < FAIStimulus::NeverHappenedAge);
+		return LastSensedStimuli.IsValidIndex(Sense) && (LastSensedStimuli[Sense].IsValid() && (LastSensedStimuli[Sense].IsExpired() == false));
 	}
 
+	/** Indicates a currently active stimuli for a given sense
+	* @param Sense	The AISenseID of the sense
+	*
+	* @return True if a target is still sensed for the given sense and the stimuli is not expired.
+	*/
 	FORCEINLINE bool IsSenseActive(FAISenseID Sense) const
 	{
 		return LastSensedStimuli.IsValidIndex(Sense) && LastSensedStimuli[Sense].IsActive();
diff --git a/Engine/Source/Runtime/AIModule/Classes/Perception/AIPerceptionTypes.h b/Engine/Source/Runtime/AIModule/Classes/Perception/AIPerceptionTypes.h
--- a/Engine/Source/Runtime/AIModule/Classes/Perception/AIPerceptionTypes.h
+++ b/Engine/Source/Runtime/AIModule/Classes/Perception/AIPerceptionTypes.h
@@ -192,9 +192,9 @@ public:
 	FORCEINLINE bool IsExpired() const { return bExpired; }
 	FORCEINLINE void MarkNoLongerSensed() { bSuccessfullySensed = false; }
 	FORCEINLINE void MarkExpired() { bExpired = true; MarkNoLongerSensed(); }
-	FORCEINLINE bool IsActive() const { return WasSuccessfullySensed() == true && GetAge() < NeverHappenedAge; }
+	FORCEINLINE bool IsActive() const { return WasSuccessfullySensed() == true && IsValid(); }
 	FORCEINLINE bool WantsToNotifyOnlyOnPerceptionChange() const { return bWantsToNotifyOnlyOnValueChange; }
-	FORCEINLINE bool IsValid() const { return Type != FAISenseID::InvalidID(); }
+	FORCEINLINE bool IsValid() const { return Type != FAISenseID::InvalidID() && GetAge() < NeverHappenedAge; }
 
 #if !(UE_BUILD_SHIPPING || UE_BUILD_TEST)
 	FString GetDebugDescription() const;
