Replaced uses of deprecated method IsSenseRegistered (file was missing from 6652669)

diff --git a/Engine/Source/Runtime/AIModule/Private/Perception/AIPerceptionComponent.cpp b/Engine/Source/Runtime/AIModule/Private/Perception/AIPerceptionComponent.cpp
--- a/Engine/Source/Runtime/AIModule/Private/Perception/AIPerceptionComponent.cpp
+++ b/Engine/Source/Runtime/AIModule/Private/Perception/AIPerceptionComponent.cpp
@@ -646,7 +646,7 @@ void UAIPerceptionComponent::GetCurrentlyPerceivedActors(TSubclassOf<UAISense> S
 	OutActors.Reserve(PerceptualData.Num());
 	for (FActorPerceptionContainer::TConstIterator DataIt = GetPerceptualDataConstIterator(); DataIt; ++DataIt)
 	{
-		const bool bCurrentlyPerceived = (SenseToUse == nullptr) ? DataIt->Value.HasAnyCurrentStimulus() : DataIt->Value.IsSenseRegistered(SenseID);
+		const bool bCurrentlyPerceived = (SenseToUse == nullptr) ? DataIt->Value.HasAnyCurrentStimulus() : DataIt->Value.IsSenseActive(SenseID);
 		if (bCurrentlyPerceived)
 		{
 			if (DataIt->Value.Target.IsValid())
