PR #5082: Added EnvQueryTest_Overlap test to be able to skip querier (Contributed by [REDACTED_PERSON])

diff --git a/Engine/Source/Runtime/AIModule/Classes/EnvironmentQuery/EnvQueryTypes.h b/Engine/Source/Runtime/AIModule/Classes/EnvironmentQuery/EnvQueryTypes.h
--- a/Engine/Source/Runtime/AIModule/Classes/EnvironmentQuery/EnvQueryTypes.h
+++ b/Engine/Source/Runtime/AIModule/Classes/EnvironmentQuery/EnvQueryTypes.h
@@ -426,7 +426,8 @@ struct AIMODULE_API FEnvOverlapData
 		OverlapChannel(ECC_WorldStatic),
 		OverlapShape(EEnvOverlapShape::Box),
 		bOnlyBlockingHits(true),
-		bOverlapComplex(false)
+		bOverlapComplex(false),
+		bSkipOverlapQuerier(false)
 	{
 	}
 
@@ -461,6 +462,10 @@ struct AIMODULE_API FEnvOverlapData
 	/** if set, overlap will run on complex collisions */
 	UPROPERTY(EditDefaultsOnly, Category = Overlap, AdvancedDisplay)
 	uint32 bOverlapComplex : 1;
+
+	/** if set, overlap will skip querier context hits */
+	UPROPERTY(EditDefaultsOnly, Category = Overlap, AdvancedDisplay)
+	uint32 bSkipOverlapQuerier : 1;
 };
 
 //////////////////////////////////////////////////////////////////////////
diff --git a/Engine/Source/Runtime/AIModule/Classes/EnvironmentQuery/Tests/EnvQueryTest_Overlap.h b/Engine/Source/Runtime/AIModule/Classes/EnvironmentQuery/Tests/EnvQueryTest_Overlap.h
--- a/Engine/Source/Runtime/AIModule/Classes/EnvironmentQuery/Tests/EnvQueryTest_Overlap.h
+++ b/Engine/Source/Runtime/AIModule/Classes/EnvironmentQuery/Tests/EnvQueryTest_Overlap.h
@@ -28,8 +28,12 @@ class UEnvQueryTest_Overlap : public UEnvQueryTest
 
 protected:
 
-	DECLARE_DELEGATE_RetVal_SixParams(bool, FRunOverlapSignature, const FVector&, const FCollisionShape&, AActor*, UWorld*, enum ECollisionChannel, const FCollisionQueryParams&);
+	bool RunOverlap(const FVector& ItemPos, const FCollisionShape& CollisionShape, const TArray<AActor*>& IgnoredActors, UWorld* World, enum ECollisionChannel Channel, const FCollisionQueryParams& Params) const;
+	bool RunOverlapBlocking(const FVector& ItemPos, const FCollisionShape& CollisionShape, const TArray<AActor*>& IgnoredActors, UWorld* World, enum ECollisionChannel Channel, const FCollisionQueryParams& Params) const;
 
+    UE_DEPRECATED(4.23, "This method is no longer called by RunTest and has been replaced by an overload using a list of actors to ignore. It now calls that overload but you should use the new overload instead.")
 	bool RunOverlap(const FVector& ItemPos, const FCollisionShape& CollisionShape, AActor* ItemActor, UWorld* World, enum ECollisionChannel Channel, const FCollisionQueryParams& Params);
+
+	UE_DEPRECATED(4.23, "This method is no longer called by RunTest and has been replaced by an overload using a list of actors to ignore. It now calls that overload but you should use the new overload instead.")
 	bool RunOverlapBlocking(const FVector& ItemPos, const FCollisionShape& CollisionShape, AActor* ItemActor, UWorld* World, enum ECollisionChannel Channel, const FCollisionQueryParams& Params);
 };
diff --git a/Engine/Source/Runtime/AIModule/Private/EnvironmentQuery/Tests/EnvQueryTest_Overlap.cpp b/Engine/Source/Runtime/AIModule/Private/EnvironmentQuery/Tests/EnvQueryTest_Overlap.cpp
--- a/Engine/Source/Runtime/AIModule/Private/EnvironmentQuery/Tests/EnvQueryTest_Overlap.cpp
+++ b/Engine/Source/Runtime/AIModule/Private/EnvironmentQuery/Tests/EnvQueryTest_Overlap.cpp
@@ -6,6 +6,7 @@
 #include "WorldCollision.h"
 #include "Engine/World.h"
 #include "EnvironmentQuery/Items/EnvQueryItemType_VectorBase.h"
+#include "EnvironmentQuery/Contexts/EnvQueryContext_Querier.h"
 
 #define LOCTEXT_NAMESPACE "EnvQueryGenerator"
 
@@ -47,23 +48,25 @@ void UEnvQueryTest_Overlap::RunTest(FEnvQueryInstance& QueryInstance) const
 			return;
 	}
 
-	FRunOverlapSignature OverlapFunc;
-	if (OverlapData.bOnlyBlockingHits)
-	{
-		OverlapFunc.BindUObject(this, &UEnvQueryTest_Overlap::RunOverlapBlocking);
-	}
-	else
+	typedef bool(UEnvQueryTest_Overlap::*FRunOverlapSignature)(const FVector&, const FCollisionShape&, const TArray<AActor*>&, UWorld*, enum ECollisionChannel, const FCollisionQueryParams&) const;
+	FRunOverlapSignature OverlapFunc = (OverlapData.bOnlyBlockingHits ? (FRunOverlapSignature)&UEnvQueryTest_Overlap::RunOverlapBlocking : (FRunOverlapSignature)&UEnvQueryTest_Overlap::RunOverlap);
+
+	TArray<AActor*> IgnoredActors;
+	if (OverlapData.bSkipOverlapQuerier)
 	{
-		OverlapFunc.BindUObject(this, &UEnvQueryTest_Overlap::RunOverlap);
+		QueryInstance.PrepareContext(UEnvQueryContext_Querier::StaticClass(), IgnoredActors);
 	}
 
 	for (FEnvQueryInstance::ItemIterator It(this, QueryInstance); It; ++It)
 	{
 		const FVector ItemLocation = GetItemLocation(QueryInstance, It.GetIndex());
 		AActor* ItemActor = GetItemActor(QueryInstance, It.GetIndex());
-		
-		const bool bHit = OverlapFunc.Execute(ItemLocation + OverlapData.ShapeOffset, CollisionShape, ItemActor, QueryInstance.World, OverlapCollisionChannel, OverlapParams);
+		IgnoredActors.Push(ItemActor);
+
+		const bool bHit = (this->*OverlapFunc)(ItemLocation + OverlapData.ShapeOffset, CollisionShape, IgnoredActors, QueryInstance.World, OverlapCollisionChannel, OverlapParams);
 		It.SetScore(TestPurpose, FilterType, bHit, bWantsHit);
+		
+		IgnoredActors.Pop(/*bAllowShrinking=*/false);
 	}
 }
 
@@ -109,18 +112,32 @@ FText UEnvQueryTest_Overlap::GetDescriptionDetails() const
 }
 
 bool UEnvQueryTest_Overlap::RunOverlap(const FVector& ItemPos, const FCollisionShape& CollisionShape, AActor* ItemActor, UWorld* World, enum ECollisionChannel Channel, const FCollisionQueryParams& Params)
+{
+	TArray<AActor*> Actors;
+	Actors.Add(ItemActor);
+	return const_cast<UEnvQueryTest_Overlap*>(this)->RunOverlap(ItemPos, CollisionShape, Actors, World, Channel, Params);
+}
+
+bool UEnvQueryTest_Overlap::RunOverlapBlocking(const FVector& ItemPos, const FCollisionShape& CollisionShape, AActor* ItemActor, UWorld* World, enum ECollisionChannel Channel, const FCollisionQueryParams& Params)
+{
+	TArray<AActor*> Actors;
+	Actors.Add(ItemActor);
+	return const_cast<UEnvQueryTest_Overlap*>(this)->RunOverlapBlocking(ItemPos, CollisionShape, Actors, World, Channel, Params);
+}
+
+bool UEnvQueryTest_Overlap::RunOverlap(const FVector& ItemPos, const FCollisionShape& CollisionShape, const TArray<AActor*>& IgnoredActors, UWorld* World, enum ECollisionChannel Channel, const FCollisionQueryParams& Params) const
 {
 	FCollisionQueryParams OverlapParams(Params);
-	OverlapParams.AddIgnoredActor(ItemActor);
+	OverlapParams.AddIgnoredActors(IgnoredActors);
 
 	const bool bHit = World->OverlapAnyTestByChannel(ItemPos, FQuat::Identity, Channel, CollisionShape, OverlapParams);
 	return bHit;
 }
 
-bool UEnvQueryTest_Overlap::RunOverlapBlocking(const FVector& ItemPos, const FCollisionShape& CollisionShape, AActor* ItemActor, UWorld* World, enum ECollisionChannel Channel, const FCollisionQueryParams& Params)
+bool UEnvQueryTest_Overlap::RunOverlapBlocking(const FVector& ItemPos, const FCollisionShape& CollisionShape, const TArray<AActor*>& IgnoredActors, UWorld* World, enum ECollisionChannel Channel, const FCollisionQueryParams& Params) const
 {
 	FCollisionQueryParams OverlapParams(Params);
-	OverlapParams.AddIgnoredActor(ItemActor);
+	OverlapParams.AddIgnoredActors(IgnoredActors);
 
 	const bool bHit = World->OverlapBlockingTestByChannel(ItemPos, FQuat::Identity, Channel, CollisionShape, OverlapParams);
 	return bHit;
