FGenericTeamId: FAttitudeSolverFunction is now a TFunction to allow non static methods to be used.
Added ResetAttitudeSolver to reset to the default attitude solver.

diff --git a/Engine/Source/Runtime/AIModule/Classes/GenericTeamAgentInterface.h b/Engine/Source/Runtime/AIModule/Classes/GenericTeamAgentInterface.h
--- a/Engine/Source/Runtime/AIModule/Classes/GenericTeamAgentInterface.h
+++ b/Engine/Source/Runtime/AIModule/Classes/GenericTeamAgentInterface.h
@@ -49,18 +49,19 @@ public:
 	static ETeamAttitude::Type GetAttitude(const AActor* A, const AActor* B);
 	static ETeamAttitude::Type GetAttitude(FGenericTeamId TeamA, FGenericTeamId TeamB)
 	{
-		return AttitudeSolverImpl ? (*AttitudeSolverImpl)(TeamA, TeamB) : ETeamAttitude::Neutral;
+		return AttitudeSolverImpl ? (AttitudeSolverImpl)(TeamA, TeamB) : ETeamAttitude::Neutral;
 	}
 
-	typedef ETeamAttitude::Type FAttitudeSolverFunction(FGenericTeamId, FGenericTeamId);
+	typedef TFunction<ETeamAttitude::Type(FGenericTeamId, FGenericTeamId)> FAttitudeSolverFunction;
 	
-	static void SetAttitudeSolver(FAttitudeSolverFunction* Solver);
+	static void SetAttitudeSolver(const FAttitudeSolverFunction& Solver);
+	static void ResetAttitudeSolver();
 
 protected:
 	// the default implementation makes all teams hostile
 	// @note that for consistency IGenericTeamAgentInterface should be using the same function 
 	//	(by default it does)
-	static FAttitudeSolverFunction* AttitudeSolverImpl;
+	static FAttitudeSolverFunction AttitudeSolverImpl;
 
 public:
 	static const FGenericTeamId NoTeam;
diff --git a/Engine/Source/Runtime/AIModule/Private/AIInterfaces.cpp b/Engine/Source/Runtime/AIModule/Private/AIInterfaces.cpp
--- a/Engine/Source/Runtime/AIModule/Private/AIInterfaces.cpp
+++ b/Engine/Source/Runtime/AIModule/Private/AIInterfaces.cpp
@@ -33,7 +33,7 @@ namespace
 	}
 }
 
-FGenericTeamId::FAttitudeSolverFunction* FGenericTeamId::AttitudeSolverImpl = &DefaultTeamAttitudeSolver;
+FGenericTeamId::FAttitudeSolverFunction FGenericTeamId::AttitudeSolverImpl = &DefaultTeamAttitudeSolver;
 
 ETeamAttitude::Type FGenericTeamId::GetAttitude(const AActor* A, const AActor* B)
 {
@@ -42,11 +42,16 @@ ETeamAttitude::Type FGenericTeamId::GetAttitude(const AActor* A, const AActor* B
 	return TeamAgentA == NULL || B == NULL ? ETeamAttitude::Neutral : TeamAgentA->GetTeamAttitudeTowards(*B);
 }
 
-void FGenericTeamId::SetAttitudeSolver(FGenericTeamId::FAttitudeSolverFunction* Solver)
+void FGenericTeamId::SetAttitudeSolver(const FGenericTeamId::FAttitudeSolverFunction& Solver)
 {
 	AttitudeSolverImpl = Solver;
 }
 
+void FGenericTeamId::ResetAttitudeSolver()
+{
+	AttitudeSolverImpl = &DefaultTeamAttitudeSolver;
+}
+
 //----------------------------------------------------------------------//
 // UGenericTeamAgentInterface
 //----------------------------------------------------------------------//
