Fixed a bug in new EQS tests sorting for 'single item' mode failing assumptions when a query templace has only scoring tests.

diff --git a/Engine/Source/Runtime/AIModule/Private/EnvironmentQuery/EnvQueryManager.cpp b/Engine/Source/Runtime/AIModule/Private/EnvironmentQuery/EnvQueryManager.cpp
--- a/Engine/Source/Runtime/AIModule/Private/EnvironmentQuery/EnvQueryManager.cpp
+++ b/Engine/Source/Runtime/AIModule/Private/EnvironmentQuery/EnvQueryManager.cpp
@@ -733,11 +733,11 @@ TSharedPtr<FEnvQueryInstance> UEnvQueryManager::CreateQueryInstance(const UEnvQu
 					MostExpensiveFilterIndex = TestIndex;
 				}
 			}
-
-			check(bOptionSingleResultSearch == false || MostExpensiveFilterIndex != INDEX_NONE);
-			
+						
 			LocalOption->Tests.Reset(SortedTests.Num());
-			if (bOptionSingleResultSearch)
+			// MostExpensiveFilterIndex can be INDEX_NONE if there are no filtering 
+			// tests done (i.e. just scoring tests).
+			if (bOptionSingleResultSearch && MostExpensiveFilterIndex != INDEX_NONE)
 			{
 				// we're going to remove the most expensive test so that 
 				// it can be added last after everything else gets sorted
@@ -767,7 +767,7 @@ TSharedPtr<FEnvQueryInstance> UEnvQueryManager::CreateQueryInstance(const UEnvQu
 			if (LocalGenerator->bAutoSortTests)
 			{
 				SortedTests.StableSort(EnvQueryTestSort::FAllMatching());
-				if (bOptionSingleResultSearch && ensure(MostExpensiveFilter))
+				if (bOptionSingleResultSearch && MostExpensiveFilter)
 				{
 					// the only difference between running for a single result is that 
 					// we want to perform a single most expensive test as the last test
