Added an accessor if the path was partial or not

#rnx

diff --git a/Engine/Source/Runtime/AIModule/Classes/Tasks/AITask_MoveTo.h b/Engine/Source/Runtime/AIModule/Classes/Tasks/AITask_MoveTo.h
--- a/Engine/Source/Runtime/AIModule/Classes/Tasks/AITask_MoveTo.h
+++ b/Engine/Source/Runtime/AIModule/Classes/Tasks/AITask_MoveTo.h
@@ -29,6 +29,7 @@ public:
 
 	EPathFollowingResult::Type GetMoveResult() const { return MoveResult; }
 	bool WasMoveSuccessful() const { return MoveResult == EPathFollowingResult::Success; }
+	bool WasMovePartial() const { return Path.IsValid() && Path->IsPartial(); }
 
 	UFUNCTION(BlueprintCallable, Category = "AI|Tasks", meta = (AdvancedDisplay = "AcceptanceRadius,StopOnOverlap,AcceptPartialPath,bUsePathfinding,bUseContinuosGoalTracking,ProjectGoalOnNavigation", DefaultToSelf = "Controller", BlueprintInternalUseOnly = "TRUE", DisplayName = "Move To Location or Actor"))
 	static UAITask_MoveTo* AIMoveTo(AAIController* Controller, FVector GoalLocation, AActor* GoalActor = nullptr,
