PR #5955: Add Multiply Score Operator to EnvQuery (Contributed by [REDACTED_PERSON])

diff --git a/Engine/Source/Runtime/AIModule/Classes/EnvironmentQuery/EnvQueryTypes.h b/Engine/Source/Runtime/AIModule/Classes/EnvironmentQuery/EnvQueryTypes.h
--- a/Engine/Source/Runtime/AIModule/Classes/EnvironmentQuery/EnvQueryTypes.h
+++ b/Engine/Source/Runtime/AIModule/Classes/EnvironmentQuery/EnvQueryTypes.h
@@ -150,6 +150,7 @@ namespace EEnvTestScoreOperator
 		AverageScore	UMETA(Tooltip = "Use average score from all contexts"),
 		MinScore		UMETA(Tooltip = "Use minimum score from all contexts"),
 		MaxScore		UMETA(Tooltip = "Use maximum score from all contexts"),
+		Multiply		UMETA(Tooltip = "Multiply scores from all contexts"),
 	};
 }
 
@@ -1203,6 +1204,9 @@ public:
 					ItemScore = Score;
 				}
 				break;
+			case EEnvTestScoreOperator::Multiply:
+				ItemScore *= Score;
+				break;
 			}
 		}
 
