PR #6101: Describing Cooldown Tags to the Visual Logger (Contributed by [REDACTED_PERSON])

diff --git a/Engine/Source/Runtime/AIModule/Private/BehaviorTree/BehaviorTreeComponent.cpp b/Engine/Source/Runtime/AIModule/Private/BehaviorTree/BehaviorTreeComponent.cpp
--- a/Engine/Source/Runtime/AIModule/Private/BehaviorTree/BehaviorTreeComponent.cpp
+++ b/Engine/Source/Runtime/AIModule/Private/BehaviorTree/BehaviorTreeComponent.cpp
@@ -2398,6 +2398,20 @@ void UBehaviorTreeComponent::DescribeSelfToVisLog(FVisualLogEntry* Snapshot) con
 
 		Snapshot->Status.Add(StatusCategory);
 	}
+
+	if (CooldownTagsMap.Num() > 0)
+	{
+		FVisualLogStatusCategory StatusCategory;
+		StatusCategory.Category = TEXT("Cooldown Tags");
+
+		for (const auto& CooldownTagPair : CooldownTagsMap)
+		{
+			const FString TimeStr = FString::Printf(TEXT("%.2fs"), CooldownTagPair.Value);
+			StatusCategory.Add(CooldownTagPair.Key.ToString(), TimeStr);
+		}
+
+		Snapshot->Status.Add(StatusCategory);
+	}
 }
 #endif // ENABLE_VISUAL_LOG
 
