BTNode: Added helper CastInstanceNodeMemory to wrap a validation with the actual cast to make it easier to detect cases where nodes attempt to use node memory without overriding GetInstanceMemorySize.

#rnx

diff --git a/Engine/Source/Runtime/AIModule/Classes/BehaviorTree/BTNode.h b/Engine/Source/Runtime/AIModule/Classes/BehaviorTree/BTNode.h
--- a/Engine/Source/Runtime/AIModule/Classes/BehaviorTree/BTNode.h
+++ b/Engine/Source/Runtime/AIModule/Classes/BehaviorTree/BTNode.h
@@ -92,6 +92,9 @@ class AIMODULE_API UBTNode : public UObject, public IGameplayTaskOwnerInterface
 	template<typename T>
 	const T* GetNodeMemory(const FBehaviorTreeInstance& BTInstance) const;
 
+	template<typename T>
+	T* CastInstanceNodeMemory(uint8* NodeMemory) const;
+
 	/** get special memory block used for hidden shared data (e.g. node instancing) */
 	template<typename T>
 	T* GetSpecialNodeMemory(uint8* NodeMemory) const;
@@ -305,6 +308,13 @@ const T* UBTNode::GetNodeMemory(const FBehaviorTreeInstance& BTInstance) const
 	return (const T*)(BTInstance.InstanceMemory.GetData() + MemoryOffset);
 }
 
+template<typename T>
+T* UBTNode::CastInstanceNodeMemory(uint8* NodeMemory) const
+{
+	checkf(sizeof(T) == GetInstanceMemorySize(), TEXT("Requesting type of %zu bytes but GetInstanceMemorySize returns %u. Make sure GetInstanceMemorySize is implemented properly in %s class hierarchy."), sizeof(T), GetInstanceMemorySize(), *GetFName().ToString());
+	return reinterpret_cast<T*>(NodeMemory);
+}
+
 template<typename T>
 T* UBTNode::GetSpecialNodeMemory(uint8* NodeMemory) const
 {
diff --git a/Engine/Source/Runtime/AIModule/Private/BehaviorTree/Decorators/BTDecorator_ConeCheck.cpp b/Engine/Source/Runtime/AIModule/Private/BehaviorTree/Decorators/BTDecorator_ConeCheck.cpp
--- a/Engine/Source/Runtime/AIModule/Private/BehaviorTree/Decorators/BTDecorator_ConeCheck.cpp
+++ b/Engine/Source/Runtime/AIModule/Private/BehaviorTree/Decorators/BTDecorator_ConeCheck.cpp
@@ -97,7 +97,7 @@ void UBTDecorator_ConeCheck::OnBlackboardChange(const UBlackboardComponent& Blac
 
 void UBTDecorator_ConeCheck::TickNode(UBehaviorTreeComponent& OwnerComp, uint8* NodeMemory, float DeltaSeconds)
 {
-	TNodeInstanceMemory* DecoratorMemory = reinterpret_cast<TNodeInstanceMemory*>(NodeMemory);
+	TNodeInstanceMemory* DecoratorMemory = CastInstanceNodeMemory<TNodeInstanceMemory>(NodeMemory);
 	
 	const bool bResult = CalcConditionImpl(OwnerComp, NodeMemory);
 	if (bResult != DecoratorMemory->bLastRawResult)
diff --git a/Engine/Source/Runtime/AIModule/Private/BehaviorTree/Decorators/BTDecorator_Cooldown.cpp b/Engine/Source/Runtime/AIModule/Private/BehaviorTree/Decorators/BTDecorator_Cooldown.cpp
--- a/Engine/Source/Runtime/AIModule/Private/BehaviorTree/Decorators/BTDecorator_Cooldown.cpp
+++ b/Engine/Source/Runtime/AIModule/Private/BehaviorTree/Decorators/BTDecorator_Cooldown.cpp
@@ -23,14 +23,14 @@ void UBTDecorator_Cooldown::PostLoad()
 
 bool UBTDecorator_Cooldown::CalculateRawConditionValue(UBehaviorTreeComponent& OwnerComp, uint8* NodeMemory) const
 {
-	FBTCooldownDecoratorMemory* DecoratorMemory = reinterpret_cast<FBTCooldownDecoratorMemory*>(NodeMemory);
+	FBTCooldownDecoratorMemory* DecoratorMemory = CastInstanceNodeMemory<FBTCooldownDecoratorMemory>(NodeMemory);
 	const float TimePassed = (OwnerComp.GetWorld()->GetTimeSeconds() - DecoratorMemory->LastUseTimestamp);
 	return TimePassed >= CoolDownTime;
 }
 
 void UBTDecorator_Cooldown::InitializeMemory(UBehaviorTreeComponent& OwnerComp, uint8* NodeMemory, EBTMemoryInit::Type InitType) const
 {
-	FBTCooldownDecoratorMemory* DecoratorMemory = reinterpret_cast<FBTCooldownDecoratorMemory*>(NodeMemory);
+	FBTCooldownDecoratorMemory* DecoratorMemory = CastInstanceNodeMemory<FBTCooldownDecoratorMemory>(NodeMemory);
 	if (InitType == EBTMemoryInit::Initialize)
 	{
 		DecoratorMemory->LastUseTimestamp = -FLT_MAX;
@@ -48,7 +48,7 @@ void UBTDecorator_Cooldown::OnNodeDeactivation(FBehaviorTreeSearchData& SearchDa
 
 void UBTDecorator_Cooldown::TickNode(UBehaviorTreeComponent& OwnerComp, uint8* NodeMemory, float DeltaSeconds)
 {
-	FBTCooldownDecoratorMemory* DecoratorMemory = reinterpret_cast<FBTCooldownDecoratorMemory*>(NodeMemory);
+	FBTCooldownDecoratorMemory* DecoratorMemory = CastInstanceNodeMemory<FBTCooldownDecoratorMemory>(NodeMemory);
 	if (!DecoratorMemory->bRequestedRestart)
 	{
 		const float TimePassed = (OwnerComp.GetWorld()->GetTimeSeconds() - DecoratorMemory->LastUseTimestamp);
@@ -71,7 +71,7 @@ void UBTDecorator_Cooldown::DescribeRuntimeValues(const UBehaviorTreeComponent&
 {
 	Super::DescribeRuntimeValues(OwnerComp, NodeMemory, Verbosity, Values);
 
-	FBTCooldownDecoratorMemory* DecoratorMemory = reinterpret_cast<FBTCooldownDecoratorMemory*>(NodeMemory);
+	FBTCooldownDecoratorMemory* DecoratorMemory = CastInstanceNodeMemory<FBTCooldownDecoratorMemory>(NodeMemory);
 	const float TimePassed = OwnerComp.GetWorld()->GetTimeSeconds() - DecoratorMemory->LastUseTimestamp;
 	
 	if (TimePassed < CoolDownTime)
diff --git a/Engine/Source/Runtime/AIModule/Private/BehaviorTree/Decorators/BTDecorator_KeepInCone.cpp b/Engine/Source/Runtime/AIModule/Private/BehaviorTree/Decorators/BTDecorator_KeepInCone.cpp
--- a/Engine/Source/Runtime/AIModule/Private/BehaviorTree/Decorators/BTDecorator_KeepInCone.cpp
+++ b/Engine/Source/Runtime/AIModule/Private/BehaviorTree/Decorators/BTDecorator_KeepInCone.cpp
@@ -76,7 +76,7 @@ bool UBTDecorator_KeepInCone::CalculateCurrentDirection(const UBehaviorTreeCompo
 
 void UBTDecorator_KeepInCone::OnBecomeRelevant(UBehaviorTreeComponent& OwnerComp, uint8* NodeMemory)
 {
-	TNodeInstanceMemory* DecoratorMemory = reinterpret_cast<TNodeInstanceMemory*>(NodeMemory);
+	TNodeInstanceMemory* DecoratorMemory = CastInstanceNodeMemory<TNodeInstanceMemory>(NodeMemory);
 	FVector InitialDir(1.0f, 0, 0);
 
 	CalculateCurrentDirection(OwnerComp, InitialDir);
@@ -85,7 +85,7 @@ void UBTDecorator_KeepInCone::OnBecomeRelevant(UBehaviorTreeComponent& OwnerComp
 
 void UBTDecorator_KeepInCone::TickNode(UBehaviorTreeComponent& OwnerComp, uint8* NodeMemory, float DeltaSeconds)
 {
-	TNodeInstanceMemory* DecoratorMemory = reinterpret_cast<TNodeInstanceMemory*>(NodeMemory);
+	TNodeInstanceMemory* DecoratorMemory = CastInstanceNodeMemory<TNodeInstanceMemory>(NodeMemory);
 	FVector CurrentDir(1.0f, 0, 0);
 	
 	if (CalculateCurrentDirection(OwnerComp, CurrentDir))
@@ -110,7 +110,7 @@ FString UBTDecorator_KeepInCone::GetStaticDescription() const
 
 void UBTDecorator_KeepInCone::DescribeRuntimeValues(const UBehaviorTreeComponent& OwnerComp, uint8* NodeMemory, EBTDescriptionVerbosity::Type Verbosity, TArray<FString>& Values) const
 {
-	TNodeInstanceMemory* DecoratorMemory = reinterpret_cast<TNodeInstanceMemory*>(NodeMemory);
+	TNodeInstanceMemory* DecoratorMemory = CastInstanceNodeMemory<TNodeInstanceMemory>(NodeMemory);
 	FVector CurrentDir(1.0f, 0, 0);
 	
 	if (CalculateCurrentDirection(OwnerComp, CurrentDir))
diff --git a/Engine/Source/Runtime/AIModule/Private/BehaviorTree/Services/BTService_RunEQS.cpp b/Engine/Source/Runtime/AIModule/Private/BehaviorTree/Services/BTService_RunEQS.cpp
--- a/Engine/Source/Runtime/AIModule/Private/BehaviorTree/Services/BTService_RunEQS.cpp
+++ b/Engine/Source/Runtime/AIModule/Private/BehaviorTree/Services/BTService_RunEQS.cpp
@@ -43,7 +43,7 @@ void UBTService_RunEQS::TickNode(UBehaviorTreeComponent& OwnerComp, uint8* NodeM
 	if (QueryOwner && EQSRequest.IsValid())
 	{
 		const UBlackboardComponent* BlackboardComponent = OwnerComp.GetBlackboardComponent();
-		FBTEQSServiceMemory* MyMemory = reinterpret_cast<FBTEQSServiceMemory*>(NodeMemory);
+		FBTEQSServiceMemory* MyMemory = CastInstanceNodeMemory<FBTEQSServiceMemory>(NodeMemory);
 
 		// Trigger new query only if the previous one has already finished
 		if (MyMemory->RequestID == INDEX_NONE)
@@ -81,7 +81,7 @@ void UBTService_RunEQS::OnQueryFinished(TSharedPtr<FEnvQueryResult> Result)
 		return;
 	}
 
-	FBTEQSServiceMemory* MyMemory = reinterpret_cast<FBTEQSServiceMemory*>(BTComp->GetNodeMemory(this, BTComp->FindInstanceContainingNode(this)));
+	FBTEQSServiceMemory* MyMemory = CastInstanceNodeMemory<FBTEQSServiceMemory>(BTComp->GetNodeMemory(this, BTComp->FindInstanceContainingNode(this)));
 	check(MyMemory);
 	ensure(MyMemory->RequestID != INDEX_NONE);
 
@@ -105,7 +105,7 @@ void UBTService_RunEQS::OnQueryFinished(TSharedPtr<FEnvQueryResult> Result)
 
 void UBTService_RunEQS::OnCeaseRelevant(UBehaviorTreeComponent& OwnerComp, uint8* NodeMemory)
 {
-	FBTEQSServiceMemory* MyMemory = reinterpret_cast<FBTEQSServiceMemory*>(NodeMemory);
+	FBTEQSServiceMemory* MyMemory = CastInstanceNodeMemory<FBTEQSServiceMemory>(NodeMemory);
 	check(MyMemory);
 
 	if (MyMemory->RequestID != INDEX_NONE)
@@ -127,7 +127,7 @@ void UBTService_RunEQS::OnCeaseRelevant(UBehaviorTreeComponent& OwnerComp, uint8
 
 void UBTService_RunEQS::InitializeMemory(UBehaviorTreeComponent& OwnerComp, uint8* NodeMemory, EBTMemoryInit::Type InitType) const
 {
-	FBTEQSServiceMemory* MyMemory = reinterpret_cast<FBTEQSServiceMemory*>(NodeMemory);
+	FBTEQSServiceMemory* MyMemory = CastInstanceNodeMemory<FBTEQSServiceMemory>(NodeMemory);
 	check(MyMemory);
 	MyMemory->RequestID = INDEX_NONE;
 }
@@ -136,7 +136,7 @@ void UBTService_RunEQS::InitializeMemory(UBehaviorTreeComponent& OwnerComp, uint
 
 void UBTService_RunEQS::CleanupMemory(UBehaviorTreeComponent& OwnerComp, uint8* NodeMemory, EBTMemoryClear::Type CleanupType) const
 {
-	const FBTEQSServiceMemory* MyMemory = reinterpret_cast<FBTEQSServiceMemory*>(NodeMemory);
+	const FBTEQSServiceMemory* MyMemory = CastInstanceNodeMemory<FBTEQSServiceMemory>(NodeMemory);
 	check(MyMemory);
 	ensure(MyMemory->RequestID == INDEX_NONE);
 }
diff --git a/Engine/Source/Runtime/AIModule/Private/BehaviorTree/Tasks/BTTask_GameplayTaskBase.cpp b/Engine/Source/Runtime/AIModule/Private/BehaviorTree/Tasks/BTTask_GameplayTaskBase.cpp
--- a/Engine/Source/Runtime/AIModule/Private/BehaviorTree/Tasks/BTTask_GameplayTaskBase.cpp
+++ b/Engine/Source/Runtime/AIModule/Private/BehaviorTree/Tasks/BTTask_GameplayTaskBase.cpp
@@ -13,7 +13,7 @@ UBTTask_GameplayTaskBase::UBTTask_GameplayTaskBase(const FObjectInitializer& Obj
 
 EBTNodeResult::Type UBTTask_GameplayTaskBase::StartGameplayTask(UBehaviorTreeComponent& OwnerComp, uint8* NodeMemory, UAITask& Task)
 {
-	FBTGameplayTaskMemory* MyMemory = reinterpret_cast<FBTGameplayTaskMemory*>(NodeMemory);
+	FBTGameplayTaskMemory* MyMemory = CastInstanceNodeMemory<FBTGameplayTaskMemory>(NodeMemory);
 	MyMemory->bObserverCanFinishTask = false;
 	MyMemory->Task = &Task;
 
@@ -35,7 +35,7 @@ EBTNodeResult::Type UBTTask_GameplayTaskBase::AbortTask(UBehaviorTreeComponent&
 {
 	if (bWaitForGameplayTask)
 	{
-		FBTGameplayTaskMemory* MyMemory = reinterpret_cast<FBTGameplayTaskMemory*>(NodeMemory);
+		FBTGameplayTaskMemory* MyMemory = CastInstanceNodeMemory<FBTGameplayTaskMemory>(NodeMemory);
 		MyMemory->bObserverCanFinishTask = false;
 
 		UAITask* TaskOb = MyMemory->Task.Get();
@@ -50,7 +50,7 @@ EBTNodeResult::Type UBTTask_GameplayTaskBase::AbortTask(UBehaviorTreeComponent&
 
 void UBTTask_GameplayTaskBase::OnTaskFinished(UBehaviorTreeComponent& OwnerComp, uint8* NodeMemory, EBTNodeResult::Type TaskResult)
 {
-	FBTGameplayTaskMemory* MyMemory = reinterpret_cast<FBTGameplayTaskMemory*>(NodeMemory);
+	FBTGameplayTaskMemory* MyMemory = CastInstanceNodeMemory<FBTGameplayTaskMemory>(NodeMemory);
 	MyMemory->Task.Reset();
 
 	Super::OnTaskFinished(OwnerComp, NodeMemory, TaskResult);
@@ -70,7 +70,7 @@ void UBTTask_GameplayTaskBase::OnGameplayTaskDeactivated(UGameplayTask& Task)
 		if (BehaviorComp)
 		{
 			uint8* RawMemory = BehaviorComp->GetNodeMemory(this, BehaviorComp->FindInstanceContainingNode(this));
-			FBTGameplayTaskMemory* MyMemory = reinterpret_cast<FBTGameplayTaskMemory*>(RawMemory);
+			FBTGameplayTaskMemory* MyMemory = CastInstanceNodeMemory<FBTGameplayTaskMemory>(RawMemory);
 
 			if (MyMemory->bObserverCanFinishTask && (AITask == MyMemory->Task))
 			{
diff --git a/Engine/Source/Runtime/AIModule/Private/BehaviorTree/Tasks/BTTask_MoveTo.cpp b/Engine/Source/Runtime/AIModule/Private/BehaviorTree/Tasks/BTTask_MoveTo.cpp
--- a/Engine/Source/Runtime/AIModule/Private/BehaviorTree/Tasks/BTTask_MoveTo.cpp
+++ b/Engine/Source/Runtime/AIModule/Private/BehaviorTree/Tasks/BTTask_MoveTo.cpp
@@ -39,7 +39,7 @@ EBTNodeResult::Type UBTTask_MoveTo::ExecuteTask(UBehaviorTreeComponent& OwnerCom
 {
 	EBTNodeResult::Type NodeResult = EBTNodeResult::InProgress;
 
-	FBTMoveToTaskMemory* MyMemory = reinterpret_cast<FBTMoveToTaskMemory*>(NodeMemory);
+	FBTMoveToTaskMemory* MyMemory = CastInstanceNodeMemory<FBTMoveToTaskMemory>(NodeMemory);
 	MyMemory->PreviousGoalLocation = FAISystem::InvalidLocation;
 	MyMemory->MoveRequestID = FAIRequestID::InvalidRequest;
 
@@ -74,7 +74,7 @@ EBTNodeResult::Type UBTTask_MoveTo::ExecuteTask(UBehaviorTreeComponent& OwnerCom
 EBTNodeResult::Type UBTTask_MoveTo::PerformMoveTask(UBehaviorTreeComponent& OwnerComp, uint8* NodeMemory)
 {
 	const UBlackboardComponent* MyBlackboard = OwnerComp.GetBlackboardComponent();
-	FBTMoveToTaskMemory* MyMemory = reinterpret_cast<FBTMoveToTaskMemory*>(NodeMemory);
+	FBTMoveToTaskMemory* MyMemory = CastInstanceNodeMemory<FBTMoveToTaskMemory>(NodeMemory);
 	AAIController* MyController = OwnerComp.GetAIOwner();
 
 	EBTNodeResult::Type NodeResult = EBTNodeResult::Failed;
@@ -197,7 +197,7 @@ EBlackboardNotificationResult UBTTask_MoveTo::OnBlackboardValueChange(const UBla
 	}
 
 	uint8* RawMemory = BehaviorComp->GetNodeMemory(this, BehaviorComp->FindInstanceContainingNode(this));
-	FBTMoveToTaskMemory* MyMemory = reinterpret_cast<FBTMoveToTaskMemory*>(RawMemory);
+	FBTMoveToTaskMemory* MyMemory = CastInstanceNodeMemory<FBTMoveToTaskMemory>(RawMemory);
 
 	const EBTTaskStatus::Type TaskStatus = BehaviorComp->GetTaskStatus(this);
 	if (TaskStatus != EBTTaskStatus::Active)
@@ -257,7 +257,7 @@ EBlackboardNotificationResult UBTTask_MoveTo::OnBlackboardValueChange(const UBla
 
 EBTNodeResult::Type UBTTask_MoveTo::AbortTask(UBehaviorTreeComponent& OwnerComp, uint8* NodeMemory)
 {
-	FBTMoveToTaskMemory* MyMemory = reinterpret_cast<FBTMoveToTaskMemory*>(NodeMemory);
+	FBTMoveToTaskMemory* MyMemory = CastInstanceNodeMemory<FBTMoveToTaskMemory>(NodeMemory);
 	if (!MyMemory->bWaitingForPath)
 	{
 		if (MyMemory->MoveRequestID.IsValid())
@@ -288,7 +288,7 @@ EBTNodeResult::Type UBTTask_MoveTo::AbortTask(UBehaviorTreeComponent& OwnerComp,
 
 void UBTTask_MoveTo::OnTaskFinished(UBehaviorTreeComponent& OwnerComp, uint8* NodeMemory, EBTNodeResult::Type TaskResult)
 {
-	FBTMoveToTaskMemory* MyMemory = reinterpret_cast<FBTMoveToTaskMemory*>(NodeMemory);
+	FBTMoveToTaskMemory* MyMemory = CastInstanceNodeMemory<FBTMoveToTaskMemory>(NodeMemory);
 	MyMemory->Task.Reset();
 
 	if (bObserveBlackboardValue)
@@ -342,7 +342,7 @@ void UBTTask_MoveTo::OnGameplayTaskDeactivated(UGameplayTask& Task)
 		if (BehaviorComp)
 		{
 			uint8* RawMemory = BehaviorComp->GetNodeMemory(this, BehaviorComp->FindInstanceContainingNode(this));
-			FBTMoveToTaskMemory* MyMemory = reinterpret_cast<FBTMoveToTaskMemory*>(RawMemory);
+			FBTMoveToTaskMemory* MyMemory = CastInstanceNodeMemory<FBTMoveToTaskMemory>(RawMemory);
 
 			if (MyMemory->bObserverCanFinishTask && (MoveTask == MyMemory->Task))
 			{
diff --git a/Engine/Source/Runtime/AIModule/Private/BehaviorTree/Tasks/BTTask_RunEQSQuery.cpp b/Engine/Source/Runtime/AIModule/Private/BehaviorTree/Tasks/BTTask_RunEQSQuery.cpp
--- a/Engine/Source/Runtime/AIModule/Private/BehaviorTree/Tasks/BTTask_RunEQSQuery.cpp
+++ b/Engine/Source/Runtime/AIModule/Private/BehaviorTree/Tasks/BTTask_RunEQSQuery.cpp
@@ -47,7 +47,7 @@ EBTNodeResult::Type UBTTask_RunEQSQuery::ExecuteTask(UBehaviorTreeComponent& Own
 	if (QueryOwner && EQSRequest.IsValid())
 	{
 		const UBlackboardComponent* BlackboardComponent = OwnerComp.GetBlackboardComponent();
-		FBTEnvQueryTaskMemory* MyMemory = reinterpret_cast<FBTEnvQueryTaskMemory*>(NodeMemory);
+		FBTEnvQueryTaskMemory* MyMemory = CastInstanceNodeMemory<FBTEnvQueryTaskMemory>(NodeMemory);
 
 		MyMemory->RequestID = EQSRequest.Execute(*QueryOwner, BlackboardComponent, QueryFinishedDelegate);
 		
