FEnvQueryInstance: Fixed Typo for EQS debug score message using boolean test.

#rnx

diff --git a/Engine/Source/Runtime/AIModule/Classes/EnvironmentQuery/EnvQueryTypes.h b/Engine/Source/Runtime/AIModule/Classes/EnvironmentQuery/EnvQueryTypes.h
--- a/Engine/Source/Runtime/AIModule/Classes/EnvironmentQuery/EnvQueryTypes.h
+++ b/Engine/Source/Runtime/AIModule/Classes/EnvironmentQuery/EnvQueryTypes.h
@@ -1091,7 +1091,7 @@ public:
 			{
 				case EEnvTestFilterType::Match:
 					bPassedTest = (bScore == bExpected);
-					UE_EQS_DBGMSG(!bPassedTest, TEXT("Boolean score don't mach (expected %s and got %s)"), bExpected ? TEXT("TRUE") : TEXT("FALSE"), bScore ? TEXT("TRUE") : TEXT("FALSE"));
+					UE_EQS_DBGMSG(!bPassedTest, TEXT("Boolean score doesn't match (expected %s and got %s)"), bExpected ? TEXT("TRUE") : TEXT("FALSE"), bScore ? TEXT("TRUE") : TEXT("FALSE"));
 					break;
 
 				case EEnvTestFilterType::Maximum:
