GameplayDebugger_NavMesh new features
- can switch between debug actor and player actor as position reference for navmesh display [Shift + Substract]
- possible to cycle through all available navigation data [Shift + Add]
- no longer fetch the entire navigation data when no tiles are found around reference position

diff --git a/Engine/Source/Developer/GameplayDebugger/Public/GameplayDebuggerCategory.h b/Engine/Source/Developer/GameplayDebugger/Public/GameplayDebuggerCategory.h
--- a/Engine/Source/Developer/GameplayDebugger/Public/GameplayDebuggerCategory.h
+++ b/Engine/Source/Developer/GameplayDebugger/Public/GameplayDebuggerCategory.h
@@ -130,6 +130,9 @@ protected:
 		return NewDataPack.PackId;
 	}
 
+	/** [AUTH] force data collection on next update */
+	void ForceImmediateCollect() { LastCollectDataTime = -FLT_MAX; }
+
 	/** update interval, 0 = each tick */
 	float CollectDataInterval;
 
diff --git a/Engine/Source/Runtime/AIModule/Private/GameplayDebugger/GameplayDebuggerCategory_Navmesh.cpp b/Engine/Source/Runtime/AIModule/Private/GameplayDebugger/GameplayDebuggerCategory_Navmesh.cpp
--- a/Engine/Source/Runtime/AIModule/Private/GameplayDebugger/GameplayDebuggerCategory_Navmesh.cpp
+++ b/Engine/Source/Runtime/AIModule/Private/GameplayDebugger/GameplayDebuggerCategory_Navmesh.cpp
@@ -22,6 +22,38 @@ FGameplayDebuggerCategory_Navmesh::FGameplayDebuggerCategory_Navmesh()
 	CollectDataInterval = 5.0f;
 	SetDataPackReplication<FNavMeshSceneProxyData>(&NavmeshRenderData);
 	SetDataPackReplication<FRepData>(&DataPack);
+
+	const FGameplayDebuggerInputHandlerConfig CycleActorReference(TEXT("Cycle Actor Reference"), TEXT("Subtract"), FGameplayDebuggerInputModifier::Shift);
+	const FGameplayDebuggerInputHandlerConfig CycleNavigationData(TEXT("Cycle NavData"), TEXT("Add"), FGameplayDebuggerInputModifier::Shift);
+	
+	BindKeyPress(CycleActorReference, this, &FGameplayDebuggerCategory_Navmesh::CycleActorReference, EGameplayDebuggerInputMode::Replicated);
+	BindKeyPress(CycleNavigationData, this, &FGameplayDebuggerCategory_Navmesh::CycleNavData, EGameplayDebuggerInputMode::Replicated);
+}
+
+void FGameplayDebuggerCategory_Navmesh::CycleNavData()
+{
+	bSwitchToNextNavigationData = true;
+	ForceImmediateCollect();
+}
+
+void FGameplayDebuggerCategory_Navmesh::CycleActorReference()
+{
+	switch (ActorReferenceMode)
+	{
+	case EActorReferenceMode::PlayerActorOnly:
+		// Nothing to do since we don't have a debug actor
+		break;
+	
+	case EActorReferenceMode::DebugActor:
+		ActorReferenceMode = EActorReferenceMode::PlayerActor;
+		ForceImmediateCollect();
+		break;
+	
+	case EActorReferenceMode::PlayerActor:
+		ActorReferenceMode = EActorReferenceMode::DebugActor;
+		ForceImmediateCollect();
+		break;
+	}
 }
 
 TSharedRef<FGameplayDebuggerCategory> FGameplayDebuggerCategory_Navmesh::MakeInstance()
@@ -32,35 +64,99 @@ TSharedRef<FGameplayDebuggerCategory> FGameplayDebuggerCategory_Navmesh::MakeIns
 void FGameplayDebuggerCategory_Navmesh::FRepData::Serialize(FArchive& Ar)
 {
 	Ar << NumDirtyAreas;
+	Ar << NavDataName;
+
+	uint8 Flags =
+		((bCanChangeReference			? 1 : 0) << 0) |
+		((bIsUsingPlayerActor			? 1 : 0) << 1) |
+		((bReferenceTooFarFromNavData	? 1 : 0) << 2);
+	
+
+	Ar << Flags;
+
+	bCanChangeReference			= (Flags & (1 << 0)) != 0;
+	bIsUsingPlayerActor			= (Flags & (1 << 1)) != 0;
+	bReferenceTooFarFromNavData = (Flags & (1 << 2)) != 0;
 }
 
 void FGameplayDebuggerCategory_Navmesh::CollectData(APlayerController* OwnerPC, AActor* DebugActor)
 {
 #if WITH_RECAST
-	const ARecastNavMesh* NavData = nullptr;
+	ANavigationData* NavData = nullptr;
+	const APawn* RefPawn = nullptr;
+	int32 NumNavData = 0;
 
-	APawn* PlayerPawn = OwnerPC ? OwnerPC->GetPawnOrSpectator() : nullptr;
-	APawn* DebugActorAsPawn = Cast<APawn>(DebugActor);
-	APawn* DestPawn = DebugActorAsPawn ? DebugActorAsPawn : PlayerPawn;
 	if (OwnerPC != nullptr)
 	{
 		UNavigationSystemV1* NavSys = FNavigationSystem::GetCurrent<UNavigationSystemV1>(OwnerPC->GetWorld());
 		if (NavSys) 
 		{
 			DataPack.NumDirtyAreas = NavSys->GetNumDirtyAreas();
+			NumNavData = NavSys->NavDataSet.Num();
+			
+			APawn* DebugActorAsPawn = Cast<APawn>(DebugActor);
+			
+			// Manage actor reference mode:
+			// - As soon as we get a new valid debug actor: use it as reference to preserve legacy behavior 
+			// - Debug actor is no longer valid: use player actor
+			if (ActorReferenceMode == EActorReferenceMode::PlayerActorOnly && DebugActorAsPawn != nullptr)
+			{
+				ActorReferenceMode = EActorReferenceMode::DebugActor;
+			}
+			else if (ActorReferenceMode != EActorReferenceMode::PlayerActorOnly && DebugActorAsPawn == nullptr)
+			{
+				ActorReferenceMode = EActorReferenceMode::PlayerActorOnly;
+			}
 
-			if (DestPawn != nullptr)
+			if (bSwitchToNextNavigationData || NavDataIndexToDisplay == INDEX_NONE)
 			{
-				const FNavAgentProperties& NavAgentProperties = DestPawn->GetNavAgentPropertiesRef();
-				NavData = Cast<const ARecastNavMesh>(NavSys->GetNavDataForProps(NavAgentProperties, DestPawn->GetNavAgentLocation()));
+				NavDataIndexToDisplay = FMath::Max(0, ++NavDataIndexToDisplay % NumNavData);
+				bSwitchToNextNavigationData = false;
+			}
+
+			if (NavSys->NavDataSet.IsValidIndex(NavDataIndexToDisplay))
+			{
+				NavData = NavSys->NavDataSet[NavDataIndexToDisplay];
+			}
+			
+			if (ActorReferenceMode == EActorReferenceMode::DebugActor)
+			{
+				RefPawn = DebugActorAsPawn;
+
+				// Switch to new debug actor NavigationData
+				if (PrevDebugActorReference != RefPawn)
+				{
+					const FNavAgentProperties& NavAgentProperties = RefPawn->GetNavAgentPropertiesRef();
+					NavData = NavSys->GetNavDataForProps(NavAgentProperties, RefPawn->GetNavAgentLocation());
+					NavDataIndexToDisplay = NavSys->NavDataSet.Find(NavData);
+
+					PrevDebugActorReference = RefPawn;
+				}
+			}
+			else
+			{
+				RefPawn = OwnerPC ? OwnerPC->GetPawnOrSpectator() : nullptr;
 			}
 		}
 	}
 
-	if (NavData)
+	const ARecastNavMesh* RecastNavMesh = Cast<const ARecastNavMesh>(NavData);
+	if (RecastNavMesh && RefPawn)
 	{
+		DataPack.bIsUsingPlayerActor = (ActorReferenceMode != EActorReferenceMode::DebugActor);
+		DataPack.bCanChangeReference = (ActorReferenceMode != EActorReferenceMode::PlayerActorOnly);
+
+		if (NumNavData > 1)
+		{
+			DataPack.NavDataName = FString::Printf(TEXT("[%d/%d] %s"), NavDataIndexToDisplay + 1, NumNavData, *NavData->GetFName().ToString());
+		}
+		else
+		{
+			DataPack.NavDataName = NavData->GetFName().ToString();
+		}
+
 		// add 3x3 neighborhood of target
-		const FVector TargetLocation = DestPawn->GetActorLocation();
+		const FVector TargetLocation = RefPawn->GetActorLocation();
 
 		TArray<int32> TileSet;
 		int32 TileX = 0;
@@ -70,12 +166,12 @@ void FGameplayDebuggerCategory_Navmesh::CollectData(APlayerController* OwnerPC,
 
 		int32 TargetTileX = 0;
 		int32 TargetTileY = 0;
-		NavData->GetNavMeshTileXY(TargetLocation, TargetTileX, TargetTileY);
+		RecastNavMesh->GetNavMeshTileXY(TargetLocation, TargetTileX, TargetTileY);
 		for (int32 Idx = 0; Idx < ARRAY_COUNT(DeltaX); Idx++)
 		{
 			const int32 NeiX = TargetTileX + DeltaX[Idx];
 			const int32 NeiY = TargetTileY + DeltaY[Idx];
-			NavData->GetNavMeshTilesAt(NeiX, NeiY, TileSet);
+			RecastNavMesh->GetNavMeshTilesAt(NeiX, NeiY, TileSet);
 		}
 
 		const int32 DetailFlags =
@@ -84,7 +180,16 @@ void FGameplayDebuggerCategory_Navmesh::CollectData(APlayerController* OwnerPC,
 			(1 << static_cast<int32>(ENavMeshDetailFlags::NavLinks)) |
 			(bDrawExcludedFlags ? (1 << static_cast<int32>(ENavMeshDetailFlags::MarkForbiddenPolys)) : 0);
 
-		NavmeshRenderData.GatherData(NavData, DetailFlags, TileSet);
+		// Do not attempt to gather render data when TileSet is empty otherwise the whole nav mesh will be displayed
+		DataPack.bReferenceTooFarFromNavData = (TileSet.Num() == 0);
+		if (DataPack.bReferenceTooFarFromNavData)
+		{
+			NavmeshRenderData.Reset();
+		}
+		else
+		{
+			NavmeshRenderData.GatherData(RecastNavMesh, DetailFlags, TileSet);
+		}
 	}
 #endif // WITH_RECAST
 }
@@ -92,6 +197,17 @@ void FGameplayDebuggerCategory_Navmesh::CollectData(APlayerController* OwnerPC,
 void FGameplayDebuggerCategory_Navmesh::DrawData(APlayerController* OwnerPC, FGameplayDebuggerCanvasContext& CanvasContext)
 {
 	CanvasContext.Printf(TEXT("Num dirty areas: {%s}%d"), DataPack.NumDirtyAreas > 0 ? TEXT("red") : TEXT("green"), DataPack.NumDirtyAreas);
+
+	if (!DataPack.NavDataName.IsEmpty())
+	{
+		CanvasContext.Printf(TEXT("Navigation Data: {silver}%s%s"), *DataPack.NavDataName, DataPack.bReferenceTooFarFromNavData ? TEXT(" (too far from navmesh)") : TEXT(""));
+		CanvasContext.Printf(TEXT("[{yellow}%s{white}]: Cycle NavData"), *GetInputHandlerDescription(1));
+	}
+
+	if (DataPack.bCanChangeReference)
+	{
+		CanvasContext.Printf(TEXT("[{yellow}%s{white}]: Display around %s actor"), *GetInputHandlerDescription(0), DataPack.bIsUsingPlayerActor ? TEXT("Debug") : TEXT("Player"));
+	}
 }
 
 void FGameplayDebuggerCategory_Navmesh::OnDataPackReplicated(int32 DataPackId)
diff --git a/Engine/Source/Runtime/AIModule/Private/GameplayDebugger/GameplayDebuggerCategory_Navmesh.h b/Engine/Source/Runtime/AIModule/Private/GameplayDebugger/GameplayDebuggerCategory_Navmesh.h
--- a/Engine/Source/Runtime/AIModule/Private/GameplayDebugger/GameplayDebuggerCategory_Navmesh.h
+++ b/Engine/Source/Runtime/AIModule/Private/GameplayDebugger/GameplayDebuggerCategory_Navmesh.h
@@ -24,14 +24,34 @@ public:
 
 protected:
 
+	void CycleNavData();
+	void CycleActorReference();
+
 	struct FRepData
 	{
-		int32 NumDirtyAreas = 0;
 		void Serialize(FArchive& Ar);
+
+		FString NavDataName;
+		int32 NumDirtyAreas = 0;
+		bool bCanChangeReference = false;
+		bool bIsUsingPlayerActor = false;
+		bool bReferenceTooFarFromNavData = false;
 	};
 
 	FNavMeshSceneProxyData NavmeshRenderData;
 	FRepData DataPack;
+	
+	enum class EActorReferenceMode : uint8
+	{
+		PlayerActorOnly,
+		PlayerActor,
+		DebugActor
+	};
+	EActorReferenceMode ActorReferenceMode = EActorReferenceMode::DebugActor;
+
+	int32 NavDataIndexToDisplay = INDEX_NONE;
+	bool bSwitchToNextNavigationData = false;
+	TWeakObjectPtr<const APawn> PrevDebugActorReference;
 };
 
 #endif // WITH_GAMEPLAY_DEBUGGER
