Static Analysis fixes in FGameplayDebuggerCategory_Navmesh

diff --git a/Engine/Source/Runtime/AIModule/Private/GameplayDebugger/GameplayDebuggerCategory_Navmesh.cpp b/Engine/Source/Runtime/AIModule/Private/GameplayDebugger/GameplayDebuggerCategory_Navmesh.cpp
--- a/Engine/Source/Runtime/AIModule/Private/GameplayDebugger/GameplayDebuggerCategory_Navmesh.cpp
+++ b/Engine/Source/Runtime/AIModule/Private/GameplayDebugger/GameplayDebuggerCategory_Navmesh.cpp
@@ -148,11 +148,11 @@ void FGameplayDebuggerCategory_Navmesh::CollectData(APlayerController* OwnerPC,
 
 		if (NumNavData > 1)
 		{
-			DataPack.NavDataName = FString::Printf(TEXT("[%d/%d] %s"), NavDataIndexToDisplay + 1, NumNavData, *NavData->GetFName().ToString());
+			DataPack.NavDataName = FString::Printf(TEXT("[%d/%d] %s"), NavDataIndexToDisplay + 1, NumNavData, *RecastNavMesh->GetFName().ToString());
 		}
 		else
 		{
-			DataPack.NavDataName = NavData->GetFName().ToString();
+			DataPack.NavDataName = RecastNavMesh->GetFName().ToString();
 		}
 
 		// add 3x3 neighborhood of target
