Added a method to copy key content between two keys of same key type within the same blackboard

#rnx

diff --git a/Engine/Source/Runtime/AIModule/Classes/BehaviorTree/BlackboardComponent.h b/Engine/Source/Runtime/AIModule/Classes/BehaviorTree/BlackboardComponent.h
--- a/Engine/Source/Runtime/AIModule/Classes/BehaviorTree/BlackboardComponent.h
+++ b/Engine/Source/Runtime/AIModule/Classes/BehaviorTree/BlackboardComponent.h
@@ -186,6 +186,9 @@ public:
 	void ClearValue(const FName& KeyName);
 	void ClearValue(FBlackboard::FKey KeyID);
 
+	/** Copy content from SourceKeyID to DestinationID and return true if it worked */
+	bool CopyKeyValue(FBlackboard::FKey SourceKeyID, FBlackboard::FKey DestinationID);
+
 	template<class TDataClass>
 	bool IsKeyOfType(FBlackboard::FKey KeyID) const;
 
diff --git a/Engine/Source/Runtime/AIModule/Private/BehaviorTree/BlackboardComponent.cpp b/Engine/Source/Runtime/AIModule/Private/BehaviorTree/BlackboardComponent.cpp
--- a/Engine/Source/Runtime/AIModule/Private/BehaviorTree/BlackboardComponent.cpp
+++ b/Engine/Source/Runtime/AIModule/Private/BehaviorTree/BlackboardComponent.cpp
@@ -736,6 +736,43 @@ void UBlackboardComponent::ClearValue(FBlackboard::FKey KeyID)
 	}
 }
 
+bool UBlackboardComponent::CopyKeyValue(FBlackboard::FKey SourceKeyID, FBlackboard::FKey DestinationID)
+{
+	UBlackboardData* BBAsset = GetBlackboardAsset();
+	if (BBAsset == nullptr)
+	{
+		return false;
+	}
+
+	// copy only when values are initialized
+	if (ValueMemory.Num() == 0)
+	{
+		return false;
+	}
+
+	const FBlackboardEntry* SourceValueEntryInfo = BBAsset->GetKey(SourceKeyID);
+	const FBlackboardEntry* DestinationValueEntryInfo = BBAsset->GetKey(DestinationID);
+
+	if (SourceValueEntryInfo == nullptr || DestinationValueEntryInfo == nullptr || SourceValueEntryInfo->KeyType == nullptr || DestinationValueEntryInfo->KeyType == nullptr)
+	{
+		return false;
+	}
+
+	if (SourceValueEntryInfo->KeyType->GetClass() != DestinationValueEntryInfo->KeyType->GetClass())
+	{
+		return false;
+	}
+
+	const uint16 MemDataOffset = SourceValueEntryInfo->KeyType->IsInstanced() ? sizeof(FBlackboardInstancedKeyMemory) : 0;
+	
+	uint8* SourceValueMem = GetKeyRawData(SourceKeyID) + MemDataOffset;
+	uint8* DestinationValueMem = GetKeyRawData(DestinationID) + MemDataOffset;
+
+	SourceValueEntryInfo->KeyType->CopyValues(*this, DestinationValueMem, SourceValueEntryInfo->KeyType, SourceValueMem);
+
+	return true;
+}
+
 bool UBlackboardComponent::GetLocationFromEntry(const FName& KeyName, FVector& ResultLocation) const
 {
 	const FBlackboard::FKey KeyID = GetKeyID(KeyName);
