Added a streamlined display of gameplay tags in behavior tree nodes ("TagValue" directly, instead of "(TagName=TagValue)")
#rnx

diff --git a/Engine/Source/Runtime/AIModule/Private/BehaviorTree/BlueprintNodeHelpers.cpp b/Engine/Source/Runtime/AIModule/Private/BehaviorTree/BlueprintNodeHelpers.cpp
--- a/Engine/Source/Runtime/AIModule/Private/BehaviorTree/BlueprintNodeHelpers.cpp
+++ b/Engine/Source/Runtime/AIModule/Private/BehaviorTree/BlueprintNodeHelpers.cpp
@@ -10,6 +10,7 @@
 #include "GameFramework/Controller.h"
 #include "BehaviorTree/BehaviorTreeTypes.h"
 #include "BehaviorTree/BTNode.h"
+#include "GameplayTagContainer.h"
 
 namespace BlueprintNodeHelpers
 {
@@ -96,6 +97,10 @@ namespace BlueprintNodeHelpers
 			const FBlackboardKeySelector* PropertyValue = (const FBlackboardKeySelector*)PropertyAddr;
 			ExportedStringValue = PropertyValue->SelectedKeyName.ToString();
 		}
+		else if (StructProp && StructProp->Struct && StructProp->Struct->IsChildOf(TBaseStructure<FGameplayTag>::Get()))
+		{
+			ExportedStringValue = ((const FGameplayTag*)PropertyAddr)->ToString();
+		}
 		else if (FloatProp)
 		{
 			// special case for floats to remove unnecessary zeros
