GameplayDebuggerCategory_Navmesh: Hide "Cycle Data" action when there is only one navmesh

#rnx

diff --git a/Engine/Source/Runtime/AIModule/Private/GameplayDebugger/GameplayDebuggerCategory_Navmesh.cpp b/Engine/Source/Runtime/AIModule/Private/GameplayDebugger/GameplayDebuggerCategory_Navmesh.cpp
--- a/Engine/Source/Runtime/AIModule/Private/GameplayDebugger/GameplayDebuggerCategory_Navmesh.cpp
+++ b/Engine/Source/Runtime/AIModule/Private/GameplayDebugger/GameplayDebuggerCategory_Navmesh.cpp
@@ -68,15 +68,16 @@ void FGameplayDebuggerCategory_Navmesh::FRepData::Serialize(FArchive& Ar)
 
 	uint8 Flags =
 		((bCanChangeReference			? 1 : 0) << 0) |
-		((bIsUsingPlayerActor			? 1 : 0) << 1) |
-		((bReferenceTooFarFromNavData	? 1 : 0) << 2);
-	
+		((bCanCycleNavigationData		? 1 : 0) << 1) |
+		((bIsUsingPlayerActor			? 1 : 0) << 2) |
+		((bReferenceTooFarFromNavData	? 1 : 0) << 3);
 
 	Ar << Flags;
 
 	bCanChangeReference			= (Flags & (1 << 0)) != 0;
-	bIsUsingPlayerActor			= (Flags & (1 << 1)) != 0;
-	bReferenceTooFarFromNavData = (Flags & (1 << 2)) != 0;
+	bCanCycleNavigationData		= (Flags & (1 << 1)) != 0;
+	bIsUsingPlayerActor			= (Flags & (1 << 2)) != 0;
+	bReferenceTooFarFromNavData = (Flags & (1 << 3)) != 0;
 }
 
 void FGameplayDebuggerCategory_Navmesh::CollectData(APlayerController* OwnerPC, AActor* DebugActor)
@@ -145,6 +146,7 @@ void FGameplayDebuggerCategory_Navmesh::CollectData(APlayerController* OwnerPC,
 	{
 		DataPack.bIsUsingPlayerActor = (ActorReferenceMode != EActorReferenceMode::DebugActor);
 		DataPack.bCanChangeReference = (ActorReferenceMode != EActorReferenceMode::PlayerActorOnly);
+		DataPack.bCanCycleNavigationData = (NumNavData > 1);
 
 		if (NumNavData > 1)
 		{
@@ -201,6 +203,10 @@ void FGameplayDebuggerCategory_Navmesh::DrawData(APlayerController* OwnerPC, FGa
 	if (!DataPack.NavDataName.IsEmpty())
 	{
 		CanvasContext.Printf(TEXT("Navigation Data: {silver}%s%s"), *DataPack.NavDataName, DataPack.bReferenceTooFarFromNavData ? TEXT(" (too far from navmesh)") : TEXT(""));
+	}
+
+	if (DataPack.bCanCycleNavigationData)
+	{
 		CanvasContext.Printf(TEXT("[{yellow}%s{white}]: Cycle NavData"), *GetInputHandlerDescription(1));
 	}
 
diff --git a/Engine/Source/Runtime/AIModule/Private/GameplayDebugger/GameplayDebuggerCategory_Navmesh.h b/Engine/Source/Runtime/AIModule/Private/GameplayDebugger/GameplayDebuggerCategory_Navmesh.h
--- a/Engine/Source/Runtime/AIModule/Private/GameplayDebugger/GameplayDebuggerCategory_Navmesh.h
+++ b/Engine/Source/Runtime/AIModule/Private/GameplayDebugger/GameplayDebuggerCategory_Navmesh.h
@@ -34,6 +34,7 @@ protected:
 		FString NavDataName;
 		int32 NumDirtyAreas = 0;
 		bool bCanChangeReference = false;
+		bool bCanCycleNavigationData = false;
 		bool bIsUsingPlayerActor = false;
 		bool bReferenceTooFarFromNavData = false;
 	};
