BlackboardComponent: Fixed KeyInstances used as synchronized keys and fixed method CopyKeyValues to support instances.

There was some wrong usages of IsInstanced instead of HasInstance resulting in the wrong UBlackboardKeyType used in some code paths (keytype objects from the asset instead of the instanced one in the component)
Also set flag bIsInstanced properly for UBlackboardKeyType objects created as instances when initializing keys.

diff --git a/Engine/Source/Runtime/AIModule/Private/BehaviorTree/Blackboard/BlackboardKeyType.cpp b/Engine/Source/Runtime/AIModule/Private/BehaviorTree/Blackboard/BlackboardKeyType.cpp
--- a/Engine/Source/Runtime/AIModule/Private/BehaviorTree/Blackboard/BlackboardKeyType.cpp
+++ b/Engine/Source/Runtime/AIModule/Private/BehaviorTree/Blackboard/BlackboardKeyType.cpp
@@ -24,6 +24,7 @@ void UBlackboardKeyType::InitializeKey(UBlackboardComponent& OwnerComp, FBlackbo
 	{
 		FBlackboardInstancedKeyMemory* MyMemory = (FBlackboardInstancedKeyMemory*)RawData;
 		UBlackboardKeyType* KeyInstance = NewObject<UBlackboardKeyType>(&OwnerComp, GetClass());
+		KeyInstance->bIsInstanced = true;
 		MyMemory->KeyIdx = KeyID;
 		OwnerComp.KeyInstances[KeyID] = KeyInstance;
 
diff --git a/Engine/Source/Runtime/AIModule/Private/BehaviorTree/BlackboardComponent.cpp b/Engine/Source/Runtime/AIModule/Private/BehaviorTree/BlackboardComponent.cpp
--- a/Engine/Source/Runtime/AIModule/Private/BehaviorTree/BlackboardComponent.cpp
+++ b/Engine/Source/Runtime/AIModule/Private/BehaviorTree/BlackboardComponent.cpp
@@ -232,13 +232,15 @@ void UBlackboardComponent::PopulateSynchronizedKeys()
 						check(Key.EntryName == OtherKey->EntryName);
 						check(Key.KeyType == OtherKey->KeyType);
 
-						const uint16 DataOffset = Key.KeyType->IsInstanced() ? sizeof(FBlackboardInstancedKeyMemory) : 0;
+						const bool bKeyHasInstance = Key.KeyType->HasInstance();
+						const uint16 DataOffset = bKeyHasInstance ? sizeof(FBlackboardInstancedKeyMemory) : 0;
+
 						const int32 KeyID = BlackboardAsset->GetKeyID(Key.EntryName);
 						uint8* RawData = GetKeyRawData(KeyID) + DataOffset;
 						uint8* RawSource = OtherBlackboard->GetKeyRawData(OtherKeyID) + DataOffset;
 
-						UBlackboardKeyType* KeyOb = Key.KeyType->IsInstanced() ? KeyInstances[KeyID] : Key.KeyType;
-						const UBlackboardKeyType* SourceKeyOb = Key.KeyType->IsInstanced() ? OtherBlackboard->KeyInstances[OtherKeyID] : Key.KeyType;
+						UBlackboardKeyType* KeyOb = bKeyHasInstance ? KeyInstances[KeyID] : Key.KeyType;
+						const UBlackboardKeyType* SourceKeyOb = bKeyHasInstance ? OtherBlackboard->KeyInstances[OtherKeyID] : Key.KeyType;
 
 						KeyOb->CopyValues(*this, RawData, SourceKeyOb, RawSource);
 					}
@@ -706,8 +708,10 @@ void UBlackboardComponent::ClearValue(FBlackboard::FKey KeyID)
 
 			if (BlackboardAsset->HasSynchronizedKeys() && IsKeyInstanceSynced(KeyID))
 			{
-				UBlackboardKeyType* KeyOb = EntryInfo->KeyType->IsInstanced() ? KeyInstances[KeyID] : EntryInfo->KeyType;
-				const uint16 DataOffset = EntryInfo->KeyType->IsInstanced() ? sizeof(FBlackboardInstancedKeyMemory) : 0;
+				const bool bKeyHasInstance = EntryInfo->KeyType->HasInstance();
+
+				UBlackboardKeyType* KeyOb = bKeyHasInstance ? KeyInstances[KeyID] : EntryInfo->KeyType;
+				const uint16 DataOffset = bKeyHasInstance ? sizeof(FBlackboardInstancedKeyMemory) : 0;
 				uint8* InstancedRawData = RawData + DataOffset;
 
 				// grab the value set and apply the same to synchronized keys
@@ -723,7 +727,7 @@ void UBlackboardComponent::ClearValue(FBlackboard::FKey KeyID)
 						if (OtherKeyID != FBlackboard::InvalidKey)
 						{
 							const FBlackboardEntry* OtherEntryInfo = OtherBlackboard->BlackboardAsset->GetKey(OtherKeyID);
-							UBlackboardKeyType* OtherKeyOb = EntryInfo->KeyType->IsInstanced() ? OtherBlackboard->KeyInstances[OtherKeyID] : EntryInfo->KeyType;
+							UBlackboardKeyType* OtherKeyOb = bKeyHasInstance ? OtherBlackboard->KeyInstances[OtherKeyID] : EntryInfo->KeyType;
 							uint8* OtherRawData = OtherBlackboard->GetKeyRawData(OtherKeyID) + DataOffset;
 
 							OtherKeyOb->CopyValues(*OtherBlackboard, OtherRawData, KeyOb, InstancedRawData);
@@ -736,7 +740,7 @@ void UBlackboardComponent::ClearValue(FBlackboard::FKey KeyID)
 	}
 }
 
-bool UBlackboardComponent::CopyKeyValue(FBlackboard::FKey SourceKeyID, FBlackboard::FKey DestinationID)
+bool UBlackboardComponent::CopyKeyValue(FBlackboard::FKey SourceKeyID, FBlackboard::FKey DestinationKeyID)
 {
 	UBlackboardData* BBAsset = GetBlackboardAsset();
 	if (BBAsset == nullptr)
@@ -751,7 +755,7 @@ bool UBlackboardComponent::CopyKeyValue(FBlackboard::FKey SourceKeyID, FBlackboa
 	}
 
 	const FBlackboardEntry* SourceValueEntryInfo = BBAsset->GetKey(SourceKeyID);
-	const FBlackboardEntry* DestinationValueEntryInfo = BBAsset->GetKey(DestinationID);
+	const FBlackboardEntry* DestinationValueEntryInfo = BBAsset->GetKey(DestinationKeyID);
 
 	if (SourceValueEntryInfo == nullptr || DestinationValueEntryInfo == nullptr || SourceValueEntryInfo->KeyType == nullptr || DestinationValueEntryInfo->KeyType == nullptr)
 	{
@@ -763,12 +767,17 @@ bool UBlackboardComponent::CopyKeyValue(FBlackboard::FKey SourceKeyID, FBlackboa
 		return false;
 	}
 
-	const uint16 MemDataOffset = SourceValueEntryInfo->KeyType->IsInstanced() ? sizeof(FBlackboardInstancedKeyMemory) : 0;
+	const bool bKeyHasInstance = SourceValueEntryInfo->KeyType->HasInstance();
+
+	const uint16 MemDataOffset = bKeyHasInstance ? sizeof(FBlackboardInstancedKeyMemory) : 0;
 	
-	uint8* SourceValueMem = GetKeyRawData(SourceKeyID) + MemDataOffset;
-	uint8* DestinationValueMem = GetKeyRawData(DestinationID) + MemDataOffset;
+	const uint8* SourceValueMem = GetKeyRawData(SourceKeyID) + MemDataOffset;
+	uint8* DestinationValueMem = GetKeyRawData(DestinationKeyID) + MemDataOffset;
+
+	const UBlackboardKeyType* SourceKeyOb = bKeyHasInstance ? KeyInstances[SourceKeyID] : SourceValueEntryInfo->KeyType;
+	UBlackboardKeyType*	DestKeyOb = bKeyHasInstance ? KeyInstances[DestinationKeyID] : DestinationValueEntryInfo->KeyType;
 
-	SourceValueEntryInfo->KeyType->CopyValues(*this, DestinationValueMem, SourceValueEntryInfo->KeyType, SourceValueMem);
+	DestKeyOb->CopyValues(*this, DestinationValueMem, SourceKeyOb, SourceValueMem);
 
 	return true;
 }
