Extended GameplayDebuggerCategory_Navmesh with information on number of running and remaining tile building tasks

diff --git a/Engine/Source/Runtime/AIModule/Private/GameplayDebugger/GameplayDebuggerCategory_Navmesh.cpp b/Engine/Source/Runtime/AIModule/Private/GameplayDebugger/GameplayDebuggerCategory_Navmesh.cpp
--- a/Engine/Source/Runtime/AIModule/Private/GameplayDebugger/GameplayDebuggerCategory_Navmesh.cpp
+++ b/Engine/Source/Runtime/AIModule/Private/GameplayDebugger/GameplayDebuggerCategory_Navmesh.cpp
@@ -64,6 +64,8 @@ TSharedRef<FGameplayDebuggerCategory> FGameplayDebuggerCategory_Navmesh::MakeIns
 void FGameplayDebuggerCategory_Navmesh::FRepData::Serialize(FArchive& Ar)
 {
 	Ar << NumDirtyAreas;
+	Ar << NumRunningTasks;
+	Ar << NumRemainingTasks;
 	Ar << NavDataName;
 
 	uint8 Flags =
@@ -93,6 +95,9 @@ void FGameplayDebuggerCategory_Navmesh::CollectData(APlayerController* OwnerPC,
 		if (NavSys) 
 		{
 			DataPack.NumDirtyAreas = NavSys->GetNumDirtyAreas();
+			DataPack.NumRunningTasks = NavSys->GetNumRunningBuildTasks();
+			DataPack.NumRemainingTasks = NavSys->GetNumRemainingBuildTasks();
+
 			NumNavData = NavSys->NavDataSet.Num();
 			
 			APawn* DebugActorAsPawn = Cast<APawn>(DebugActor);
@@ -199,6 +204,7 @@ void FGameplayDebuggerCategory_Navmesh::CollectData(APlayerController* OwnerPC,
 void FGameplayDebuggerCategory_Navmesh::DrawData(APlayerController* OwnerPC, FGameplayDebuggerCanvasContext& CanvasContext)
 {
 	CanvasContext.Printf(TEXT("Num dirty areas: {%s}%d"), DataPack.NumDirtyAreas > 0 ? TEXT("red") : TEXT("green"), DataPack.NumDirtyAreas);
+	CanvasContext.Printf(TEXT("Tile jobs running/remaining: %d / %d"), DataPack.NumRunningTasks, DataPack.NumRemainingTasks);
 
 	if (!DataPack.NavDataName.IsEmpty())
 	{
diff --git a/Engine/Source/Runtime/AIModule/Private/GameplayDebugger/GameplayDebuggerCategory_Navmesh.h b/Engine/Source/Runtime/AIModule/Private/GameplayDebugger/GameplayDebuggerCategory_Navmesh.h
--- a/Engine/Source/Runtime/AIModule/Private/GameplayDebugger/GameplayDebuggerCategory_Navmesh.h
+++ b/Engine/Source/Runtime/AIModule/Private/GameplayDebugger/GameplayDebuggerCategory_Navmesh.h
@@ -33,6 +33,8 @@ protected:
 
 		FString NavDataName;
 		int32 NumDirtyAreas = 0;
+		uint16 NumRunningTasks = 0;
+		uint16 NumRemainingTasks = 0;
 		bool bCanChangeReference = false;
 		bool bCanCycleNavigationData = false;
 		bool bIsUsingPlayerActor = false;
