GameplayDebuggerCategory_Navmesh: prevent division by 0 when there is no navmesh

diff --git a/Engine/Source/Runtime/AIModule/Private/GameplayDebugger/GameplayDebuggerCategory_Navmesh.cpp b/Engine/Source/Runtime/AIModule/Private/GameplayDebugger/GameplayDebuggerCategory_Navmesh.cpp
--- a/Engine/Source/Runtime/AIModule/Private/GameplayDebugger/GameplayDebuggerCategory_Navmesh.cpp
+++ b/Engine/Source/Runtime/AIModule/Private/GameplayDebugger/GameplayDebuggerCategory_Navmesh.cpp
@@ -116,7 +116,7 @@ void FGameplayDebuggerCategory_Navmesh::CollectData(APlayerController* OwnerPC,
 
 			if (bSwitchToNextNavigationData || NavDataIndexToDisplay == INDEX_NONE)
 			{
-				NavDataIndexToDisplay = FMath::Max(0, ++NavDataIndexToDisplay % NumNavData);
+				NavDataIndexToDisplay = (NumNavData > 0) ? FMath::Max(0, ++NavDataIndexToDisplay % NumNavData) : INDEX_NONE;
 				bSwitchToNextNavigationData = false;
 			}
 
