Behavior Tree node display: Add support for strings/text types, and improved support for gameplay tags to trim off required prefixes from Categories metadata
#rnx

diff --git a/Engine/Source/Runtime/AIModule/Private/BehaviorTree/BlueprintNodeHelpers.cpp b/Engine/Source/Runtime/AIModule/Private/BehaviorTree/BlueprintNodeHelpers.cpp
--- a/Engine/Source/Runtime/AIModule/Private/BehaviorTree/BlueprintNodeHelpers.cpp
+++ b/Engine/Source/Runtime/AIModule/Private/BehaviorTree/BlueprintNodeHelpers.cpp
@@ -37,7 +37,9 @@ namespace BlueprintNodeHelpers
 	{
 		if (TestProperty->IsA(UNumericProperty::StaticClass()) ||
 			TestProperty->IsA(UBoolProperty::StaticClass()) ||
-			TestProperty->IsA(UNameProperty::StaticClass()))
+			TestProperty->IsA(UNameProperty::StaticClass()) ||
+			TestProperty->IsA(UStrProperty::StaticClass()) ||
+			TestProperty->IsA(UTextProperty::StaticClass()) )
 		{
 			return true;
 		}
@@ -99,7 +101,14 @@ namespace BlueprintNodeHelpers
 		}
 		else if (StructProp && StructProp->Struct && StructProp->Struct->IsChildOf(TBaseStructure<FGameplayTag>::Get()))
 		{
+			const FString CategoryLimit = StructProp->GetMetaData(TEXT("Categories"));
+
 			ExportedStringValue = ((const FGameplayTag*)PropertyAddr)->ToString();
+
+			if (!CategoryLimit.IsEmpty() && ExportedStringValue.StartsWith(CategoryLimit))
+			{
+				ExportedStringValue = ExportedStringValue.Mid(CategoryLimit.Len());
+			}
 		}
 		else if (FloatProp)
 		{
