Fix non-editor compile error in BlueprintNodeHelpers.cpp
#rnx

diff --git a/Engine/Source/Runtime/AIModule/Private/BehaviorTree/BlueprintNodeHelpers.cpp b/Engine/Source/Runtime/AIModule/Private/BehaviorTree/BlueprintNodeHelpers.cpp
--- a/Engine/Source/Runtime/AIModule/Private/BehaviorTree/BlueprintNodeHelpers.cpp
+++ b/Engine/Source/Runtime/AIModule/Private/BehaviorTree/BlueprintNodeHelpers.cpp
@@ -101,14 +101,15 @@ namespace BlueprintNodeHelpers
 		}
 		else if (StructProp && StructProp->Struct && StructProp->Struct->IsChildOf(TBaseStructure<FGameplayTag>::Get()))
 		{
-			const FString CategoryLimit = StructProp->GetMetaData(TEXT("Categories"));
-
 			ExportedStringValue = ((const FGameplayTag*)PropertyAddr)->ToString();
 
+#if WITH_EDITOR
+			const FString CategoryLimit = StructProp->GetMetaData(TEXT("Categories"));
 			if (!CategoryLimit.IsEmpty() && ExportedStringValue.StartsWith(CategoryLimit))
 			{
 				ExportedStringValue = ExportedStringValue.Mid(CategoryLimit.Len());
 			}
+#endif
 		}
 		else if (FloatProp)
 		{
