Adding a missing include of TextProperty.h for non-unity builds
#rnx

diff --git a/Engine/Source/Runtime/AIModule/Private/BehaviorTree/BlueprintNodeHelpers.cpp b/Engine/Source/Runtime/AIModule/Private/BehaviorTree/BlueprintNodeHelpers.cpp
--- a/Engine/Source/Runtime/AIModule/Private/BehaviorTree/BlueprintNodeHelpers.cpp
+++ b/Engine/Source/Runtime/AIModule/Private/BehaviorTree/BlueprintNodeHelpers.cpp
@@ -11,6 +11,7 @@
 #include "BehaviorTree/BehaviorTreeTypes.h"
 #include "BehaviorTree/BTNode.h"
 #include "GameplayTagContainer.h"
+#include "UObject/TextProperty.h"
 
 namespace BlueprintNodeHelpers
 {
