BT running and stoping capabilites added to gameplay behaviors

In addition, added a GameplayBehaviorManager that knows about all the running gameplay behaviors and can be used to stop an arbitrary behavior.
Added a BT node for stopping current behavior (optionally a filtering class can be used).
Asset samples following in the next CL

diff --git a/Engine/Source/Runtime/AIModule/Private/AIModule.cpp b/Engine/Source/Runtime/AIModule/Private/AIModule.cpp
--- a/Engine/Source/Runtime/AIModule/Private/AIModule.cpp
+++ b/Engine/Source/Runtime/AIModule/Private/AIModule.cpp
@@ -5,6 +5,7 @@
 #include "AISystem.h"
 #include "VisualLogger/VisualLogger.h"
 #include "GameFramework/WorldSettings.h"
+#include "DefaultManagerInstanceTracker.h"
 
 #if WITH_EDITOR
 #include "Developer/AssetTools/Public/IAssetTools.h"
@@ -28,6 +29,7 @@
 #define LOCTEXT_NAMESPACE "AIModule"
 
 DEFINE_LOG_CATEGORY_STATIC(LogAIModule, Log, All);
+DEFINE_LOG_CATEGORY(LogManagerInstanceTracker);
 
 class FAIModule : public IAIModule
 {
diff --git a/Engine/Source/Runtime/AIModule/Public/DefaultManagerInstanceTracker.h b/Engine/Source/Runtime/AIModule/Public/DefaultManagerInstanceTracker.h
new file mode 100644
--- /dev/null
+++ b/Engine/Source/Runtime/AIModule/Public/DefaultManagerInstanceTracker.h
@@ -0,0 +1,60 @@
+// Copyright 1998-2019 Epic Games, Inc. All Rights Reserved.
+
+#pragma once
+
+#include "Containers/Map.h"
+#include "Engine/World.h"
+	
+
+AIMODULE_API DECLARE_LOG_CATEGORY_EXTERN(LogManagerInstanceTracker, Warning, All);
+
+template<typename FManager>
+struct TDefaultManagerInstanceTracker
+{
+	TMap<UWorld*, FManager*> WorldToInstanceMap;
+	bool bCreateIfMissing = false;
+
+	TDefaultManagerInstanceTracker()
+	{
+		FWorldDelegates::OnWorldCleanup.AddRaw(this, &TDefaultManagerInstanceTracker::OnWorldCleanup);
+	}
+
+	FManager* GetManagerInstance(UWorld& World)
+	{
+		FManager** FoundInstance = WorldToInstanceMap.Find(&World);
+		if (FoundInstance)
+		{
+			return *FoundInstance;
+		}
+		else if (World.bIsTearingDown)
+		{
+			return nullptr;
+		}
+
+		FManager* NewInstance = nullptr;
+		if (bCreateIfMissing)
+		{
+			NewInstance = NewObject<FManager>(&World);
+			NewInstance->AddToRoot();
+			WorldToInstanceMap.Add(&World, NewInstance);
+		}
+		else
+		{
+			UE_LOG(LogManagerInstanceTracker, Warning
+				, TEXT("Calling FManager::GetCurrent while no instance for world %s has been created, and bCreateIfMissing == false")
+				, *World.GetName());
+		}
+
+		return NewInstance;
+	}
+
+	void OnWorldCleanup(UWorld* World, bool /*bSessionEnded*/, bool /*bCleanupResources*/)
+	{
+		FManager** FoundInstance = WorldToInstanceMap.Find(World);
+		if (FoundInstance)
+		{
+			(*FoundInstance)->RemoveFromRoot();
+			WorldToInstanceMap.Remove(World);
+		}
+	}
+};
