Fixed BTMoveTo task ObservedBlackboardValueTolerance not being visible in editor.

diff --git a/Engine/Source/Runtime/AIModule/Classes/BehaviorTree/Tasks/BTTask_MoveTo.h b/Engine/Source/Runtime/AIModule/Classes/BehaviorTree/Tasks/BTTask_MoveTo.h
--- a/Engine/Source/Runtime/AIModule/Classes/BehaviorTree/Tasks/BTTask_MoveTo.h
+++ b/Engine/Source/Runtime/AIModule/Classes/BehaviorTree/Tasks/BTTask_MoveTo.h
@@ -48,11 +48,11 @@ class AIMODULE_API UBTTask_MoveTo : public UBTTask_BlackboardBase
 	/** if task is expected to react to changes to location represented by BB key 
 	 *	this property can be used to tweak sensitivity of the mechanism. Value is 
 	 *	recommended to be less then AcceptableRadius */
-	UPROPERTY(Category=Blackboard, EditAnywhere, AdvancedDisplay, meta = (ClampMin = "1", UIMin = "1", EditCondition="bObserveBlackboardValue"))
+	UPROPERTY(Category=Blackboard, EditAnywhere, meta = (ClampMin = "1", UIMin = "1", EditCondition="bObserveBlackboardValue", DisplayAfter="bObserveBlackboardValue"))
 	float ObservedBlackboardValueTolerance;
 
 	/** if move goal in BB changes the move will be redirected to new location */
-	UPROPERTY()
+	UPROPERTY(Category = Blackboard, EditAnywhere)
 	uint32 bObserveBlackboardValue : 1;
 
 	UPROPERTY(Category = Node, EditAnywhere)
