Mark UWorld::GetPawnIterator() deprecated and fix up remaining instance of it in use in Engine code

diff --git a/Engine/Source/Runtime/AIModule/Private/Perception/AIPerceptionSystem.cpp b/Engine/Source/Runtime/AIModule/Private/Perception/AIPerceptionSystem.cpp
--- a/Engine/Source/Runtime/AIModule/Private/Perception/AIPerceptionSystem.cpp
+++ b/Engine/Source/Runtime/AIModule/Private/Perception/AIPerceptionSystem.cpp
@@ -2,6 +2,7 @@
 
 #include "Perception/AIPerceptionSystem.h"
 #include "EngineGlobals.h"
+#include "EngineUtils.h"
 #include "TimerManager.h"
 #include "Engine/Engine.h"
 #include "AISystem.h"
@@ -572,12 +573,9 @@ void UAIPerceptionSystem::StartPlay()
 void UAIPerceptionSystem::RegisterAllPawnsAsSourcesForSense(FAISenseID SenseID)
 {
 	UWorld* World = GetWorld();
-	for (auto PawnIt = World->GetPawnIterator(); PawnIt; ++PawnIt)
+	for (TActorIterator<APawn> PawnIt(World); PawnIt; ++PawnIt)
 	{
-		if (PawnIt->Get())
-		{
-			RegisterSource(SenseID, **PawnIt);
-		}
+		RegisterSource(SenseID, **PawnIt);
 	}
 }
 
diff --git a/Engine/Source/Runtime/Engine/Classes/Engine/World.h b/Engine/Source/Runtime/Engine/Classes/Engine/World.h
--- a/Engine/Source/Runtime/Engine/Classes/Engine/World.h
+++ b/Engine/Source/Runtime/Engine/Classes/Engine/World.h
@@ -2122,6 +2122,7 @@ public:
 	FConstControllerIterator GetControllerIterator() const;
 
 	/** @return Returns an iterator for the pawn list. */
+	UE_DEPRECATED(4.24, "The PawnIterator is an inefficient mechanism for iterating pawns. Please use TActorIterator<PawnType> instead.")
 	FConstPawnIterator GetPawnIterator() const;
 	
 	/** @return Returns the number of Pawns. */
