Added comment to method

#rnx

diff --git a/Engine/Source/Runtime/AIModule/Classes/BehaviorTree/BehaviorTreeComponent.h b/Engine/Source/Runtime/AIModule/Classes/BehaviorTree/BehaviorTreeComponent.h
--- a/Engine/Source/Runtime/AIModule/Classes/BehaviorTree/BehaviorTreeComponent.h
+++ b/Engine/Source/Runtime/AIModule/Classes/BehaviorTree/BehaviorTreeComponent.h
@@ -380,6 +380,7 @@ protected:
 	/** check if debugger is currently running and can gather data */
 	static bool IsDebuggerActive();
 
+	/** Return NodeA's relative priority in regards to NodeB */
 	EBTNodeRelativePriority CalculateRelativePriority(const UBTNode* NodeA, const UBTNode* NodeB) const;
 
 	friend UBTNode;
