PR #6124: Bugfix: Fix FBlackboardKeySelector::AddClassFilter (Contributed by [REDACTED_PERSON])

I've extended the CL to include previous AddClassFilter implementation's deprecation.

diff --git a/Engine/Source/Runtime/AIModule/Classes/BehaviorTree/BehaviorTreeTypes.h b/Engine/Source/Runtime/AIModule/Classes/BehaviorTree/BehaviorTreeTypes.h
--- a/Engine/Source/Runtime/AIModule/Classes/BehaviorTree/BehaviorTreeTypes.h
+++ b/Engine/Source/Runtime/AIModule/Classes/BehaviorTree/BehaviorTreeTypes.h
@@ -517,7 +517,7 @@ public:
 
 	/** helper functions for setting basic filters */
 	void AddObjectFilter(UObject* Owner, FName PropertyName, TSubclassOf<UObject> AllowedClass);
-	void AddClassFilter(UObject* Owner, FName PropertyName, TSubclassOf<UClass> AllowedClass);
+	void AddClassFilter(UObject* Owner, FName PropertyName, TSubclassOf<UObject> AllowedClass);
 	void AddEnumFilter(UObject* Owner, FName PropertyName, UEnum* AllowedEnum);
 	void AddNativeEnumFilter(UObject* Owner, FName PropertyName, const FString& AllowedEnumName);
 	void AddIntFilter(UObject* Owner, FName PropertyName);
@@ -534,6 +534,9 @@ public:
 	FORCEINLINE void InvalidateResolvedKey() { SelectedKeyID = FBlackboard::InvalidKey; }
 
 	friend FBlackboardDecoratorDetails;
+
+	UE_DEPRECATED(4.24, "This version of AddClassFilter is deprecated. Please provide AllowedClass as TSubclassOf<UObject>")
+	void AddClassFilter(UObject* Owner, FName PropertyName, TSubclassOf<UClass> AllowedClass);
 };
 
 UCLASS(Abstract)
diff --git a/Engine/Source/Runtime/AIModule/Private/BehaviorTree/BehaviorTreeTypes.cpp b/Engine/Source/Runtime/AIModule/Private/BehaviorTree/BehaviorTreeTypes.cpp
--- a/Engine/Source/Runtime/AIModule/Private/BehaviorTree/BehaviorTreeTypes.cpp
+++ b/Engine/Source/Runtime/AIModule/Private/BehaviorTree/BehaviorTreeTypes.cpp
@@ -353,7 +353,7 @@ void FBlackboardKeySelector::AddObjectFilter(UObject* Owner, FName PropertyName,
 	AllowedTypes.Add(FilterOb);
 }
 
-void FBlackboardKeySelector::AddClassFilter(UObject* Owner, FName PropertyName, TSubclassOf<UClass> AllowedClass)
+void FBlackboardKeySelector::AddClassFilter(UObject* Owner, FName PropertyName, TSubclassOf<UObject> AllowedClass)
 {
 	const FName FilterName = MakeUniqueObjectName(Owner, UBlackboardKeyType_Class::StaticClass(), *FString::Printf(TEXT("%s_Class"), *PropertyName.ToString()));
 	UBlackboardKeyType_Class* FilterOb = NewObject<UBlackboardKeyType_Class>(Owner, FilterName);
@@ -419,6 +419,13 @@ void FBlackboardKeySelector::AddNameFilter(UObject* Owner, FName PropertyName)
 	AllowedTypes.Add(NewObject<UBlackboardKeyType_Name>(Owner, *FilterName));
 }
 
+// deprecated
+void FBlackboardKeySelector::AddClassFilter(UObject* Owner, FName PropertyName, TSubclassOf<UClass> AllowedClass)
+{
+	TSubclassOf<UObject> AllowedObjectClass = AllowedClass;
+	AddClassFilter(Owner, PropertyName, AllowedObjectClass);
+}
+
 //----------------------------------------------------------------------//
 // UBehaviorTreeTypes
 //----------------------------------------------------------------------//
