Fix for a perception adding and removing a source in the same frame

diff --git a/Engine/Source/Runtime/AIModule/Private/Perception/AIPerceptionSystem.cpp b/Engine/Source/Runtime/AIModule/Private/Perception/AIPerceptionSystem.cpp
--- a/Engine/Source/Runtime/AIModule/Private/Perception/AIPerceptionSystem.cpp
+++ b/Engine/Source/Runtime/AIModule/Private/Perception/AIPerceptionSystem.cpp
@@ -387,6 +387,14 @@ void UAIPerceptionSystem::UnregisterSource(AActor& SourceActor, const TSubclassO
 	{
 		UE_VLOG(this, LogAIPerception, Log, TEXT("UnregisterSource called for %s but it doesn't seem to be registered as a source"), *SourceActor.GetName());
 	}
+	// Remove this from any pending adds (add/remove same frame)
+	for (int32 RemoveIndex = SourcesToRegister.Num() - 1; RemoveIndex >= 0; RemoveIndex--)
+	{
+		if (SourcesToRegister[RemoveIndex].Source == &SourceActor)
+		{
+			SourcesToRegister.RemoveAt(RemoveIndex, 1, false);
+		}
+	}
 }
 
 void UAIPerceptionSystem::OnListenerRemoved(const FPerceptionListener& NewListener)
