Fixed the sense removal code to only remove pending (same frame) updates that match sense id

diff --git a/Engine/Source/Runtime/AIModule/Private/Perception/AIPerceptionSystem.cpp b/Engine/Source/Runtime/AIModule/Private/Perception/AIPerceptionSystem.cpp
--- a/Engine/Source/Runtime/AIModule/Private/Perception/AIPerceptionSystem.cpp
+++ b/Engine/Source/Runtime/AIModule/Private/Perception/AIPerceptionSystem.cpp
@@ -391,7 +391,8 @@ void UAIPerceptionSystem::UnregisterSource(AActor& SourceActor, const TSubclassO
 	// Remove this from any pending adds (add/remove same frame)
 	for (int32 RemoveIndex = SourcesToRegister.Num() - 1; RemoveIndex >= 0; RemoveIndex--)
 	{
-		if (SourcesToRegister[RemoveIndex].Source == &SourceActor)
+		if (SourcesToRegister[RemoveIndex].Source == &SourceActor &&
+			SourcesToRegister[RemoveIndex].SenseID == UAISense::GetSenseID(Sense))
 		{
 			SourcesToRegister.RemoveAt(RemoveIndex, 1, false);
 		}
