AIPerceptionComponent: Added functionality to allow clearing of perceptual data by converting ForgetAll to a BlueprintCallable UFUNCTION.

diff --git a/Engine/Source/Runtime/AIModule/Classes/Perception/AIPerceptionComponent.h b/Engine/Source/Runtime/AIModule/Classes/Perception/AIPerceptionComponent.h
--- a/Engine/Source/Runtime/AIModule/Classes/Perception/AIPerceptionComponent.h
+++ b/Engine/Source/Runtime/AIModule/Classes/Perception/AIPerceptionComponent.h
@@ -293,6 +293,7 @@ public:
 	void ForgetActor(AActor* ActorToForget);
 
 	/** basically cleans up PerceptualData, resulting in loss of all previous perception */
+	UFUNCTION(BlueprintCallable, Category = "AI|Perception")
 	void ForgetAll();
 
 	float GetYoungestStimulusAge(const AActor& Source) const;
