Fixed a crash in AI CrowdManager resulting from the navigation data the manager is relying on going away, being streamed out or just being unregistered from the NavSystem.

Changed UNavigationSystemV1::OnBeginTearingDown to call UNavigationSystemV1::CleanUp so that CrowdManager doesn't do unnecessary work when we unregister nav data on EndPlay.

diff --git a/Engine/Source/Runtime/AIModule/Classes/Navigation/CrowdManager.h b/Engine/Source/Runtime/AIModule/Classes/Navigation/CrowdManager.h
--- a/Engine/Source/Runtime/AIModule/Classes/Navigation/CrowdManager.h
+++ b/Engine/Source/Runtime/AIModule/Classes/Navigation/CrowdManager.h
@@ -242,6 +242,20 @@ public:
 	/** notify called when detour navmesh is changed */
 	void OnNavMeshUpdate();
 
+	/** Tests if NavData is a suitable nav data type to be used by this CrowdManager
+	 *	instance. */
+	virtual bool IsSuitableNavData(const ANavigationData& NavData) const;
+
+	/** Called by the nav system when a new navigation data instance is registered. 
+	 *	If the CrowdManager instance had no nav data cached it will consider this
+	 *	NavDataInstance and update if necesary. */
+	virtual void OnNavDataRegistered(ANavigationData& NavData) override;
+
+	/** Called by the nav system when a navigation data instance is removed. The 
+	 *	crowd manager will see if it's the nav data being used by it an if so try
+	 *	to find another one. If there's none the crowd manager will stop working. */
+	virtual void OnNavDataUnregistered(ANavigationData& NavData) override;
+
 	const ANavigationData* GetNavData() const { return MyNavData; }
 
 	UWorld* GetWorld() const override;
@@ -334,6 +348,11 @@ protected:
 	/** called from tick, after move points were updated, before any steering/avoidance */
 	virtual void PostMovePointUpdate();
 
+	/** Sets NavData as MyNavData. If Null and bFindNewNavDataIfNull is true then
+	 *	the manager will search for a new NavData instance that meets the 
+	 *	IsSuitableNavData() condition. */
+	void SetNavData(ANavigationData* NavData, const bool bFindNewNavDataIfNull = true);
+
 #if WITH_RECAST
 	void AddAgent(const ICrowdAgentInterface* Agent, FCrowdAgentData& AgentData) const;
 	void RemoveAgent(const ICrowdAgentInterface* Agent, FCrowdAgentData* AgentData) const;
diff --git a/Engine/Source/Runtime/AIModule/Private/Navigation/CrowdManager.cpp b/Engine/Source/Runtime/AIModule/Private/Navigation/CrowdManager.cpp
--- a/Engine/Source/Runtime/AIModule/Private/Navigation/CrowdManager.cpp
+++ b/Engine/Source/Runtime/AIModule/Private/Navigation/CrowdManager.cpp
@@ -1258,15 +1258,11 @@ void UCrowdManager::UpdateNavData()
 		UNavigationSystemV1* NavSys = Cast<UNavigationSystemV1>(GetOuter());
 		if (NavSys)
 		{
-			for (int32 Idx = 0; Idx < NavSys->NavDataSet.Num(); Idx++)
+			for (ANavigationData* NavData : NavSys->NavDataSet)
 			{
-				ARecastNavMesh* RecastNavData = Cast<ARecastNavMesh>(NavSys->NavDataSet[Idx]);
-				if (RecastNavData && RecastNavData->IsSupportingDefaultAgent())
+				if (NavData && IsSuitableNavData(*NavData))
 				{
-					MyNavData = RecastNavData;
-					RecastNavData->OnNavMeshUpdate.AddUObject(this, &UCrowdManager::OnNavMeshUpdate);
-					OnNavMeshUpdate();
-
+					SetNavData(NavData);
 					break;
 				}
 			}
@@ -1274,6 +1270,37 @@ void UCrowdManager::UpdateNavData()
 	}
 }
 
+void UCrowdManager::SetNavData(ANavigationData* NavData, const bool bFindNewNavDataIfNull)
+{
+	if (NavData != MyNavData && MyNavData != nullptr)
+	{
+		// clean up
+		ARecastNavMesh* AsRecastNavData = Cast<ARecastNavMesh>(MyNavData);
+		if (AsRecastNavData)
+		{
+			AsRecastNavData->OnNavMeshUpdate.RemoveAll(this);
+		}
+	}
+
+	if (NavData)
+	{
+		ARecastNavMesh* AsRecastNavData = Cast<ARecastNavMesh>(NavData);
+		if (AsRecastNavData)
+		{
+			AsRecastNavData->OnNavMeshUpdate.AddUObject(this, &UCrowdManager::OnNavMeshUpdate);
+		}
+	}
+
+	MyNavData = NavData;
+	UE_VLOG(this, LogCrowdFollowing, Log, TEXT("Setting cached nav data to %s")
+		, MyNavData ? *MyNavData->GetFullName() : TEXT("NONE"));
+
+	if (NavData || bFindNewNavDataIfNull)
+	{
+		OnNavMeshUpdate();
+	}
+}
+
 void UCrowdManager::OnNavMeshUpdate()
 {
 #if WITH_RECAST
@@ -1282,6 +1309,36 @@ void UCrowdManager::OnNavMeshUpdate()
 #endif // WITH_RECAST
 }
 
+bool UCrowdManager::IsSuitableNavData(const ANavigationData& NavData) const
+{
+	return NavData.IsSupportingDefaultAgent() && Cast<ARecastNavMesh>(&NavData);
+}
+
+void UCrowdManager::OnNavDataRegistered(ANavigationData& NavData)
+{
+	if (MyNavData == nullptr && IsSuitableNavData(NavData))
+	{
+		SetNavData(&NavData);
+	}
+}
+
+void UCrowdManager::OnNavDataUnregistered(ANavigationData& NavData)
+{
+	if (MyNavData == &NavData)
+	{
+		UE_VLOG(this, LogCrowdFollowing, Warning, TEXT("%s is being unregistered. Will look for new nav data")
+			, *NavData.GetName());
+		
+		// note that this will cause the manager to look for another nav data instance
+		SetNavData(nullptr, /*bFindNewNavDataIfNull=*/true);
+
+		if (MyNavData == nullptr)
+		{
+			UE_VLOG(this, LogCrowdFollowing, Warning, TEXT("Failed to find suitable navigation data instance (no navmesh). Shutting down."));
+		}
+	}
+}
+
 void UCrowdManager::UpdateAvoidanceConfig()
 {
 #if WITH_RECAST
diff --git a/Engine/Source/Runtime/NavigationSystem/Private/NavigationSystem.cpp b/Engine/Source/Runtime/NavigationSystem/Private/NavigationSystem.cpp
--- a/Engine/Source/Runtime/NavigationSystem/Private/NavigationSystem.cpp
+++ b/Engine/Source/Runtime/NavigationSystem/Private/NavigationSystem.cpp
@@ -723,7 +723,7 @@ void UNavigationSystemV1::OnInitializeActors()
 
 void UNavigationSystemV1::OnBeginTearingDown()
 {
-	DestroyNavOctree();
+	CleanUp(FNavigationSystem::ECleanupMode::CleanupWithWorld);
 }
 
 void UNavigationSystemV1::OnWorldInitDone(FNavigationSystemRunMode Mode)
@@ -1975,6 +1975,11 @@ UNavigationSystemV1::ERegistrationResult UNavigationSystemV1::RegisterNavData(AN
 	// care needs to be taken to not make it implement navigation for agent who's real implementation has 
 	// not been loaded yet.
 
+	if (Result == RegistrationSuccessful && CrowdManager != nullptr)
+	{
+		CrowdManager->OnNavDataRegistered(*NavData);
+	}
+
 	return Result;
 }
 
@@ -1992,6 +1997,11 @@ void UNavigationSystemV1::UnregisterNavData(ANavigationData* NavData)
 	FScopeLock Lock(&NavDataRegistration);
 	NavDataRegistrationQueue.Remove(NavData);
 	NavData->OnUnregistered();
+
+	if (CrowdManager != nullptr)
+	{
+		CrowdManager->OnNavDataUnregistered(*NavData);
+	}
 }
 
 void UNavigationSystemV1::RegisterCustomLink(INavLinkCustomInterface& CustomLink)
diff --git a/Engine/Source/Runtime/NavigationSystem/Public/CrowdManagerBase.h b/Engine/Source/Runtime/NavigationSystem/Public/CrowdManagerBase.h
--- a/Engine/Source/Runtime/NavigationSystem/Public/CrowdManagerBase.h
+++ b/Engine/Source/Runtime/NavigationSystem/Public/CrowdManagerBase.h
@@ -4,6 +4,8 @@
 
 #include "CrowdManagerBase.generated.h"
 
+class ANavigationData;
+
 /** Base class for Crowd Managers. If you want to create a custom crowd manager
  *	implement a class extending this one and set UNavigationSystemV1::CrowdManagerClass
  *	to point at your class */
@@ -13,5 +15,13 @@ class NAVIGATIONSYSTEM_API UCrowdManagerBase : public UObject
 	GENERATED_BODY()
 public:
 	virtual void Tick(float DeltaTime) PURE_VIRTUAL(UCrowdManagerBase::Tick, );
+
+	/** Called by the nav system when a new navigation data instance is registered. */
+	virtual void OnNavDataRegistered(ANavigationData& NavDataInstance) PURE_VIRTUAL(UCrowdManagerBase::OnNavDataRegistered, );
+
+	/** Called by the nav system when a navigation data instance is removed. */
+	virtual void OnNavDataUnregistered(ANavigationData& NavDataInstance) PURE_VIRTUAL(UCrowdManagerBase::OnNavDataUnregistered, );
+
+	virtual void CleanUp(float DeltaTime) PURE_VIRTUAL(UCrowdManagerBase::Tick, );
 };
 
