Adding a function SetHearingRange

diff --git a/Engine/Source/Runtime/AIModule/Private/Perception/AIPerceptionComponent.cpp b/Engine/Source/Runtime/AIModule/Private/Perception/AIPerceptionComponent.cpp
--- a/Engine/Source/Runtime/AIModule/Private/Perception/AIPerceptionComponent.cpp
+++ b/Engine/Source/Runtime/AIModule/Private/Perception/AIPerceptionComponent.cpp
@@ -109,7 +109,8 @@ void UAIPerceptionComponent::ConfigureSense(UAISenseConfig& Config)
 	// first check if we're reconfiguring a sense
 	for (int32 Index = 0; Index < SensesConfig.Num(); ++Index)
 	{
-		if (SensesConfig[Index]->GetClass() == Config.GetClass())
+		if (SensesConfig[Index] != nullptr
+			&& SensesConfig[Index]->GetClass() == Config.GetClass())
 		{
 			ConfigIndex = Index;
 			SensesConfig[Index] = &Config;
