BehaviorTree: prevent tree restarts from decorator nodes requesting execution while processing higher priority execution.

diff --git a/Engine/Source/Runtime/AIModule/Classes/BehaviorTree/BehaviorTreeComponent.h b/Engine/Source/Runtime/AIModule/Classes/BehaviorTree/BehaviorTreeComponent.h
--- a/Engine/Source/Runtime/AIModule/Classes/BehaviorTree/BehaviorTreeComponent.h
+++ b/Engine/Source/Runtime/AIModule/Classes/BehaviorTree/BehaviorTreeComponent.h
@@ -53,6 +53,9 @@ struct FBTPendingExecutionInfo
 	/** next task to execute */
 	UBTTaskNode* NextTask;
 
+	/** execution request node index */
+	FBTNodeIndex ExecutionNodeIndex;
+
 	/** if set, tree ran out of nodes */
 	uint32 bOutOfNodes : 1;
 
@@ -242,6 +245,9 @@ protected:
 	/** execution request, search will be performed when current task finish execution/aborting */
 	FBTNodeExecutionInfo ExecutionRequest;
 
+	/** node index currently processed while applying the search data */
+	FBTNodeIndex InProgressExecutionNodeIndex;
+
 	/** result of ExecutionRequest, will be applied when current task finish aborting */
 	FBTPendingExecutionInfo PendingExecution;
 
diff --git a/Engine/Source/Runtime/AIModule/Private/BehaviorTree/BehaviorTreeComponent.cpp b/Engine/Source/Runtime/AIModule/Private/BehaviorTree/BehaviorTreeComponent.cpp
--- a/Engine/Source/Runtime/AIModule/Private/BehaviorTree/BehaviorTreeComponent.cpp
+++ b/Engine/Source/Runtime/AIModule/Private/BehaviorTree/BehaviorTreeComponent.cpp
@@ -857,6 +857,16 @@ void UBehaviorTreeComponent::RequestExecution(UBTCompositeNode* RequestedOn, int
 		return;
 	}
 
+	// while processing execution request, we do not request execution from nodes 
+	// with a priority lower than the processed node since it may no longer be
+	// in the active branch and will force a restart of the tree
+	if (InProgressExecutionNodeIndex.IsSet() && InProgressExecutionNodeIndex.TakesPriorityOver(ExecutionIdx) && RequestedBy->IsA(UBTDecorator::StaticClass()))
+	{
+		UE_VLOG(GetOwner(), LogBehaviorTree, Warning, TEXT("> skip: requesting execution from a lower priority node: %s < %s"), *ExecutionIdx.Describe(), *InProgressExecutionNodeIndex.Describe());
+		StoreDebuggerRestart(DebuggerNode, InstanceIdx, false);
+		return;
+	}
+
 	// when it's aborting and moving to higher priority node:
 	if (bSwitchToHigherPriority)
 	{
@@ -1461,6 +1471,8 @@ void UBehaviorTreeComponent::ProcessExecutionRequest()
 
 	if (!SearchData.bPostponeSearch)
 	{
+		PendingExecution.ExecutionNodeIndex = FBTNodeIndex(ExecutionRequest.ExecuteInstanceIdx, ExecutionRequest.ExecuteNode->GetExecutionIndex());
+
 		// clear request accumulator
 		ExecutionRequest = FBTNodeExecutionInfo();
 
@@ -1499,6 +1511,7 @@ void UBehaviorTreeComponent::ProcessPendingExecution()
 	}
 
 	FBTPendingExecutionInfo SavedInfo = PendingExecution;
+	InProgressExecutionNodeIndex = PendingExecution.ExecutionNodeIndex;
 	PendingExecution = FBTPendingExecutionInfo();
 
 	// collect all aux nodes that have lower priority than new task
@@ -1530,6 +1543,8 @@ void UBehaviorTreeComponent::ProcessPendingExecution()
 	{
 		OnTreeFinished();
 	}
+
+	InProgressExecutionNodeIndex = FBTNodeIndex();
 }
 
 void UBehaviorTreeComponent::RollbackSearchChanges()
