Fixed some LogBehaviorTree UE_VLOG by using owning actor instead of a component
Prefixed node name by the parent tree name in DescribeNodeHelper to get more detailed logs

#rnx

diff --git a/Engine/Source/Runtime/AIModule/Private/BehaviorTree/BTCompositeNode.cpp b/Engine/Source/Runtime/AIModule/Private/BehaviorTree/BTCompositeNode.cpp
--- a/Engine/Source/Runtime/AIModule/Private/BehaviorTree/BTCompositeNode.cpp
+++ b/Engine/Source/Runtime/AIModule/Private/BehaviorTree/BTCompositeNode.cpp
@@ -449,18 +449,20 @@ bool UBTCompositeNode::DoDecoratorsAllowExecution(UBehaviorTreeComponent& OwnerC
 			const bool bIsAllowed = TestDecorator ? TestDecorator->WrappedCanExecute(OwnerComp, TestDecorator->GetNodeMemory<uint8>(MyInstance)) : false;
 			OwnerComp.StoreDebuggerSearchStep(TestDecorator, InstanceIdx, bIsAllowed);
 
+			const UBTNode* ChildNode = GetChildNode(ChildIdx);
+
 			if (!bIsAllowed)
 			{
-				UE_VLOG(OwnerComp.GetOwner(), LogBehaviorTree, Verbose, TEXT("Child[%d] execution forbidden by %s"),
-					ChildIdx, *UBehaviorTreeTypes::DescribeNodeHelper(TestDecorator));
+				UE_VLOG(OwnerComp.GetOwner(), LogBehaviorTree, Verbose, TEXT("Child[%d] \"%s\" execution forbidden by %s"),
+					ChildIdx, *UBehaviorTreeTypes::DescribeNodeHelper(ChildNode), *UBehaviorTreeTypes::DescribeNodeHelper(TestDecorator));
 
 				bResult = false;
 				break;
 			}
 			else
 			{
-				UE_VLOG(OwnerComp.GetOwner(), LogBehaviorTree, Verbose, TEXT("Child[%d] execution allowed by %s"),
-					ChildIdx, *UBehaviorTreeTypes::DescribeNodeHelper(TestDecorator));
+				UE_VLOG(OwnerComp.GetOwner(), LogBehaviorTree, Verbose, TEXT("Child[%d] \"%s\" execution allowed by %s"),
+					ChildIdx, *UBehaviorTreeTypes::DescribeNodeHelper(ChildNode), *UBehaviorTreeTypes::DescribeNodeHelper(TestDecorator));
 			}
 		}
 	}
diff --git a/Engine/Source/Runtime/AIModule/Private/BehaviorTree/BTDecorator.cpp b/Engine/Source/Runtime/AIModule/Private/BehaviorTree/BTDecorator.cpp
--- a/Engine/Source/Runtime/AIModule/Private/BehaviorTree/BTDecorator.cpp
+++ b/Engine/Source/Runtime/AIModule/Private/BehaviorTree/BTDecorator.cpp
@@ -103,7 +103,7 @@ void UBTDecorator::ConditionalFlowAbort(UBehaviorTreeComponent& OwnerComp, EBTDe
 	const bool bAlwaysRequestWhenPassing = (RequestMode == EBTDecoratorAbortRequest::ConditionPassing);
 
 	const bool bLogRestart = (bIsExecutingBranch != bPass) || (bIsExecutingBranch && bPass && (bAlwaysRequestWhenPassing || bAbortPending));
-	UE_VLOG(&OwnerComp, LogBehaviorTree, Verbose, TEXT("%s, ConditionalFlowAbort(%s) pass:%d executingBranch:%d abortPending:%d => %s"),
+	UE_VLOG(OwnerComp.GetOwner(), LogBehaviorTree, Verbose, TEXT("%s, ConditionalFlowAbort(%s) pass:%d executingBranch:%d abortPending:%d => %s"),
 		*UBehaviorTreeTypes::DescribeNodeHelper(this),
 		bAlwaysRequestWhenPassing ? TEXT("always when passing") : TEXT("on change"),
 		bPass ? 1 : 0,
diff --git a/Engine/Source/Runtime/AIModule/Private/BehaviorTree/BehaviorTreeComponent.cpp b/Engine/Source/Runtime/AIModule/Private/BehaviorTree/BehaviorTreeComponent.cpp
--- a/Engine/Source/Runtime/AIModule/Private/BehaviorTree/BehaviorTreeComponent.cpp
+++ b/Engine/Source/Runtime/AIModule/Private/BehaviorTree/BehaviorTreeComponent.cpp
@@ -1324,6 +1324,10 @@ void UBehaviorTreeComponent::ProcessExecutionRequest()
 			{
 				// error occurred and tree will restart, all pending deactivation notifies will be lost
 				// this is should happen
+
+				BT_SEARCHLOG(SearchData, Error, TEXT("Unable to deactivate up to %s. Active node is %s. All pending updates will be lost!"), 
+					*UBehaviorTreeTypes::DescribeNodeHelper(ExecutionRequest.ExecuteNode), 
+					*UBehaviorTreeTypes::DescribeNodeHelper(InstanceStack[ActiveInstanceIdx].ActiveNode));
 				SearchData.PendingUpdates.Reset();
 
 				return;
diff --git a/Engine/Source/Runtime/AIModule/Private/BehaviorTree/BehaviorTreeTypes.cpp b/Engine/Source/Runtime/AIModule/Private/BehaviorTree/BehaviorTreeTypes.cpp
--- a/Engine/Source/Runtime/AIModule/Private/BehaviorTree/BehaviorTreeTypes.cpp
+++ b/Engine/Source/Runtime/AIModule/Private/BehaviorTree/BehaviorTreeTypes.cpp
@@ -1,6 +1,7 @@
 // Copyright 1998-2019 Epic Games, Inc. All Rights Reserved.
 
 #include "BehaviorTree/BehaviorTreeTypes.h"
+#include "BehaviorTree/BehaviorTree.h"
 #include "GameFramework/Actor.h"
 #include "BehaviorTree/BTDecorator.h"
 #include "BehaviorTree/BTService.h"
@@ -463,7 +464,7 @@ FString UBehaviorTreeTypes::DescribeNodeUpdateMode(EBTNodeUpdateMode::Type Updat
 
 FString UBehaviorTreeTypes::DescribeNodeHelper(const UBTNode* Node)
 {
-	return Node ? FString::Printf(TEXT("%s[%d]"), *Node->GetNodeName(), Node->GetExecutionIndex()) : FString();
+	return Node ? FString::Printf(TEXT("%s::%s[%d]"), *GetNameSafe(Node->GetTreeAsset()), *Node->GetNodeName(), Node->GetExecutionIndex()) : FString();
 }
 
 FString UBehaviorTreeTypes::GetShortTypeName(const UObject* Ob)
diff --git a/Engine/Source/Runtime/AIModule/Private/BehaviorTree/Tasks/BTTask_MoveTo.cpp b/Engine/Source/Runtime/AIModule/Private/BehaviorTree/Tasks/BTTask_MoveTo.cpp
--- a/Engine/Source/Runtime/AIModule/Private/BehaviorTree/Tasks/BTTask_MoveTo.cpp
+++ b/Engine/Source/Runtime/AIModule/Private/BehaviorTree/Tasks/BTTask_MoveTo.cpp
@@ -196,13 +196,14 @@ EBlackboardNotificationResult UBTTask_MoveTo::OnBlackboardValueChange(const UBla
 		return EBlackboardNotificationResult::RemoveObserver;
 	}
 
+	AAIController* MyController = BehaviorComp->GetAIOwner();
 	uint8* RawMemory = BehaviorComp->GetNodeMemory(this, BehaviorComp->FindInstanceContainingNode(this));
 	FBTMoveToTaskMemory* MyMemory = CastInstanceNodeMemory<FBTMoveToTaskMemory>(RawMemory);
 
 	const EBTTaskStatus::Type TaskStatus = BehaviorComp->GetTaskStatus(this);
 	if (TaskStatus != EBTTaskStatus::Active)
 	{
-		UE_VLOG(BehaviorComp, LogBehaviorTree, Error, TEXT("BT MoveTo \'%s\' task observing BB entry while no longer being active!"), *GetNodeName());
+		UE_VLOG(MyController, LogBehaviorTree, Error, TEXT("BT MoveTo \'%s\' task observing BB entry while no longer being active!"), *GetNodeName());
 
 		// resetting BBObserverDelegateHandle without unregistering observer since 
 		// returning EBlackboardNotificationResult::RemoveObserver here will take care of that for us
@@ -231,7 +232,7 @@ EBlackboardNotificationResult UBTTask_MoveTo::OnBlackboardValueChange(const UBla
 			// don't abort move if using AI tasks - it will mess things up
 			if (MyMemory->MoveRequestID.IsValid())
 			{
-				UE_VLOG(BehaviorComp, LogBehaviorTree, Log, TEXT("Blackboard value for goal has changed, aborting current move request"));
+				UE_VLOG(MyController, LogBehaviorTree, Log, TEXT("Blackboard value for goal has changed, aborting current move request"));
 				StopWaitingForMessages(*BehaviorComp);
 				BehaviorComp->GetAIOwner()->GetPathFollowingComponent()->AbortMove(*this, FPathFollowingResultFlags::NewRequest, MyMemory->MoveRequestID, EPathFollowingVelocityMode::Keep);
 			}
@@ -278,7 +279,7 @@ EBTNodeResult::Type UBTTask_MoveTo::AbortTask(UBehaviorTreeComponent& OwnerComp,
 			}
 			else
 			{
-				UE_VLOG(&OwnerComp, LogBehaviorTree, Error, TEXT("Can't abort path following! bWaitingForPath:false, MoveRequestID:invalid, MoveTask:none!"));
+				UE_VLOG(OwnerComp.GetAIOwner(), LogBehaviorTree, Error, TEXT("Can't abort path following! bWaitingForPath:false, MoveRequestID:invalid, MoveTask:none!"));
 			}
 		}
 	}
