BehaviorTree: prevent tree restarts caused by execution requests queued for branches getting deactivated while applying search data

diff --git a/Engine/Source/Runtime/AIModule/Classes/BehaviorTree/BehaviorTreeComponent.h b/Engine/Source/Runtime/AIModule/Classes/BehaviorTree/BehaviorTreeComponent.h
--- a/Engine/Source/Runtime/AIModule/Classes/BehaviorTree/BehaviorTreeComponent.h
+++ b/Engine/Source/Runtime/AIModule/Classes/BehaviorTree/BehaviorTreeComponent.h
@@ -340,7 +340,7 @@ protected:
 	void ExecuteTask(UBTTaskNode* TaskNode);
 
 	/** deactivate all nodes up to requested one */
-	bool DeactivateUpTo(UBTCompositeNode* Node, uint16 NodeInstanceIdx, EBTNodeResult::Type& NodeResult);
+	bool DeactivateUpTo(UBTCompositeNode* Node, uint16 NodeInstanceIdx, EBTNodeResult::Type& NodeResult, int32& OutLastDeactivatedChildIndex);
 
 	/** update state of aborting tasks */
 	void UpdateAbortingTasks();
diff --git a/Engine/Source/Runtime/AIModule/Classes/BehaviorTree/BehaviorTreeTypes.h b/Engine/Source/Runtime/AIModule/Classes/BehaviorTree/BehaviorTreeTypes.h
--- a/Engine/Source/Runtime/AIModule/Classes/BehaviorTree/BehaviorTreeTypes.h
+++ b/Engine/Source/Runtime/AIModule/Classes/BehaviorTree/BehaviorTreeTypes.h
@@ -369,6 +369,7 @@ struct FBTNodeIndex
 	bool IsSet() const { return InstanceIndex < MAX_uint16; }
 
 	FORCEINLINE bool operator==(const FBTNodeIndex& Other) const { return Other.ExecutionIndex == ExecutionIndex && Other.InstanceIndex == InstanceIndex; }
+	FORCEINLINE bool operator!=(const FBTNodeIndex& Other) const { return !operator==(Other); }
 	FORCEINLINE friend uint32 GetTypeHash(const FBTNodeIndex& Other) { return Other.ExecutionIndex ^ Other.InstanceIndex; }
 
 	FORCEINLINE FString Describe() const { return FString::Printf(TEXT("[%d:%d]"), InstanceIndex, ExecutionIndex); }
@@ -419,6 +420,9 @@ struct FBehaviorTreeSearchData
 	/** notifies for tree instances */
 	TArray<FBehaviorTreeSearchUpdateNotify> PendingNotifies;
 
+	/** node under which the search was performed */
+	FBTNodeIndex SearchRootNode;
+
 	/** first node allowed in search */
 	FBTNodeIndex SearchStart;
 
@@ -431,6 +435,15 @@ struct FBehaviorTreeSearchData
 	/** active instance index to rollback to */
 	int32 RollbackInstanceIdx;
 
+	/** start index of the deactivated branch */
+	uint16 DeactivatedBranchStartIndex;
+
+	/** end index of the deactivated branch */
+	uint16 DeactivatedBranchEndIndex;
+
+	/** if set, execution request from node in the deactivated branch will be skipped */
+	uint32 bFilterOutRequestFromDeactivatedBranch : 1;
+
 	/** if set, current search will be restarted in next tick */
 	uint32 bPostponeSearch : 1;
 
@@ -450,9 +463,17 @@ struct FBehaviorTreeSearchData
 	void Reset();
 
 	FBehaviorTreeSearchData(UBehaviorTreeComponent& InOwnerComp) 
-		: OwnerComp(InOwnerComp), RollbackInstanceIdx(INDEX_NONE), bPostponeSearch(false), bSearchInProgress(false), bPreserveActiveNodeMemoryOnRollback(false)
+		: OwnerComp(InOwnerComp), RollbackInstanceIdx(INDEX_NONE)
+		, DeactivatedBranchStartIndex(INDEX_NONE)
+		, DeactivatedBranchEndIndex(INDEX_NONE)
+		, bFilterOutRequestFromDeactivatedBranch(false)
+		, bPostponeSearch(false)
+		, bSearchInProgress(false)
+		, bPreserveActiveNodeMemoryOnRollback(false)
 	{}
 
+	FBehaviorTreeSearchData() = delete;
+
 private:
 
 	static int32 NextSearchId;
diff --git a/Engine/Source/Runtime/AIModule/Private/BehaviorTree/BehaviorTreeComponent.cpp b/Engine/Source/Runtime/AIModule/Private/BehaviorTree/BehaviorTreeComponent.cpp
--- a/Engine/Source/Runtime/AIModule/Private/BehaviorTree/BehaviorTreeComponent.cpp
+++ b/Engine/Source/Runtime/AIModule/Private/BehaviorTree/BehaviorTreeComponent.cpp
@@ -346,8 +346,9 @@ void UBehaviorTreeComponent::StopTree(EBTStopMode::Type StopMode)
 	// make sure that all nodes are getting deactivation notifies
 	if (InstanceStack.Num())
 	{
+		int32 DeactivatedChildIndex = INDEX_NONE;
 		EBTNodeResult::Type AbortedResult = EBTNodeResult::Aborted;
-		DeactivateUpTo(InstanceStack[0].RootNode, 0, AbortedResult);
+		DeactivateUpTo(InstanceStack[0].RootNode, 0, AbortedResult, DeactivatedChildIndex);
 	}
 
 	// clear current state, don't touch debugger data
@@ -857,6 +858,37 @@ void UBehaviorTreeComponent::RequestExecution(UBTCompositeNode* RequestedOn, int
 		return;
 	}
 
+	if (SearchData.bFilterOutRequestFromDeactivatedBranch)
+	{
+		// request on same node or with higher priority doesn't require additional checks
+		FBTNodeIndex RequestedOnNodeIdx(InstanceIdx, RequestedOn->GetExecutionIndex());
+		if (SearchData.SearchRootNode != RequestedOnNodeIdx && SearchData.SearchRootNode.TakesPriorityOver(RequestedOnNodeIdx))
+		{
+			// for requests inside the same instance, perform range check to skip node from the branch deactivated by applying the search data
+			if (InstanceIdx == SearchData.SearchRootNode.InstanceIndex)
+			{
+				const uint16 StartIndex = SearchData.DeactivatedBranchStartIndex;
+				const uint16 EndIndex = SearchData.DeactivatedBranchEndIndex;
+				const bool bCheckDeactivatedBranch = (StartIndex != INDEX_NONE && EndIndex != INDEX_NONE);
+				if (bCheckDeactivatedBranch && (StartIndex <= RequestedOnNodeIdx.ExecutionIndex && RequestedOnNodeIdx.ExecutionIndex < EndIndex))
+				{
+					UE_VLOG(GetOwner(), LogBehaviorTree, Log, TEXT("> skip: node index %s in a deactivated branch [%d..%d[ (applying search data for %s)"),
+						*RequestedOnNodeIdx.Describe(), StartIndex, EndIndex, *SearchData.SearchRootNode.Describe());
+					StoreDebuggerRestart(DebuggerNode, InstanceIdx, false);
+					return;
+				}
+			}
+			// for other instances (subtree), perform check to skip removed subtree
+			else if (SearchData.PendingNotifies.FindByPredicate([InstanceIdx](const FBehaviorTreeSearchUpdateNotify& Item) { return (Item.InstanceIndex == InstanceIdx); }))
+			{
+				UE_VLOG(GetOwner(), LogBehaviorTree, Log, TEXT("> skip: node index %s from a tree instance in a deactivated branch (applying search data for %s)"),
+					*RequestedOnNodeIdx.Describe(), *SearchData.SearchRootNode.Describe());
+				StoreDebuggerRestart(DebuggerNode, InstanceIdx, false);
+				return;
+			}
+		}
+	}
+
 	// when it's aborting and moving to higher priority node:
 	if (bSwitchToHigherPriority)
 	{
@@ -1123,8 +1155,12 @@ void UBehaviorTreeComponent::ApplySearchData(UBTNode* NewActiveNode)
 	// apply changes to aux nodes and parallel tasks
 	const int32 NewNodeExecutionIndex = NewActiveNode ? NewActiveNode->GetExecutionIndex() : 0;
 
+	SearchData.bFilterOutRequestFromDeactivatedBranch = true;
+
 	ApplySearchUpdates(SearchData.PendingUpdates, NewNodeExecutionIndex);
 	ApplySearchUpdates(SearchData.PendingUpdates, NewNodeExecutionIndex, true);
+	
+	SearchData.bFilterOutRequestFromDeactivatedBranch = false;
 
 	// tick newly added aux nodes to compensate for tick-search order changes
 	UWorld* MyWorld = GetWorld();
@@ -1307,9 +1343,13 @@ void UBehaviorTreeComponent::ProcessExecutionRequest()
 		CopyInstanceMemoryToPersistent();
 
 		// deactivate up to ExecuteNode
+		SearchData.DeactivatedBranchStartIndex = INDEX_NONE;
+		SearchData.DeactivatedBranchEndIndex = INDEX_NONE;
+
 		if (InstanceStack[ActiveInstanceIdx].ActiveNode != ExecutionRequest.ExecuteNode)
 		{
-			const bool bDeactivated = DeactivateUpTo(ExecutionRequest.ExecuteNode, ExecutionRequest.ExecuteInstanceIdx, NodeResult);
+			int32 LastDeactivatedChildIndex = INDEX_NONE;
+			const bool bDeactivated = DeactivateUpTo(ExecutionRequest.ExecuteNode, ExecutionRequest.ExecuteInstanceIdx, NodeResult, LastDeactivatedChildIndex);
 			if (!bDeactivated)
 			{
 				// error occurred and tree will restart, all pending deactivation notifies will be lost
@@ -1322,6 +1362,11 @@ void UBehaviorTreeComponent::ProcessExecutionRequest()
 
 				return;
 			}
+			else if (LastDeactivatedChildIndex != INDEX_NONE)
+			{
+				SearchData.DeactivatedBranchStartIndex = ExecutionRequest.ExecuteNode->GetChildExecutionIndex(LastDeactivatedChildIndex);
+				SearchData.DeactivatedBranchEndIndex = ExecutionRequest.ExecuteNode->GetChildExecutionIndex(LastDeactivatedChildIndex + 1);
+			}
 		}
 
 		FBehaviorTreeInstance& ActiveInstance = InstanceStack[ActiveInstanceIdx];
@@ -1329,6 +1374,7 @@ void UBehaviorTreeComponent::ProcessExecutionRequest()
 		SearchData.AssignSearchId();
 		SearchData.bPostponeSearch = false;
 		SearchData.bSearchInProgress = true;
+		SearchData.SearchRootNode = FBTNodeIndex(ExecutionRequest.ExecuteInstanceIdx, ExecutionRequest.ExecuteNode->GetExecutionIndex());
 
 		// activate root node if needed (can't be handled by parent composite...)
 		if (ActiveInstance.ActiveNode == NULL)
@@ -1476,7 +1522,13 @@ void UBehaviorTreeComponent::ProcessExecutionRequest()
 			// abort task if needed
 			if (InstanceStack.Last().ActiveNodeType == EBTActiveNode::ActiveTask)
 			{
+				// prevent new execution requests for nodes inside the deativated branch 
+				// that may result from the aborted task.
+				SearchData.bFilterOutRequestFromDeactivatedBranch = true;
+
 				AbortCurrentTask();
+
+				SearchData.bFilterOutRequestFromDeactivatedBranch = false;
 			}
 
 			// set next task to execute
@@ -1573,7 +1625,7 @@ void UBehaviorTreeComponent::RollbackSearchChanges()
 	}
 }
 
-bool UBehaviorTreeComponent::DeactivateUpTo(UBTCompositeNode* Node, uint16 NodeInstanceIdx, EBTNodeResult::Type& NodeResult)
+bool UBehaviorTreeComponent::DeactivateUpTo(UBTCompositeNode* Node, uint16 NodeInstanceIdx, EBTNodeResult::Type& NodeResult, int32& OutLastDeactivatedChildIndex)
 {
 	UBTNode* DeactivatedChild = InstanceStack[ActiveInstanceIdx].ActiveNode;
 	bool bDeactivateRoot = true;
@@ -1590,7 +1642,8 @@ bool UBehaviorTreeComponent::DeactivateUpTo(UBTCompositeNode* Node, uint16 NodeI
 		UBTCompositeNode* NotifyParent = DeactivatedChild->GetParentNode();
 		if (NotifyParent)
 		{
-			NotifyParent->OnChildDeactivation(SearchData, *DeactivatedChild, NodeResult);
+			OutLastDeactivatedChildIndex = NotifyParent->GetChildIndex(SearchData, *DeactivatedChild);
+			NotifyParent->OnChildDeactivation(SearchData, OutLastDeactivatedChildIndex, NodeResult);
 
 			BT_SEARCHLOG(SearchData, Verbose, TEXT("Deactivate node: %s"), *UBehaviorTreeTypes::DescribeNodeHelper(DeactivatedChild));
 			StoreDebuggerSearchStep(DeactivatedChild, ActiveInstanceIdx, NodeResult);
diff --git a/Engine/Source/Runtime/AIModule/Private/BehaviorTree/BehaviorTreeTypes.cpp b/Engine/Source/Runtime/AIModule/Private/BehaviorTree/BehaviorTreeTypes.cpp
--- a/Engine/Source/Runtime/AIModule/Private/BehaviorTree/BehaviorTreeTypes.cpp
+++ b/Engine/Source/Runtime/AIModule/Private/BehaviorTree/BehaviorTreeTypes.cpp
@@ -282,9 +282,13 @@ void FBehaviorTreeSearchData::Reset()
 {
 	PendingUpdates.Reset();
 	PendingNotifies.Reset();
+	SearchRootNode = FBTNodeIndex();
 	SearchStart = FBTNodeIndex();
 	SearchEnd = FBTNodeIndex();
 	RollbackInstanceIdx = INDEX_NONE;
+	DeactivatedBranchStartIndex = INDEX_NONE;
+	DeactivatedBranchEndIndex = INDEX_NONE;
+	bFilterOutRequestFromDeactivatedBranch = false;
 	bSearchInProgress = false;
 	bPostponeSearch = false;
 	bPreserveActiveNodeMemoryOnRollback = false;
