UAIDataProvider_QueryParams: Fixed data binding for boolean values. This value type is not encoded in the float param using reinterpret cast like integer but uses specific float values (-1.0f, 1.0f) so it needs to be retrieved using similar logic.

diff --git a/Engine/Source/Runtime/AIModule/Private/DataProviders/AIDataProvider_QueryParams.cpp b/Engine/Source/Runtime/AIModule/Private/DataProviders/AIDataProvider_QueryParams.cpp
--- a/Engine/Source/Runtime/AIModule/Private/DataProviders/AIDataProvider_QueryParams.cpp
+++ b/Engine/Source/Runtime/AIModule/Private/DataProviders/AIDataProvider_QueryParams.cpp
@@ -9,8 +9,12 @@ void UAIDataProvider_QueryParams::BindData(const UObject& Owner, int32 RequestId
 	if (QueryManager)
 	{
 		FloatValue = QueryManager->FindNamedParam(RequestId, ParamName);
+
+		// int param was encoded directly in the float value
 		IntValue = *((int32*)&FloatValue);
-		BoolValue = *((bool*)&FloatValue);
+
+		// bool param was encoded as -1.0f (false) or 1.0f (true) in the float value
+		BoolValue = FloatValue > 0.f;
 	}
 	else
 	{
