Crash timing sensitive fix when stopping BT and processing AITask on same frame with Timers.

#rnx

diff --git a/Engine/Source/Runtime/AIModule/Private/BehaviorTree/Tasks/BTTask_MoveTo.cpp b/Engine/Source/Runtime/AIModule/Private/BehaviorTree/Tasks/BTTask_MoveTo.cpp
--- a/Engine/Source/Runtime/AIModule/Private/BehaviorTree/Tasks/BTTask_MoveTo.cpp
+++ b/Engine/Source/Runtime/AIModule/Private/BehaviorTree/Tasks/BTTask_MoveTo.cpp
@@ -345,7 +345,7 @@ void UBTTask_MoveTo::OnGameplayTaskDeactivated(UGameplayTask& Task)
 			uint8* RawMemory = BehaviorComp->GetNodeMemory(this, BehaviorComp->FindInstanceContainingNode(this));
 			FBTMoveToTaskMemory* MyMemory = CastInstanceNodeMemory<FBTMoveToTaskMemory>(RawMemory);
 
-			if (MyMemory->bObserverCanFinishTask && (MoveTask == MyMemory->Task))
+			if (MyMemory && MyMemory->bObserverCanFinishTask && (MoveTask == MyMemory->Task))
 			{
 				const bool bSuccess = MoveTask->WasMoveSuccessful();
 				FinishLatentTask(*BehaviorComp, bSuccess ? EBTNodeResult::Succeeded : EBTNodeResult::Failed);
diff --git a/Engine/Source/Runtime/AIModule/Private/Tasks/AITask_MoveTo.cpp b/Engine/Source/Runtime/AIModule/Private/Tasks/AITask_MoveTo.cpp
--- a/Engine/Source/Runtime/AIModule/Private/Tasks/AITask_MoveTo.cpp
+++ b/Engine/Source/Runtime/AIModule/Private/Tasks/AITask_MoveTo.cpp
@@ -242,25 +242,13 @@ void UAITask_MoveTo::ResetObservers()
 
 void UAITask_MoveTo::ResetTimers()
 {
-	if (MoveRetryTimerHandle.IsValid())
+	if (OwnerController)
 	{
-		if (OwnerController)
-		{
-			OwnerController->GetWorldTimerManager().ClearTimer(MoveRetryTimerHandle);
-		}
-
-		MoveRetryTimerHandle.Invalidate();
-	}
-
-	if (PathRetryTimerHandle.IsValid())
-	{
-		if (OwnerController)
-		{
-			OwnerController->GetWorldTimerManager().ClearTimer(PathRetryTimerHandle);
-		}
-
-		PathRetryTimerHandle.Invalidate();
+		// Remove all timers including the ones that might have been set with SetTimerForNextTick 
+		OwnerController->GetWorldTimerManager().ClearAllTimersForObject(this);
 	}
+	MoveRetryTimerHandle.Invalidate();
+	PathRetryTimerHandle.Invalidate();
 }
 
 void UAITask_MoveTo::OnDestroy(bool bInOwnerFinished)
