Turn on warning for if there are a large number of active EQS queries in test builds (previously only enabled in dev/debug builds)

#rnx

diff --git a/Engine/Source/Runtime/AIModule/Classes/EnvironmentQuery/EnvQueryManager.h b/Engine/Source/Runtime/AIModule/Classes/EnvironmentQuery/EnvQueryManager.h
--- a/Engine/Source/Runtime/AIModule/Classes/EnvironmentQuery/EnvQueryManager.h
+++ b/Engine/Source/Runtime/AIModule/Classes/EnvironmentQuery/EnvQueryManager.h
@@ -307,7 +307,8 @@ private:
 
 #if !(UE_BUILD_SHIPPING || UE_BUILD_TEST)
 	static bool bAllowEQSTimeSlicing;
-
+#endif
+#if !(UE_BUILD_SHIPPING)
 	mutable double LastQueryCountWarningThresholdTime;
 
 	void CheckQueryCount() const;
diff --git a/Engine/Source/Runtime/AIModule/Private/EnvironmentQuery/EnvQueryManager.cpp b/Engine/Source/Runtime/AIModule/Private/EnvironmentQuery/EnvQueryManager.cpp
--- a/Engine/Source/Runtime/AIModule/Private/EnvironmentQuery/EnvQueryManager.cpp
+++ b/Engine/Source/Runtime/AIModule/Private/EnvironmentQuery/EnvQueryManager.cpp
@@ -116,7 +116,7 @@ UEnvQueryManager::UEnvQueryManager(const FObjectInitializer& ObjectInitializer)
 
 	QueryCountWarningThreshold = 0;
 	QueryCountWarningInterval = 30.0;
-#if !(UE_BUILD_SHIPPING || UE_BUILD_TEST)
+#if !(UE_BUILD_SHIPPING)
 	LastQueryCountWarningThresholdTime = -FLT_MAX;
 #endif
 
@@ -341,7 +341,7 @@ void UEnvQueryManager::Tick(float DeltaTime)
 
 	CSV_SCOPED_TIMING_STAT_EXCLUSIVE(EnvQueryManager);
 
-#if !(UE_BUILD_SHIPPING || UE_BUILD_TEST)
+#if !(UE_BUILD_SHIPPING)
 	CheckQueryCount();
 #endif
 
@@ -503,7 +503,7 @@ void UEnvQueryManager::Tick(float DeltaTime)
 	}
 }
 
-#if !(UE_BUILD_SHIPPING || UE_BUILD_TEST)
+#if !(UE_BUILD_SHIPPING)
 void UEnvQueryManager::CheckQueryCount() const
 {
 	if ((QueryCountWarningThreshold > 0) && (RunningQueries.Num() >= QueryCountWarningThreshold))
