Fix bt crash when the type of a blackboard key to none and saving.

#rnx

diff --git a/Engine/Source/Runtime/AIModule/Private/BehaviorTree/Decorators/BTDecorator_Blackboard.cpp b/Engine/Source/Runtime/AIModule/Private/BehaviorTree/Decorators/BTDecorator_Blackboard.cpp
--- a/Engine/Source/Runtime/AIModule/Private/BehaviorTree/Decorators/BTDecorator_Blackboard.cpp
+++ b/Engine/Source/Runtime/AIModule/Private/BehaviorTree/Decorators/BTDecorator_Blackboard.cpp
@@ -132,7 +132,7 @@ void UBTDecorator_Blackboard::BuildDescription()
 	{
 		// safety feature to not crash when changing couple of properties on a bunch
 		// while "post edit property" triggers for every each of them
-		if (EntryInfo->KeyType->GetClass() == BlackboardKey.SelectedKeyType)
+		if (EntryInfo->KeyType && EntryInfo->KeyType->GetClass() == BlackboardKey.SelectedKeyType)
 		{
 			const FString KeyName = EntryInfo->EntryName.ToString();
 			CacheOperationEnums();		
