EnvQueryManager: added some Verbose logging about queries starting/stopping

diff --git a/Engine/Source/Runtime/AIModule/Classes/EnvironmentQuery/EnvQueryManager.h b/Engine/Source/Runtime/AIModule/Classes/EnvironmentQuery/EnvQueryManager.h
--- a/Engine/Source/Runtime/AIModule/Classes/EnvironmentQuery/EnvQueryManager.h
+++ b/Engine/Source/Runtime/AIModule/Classes/EnvironmentQuery/EnvQueryManager.h
@@ -203,9 +203,12 @@ class AIMODULE_API UEnvQueryManager : public UAISubsystem, public FSelfRegisteri
 	/** find value of named param stored with active query */
 	float FindNamedParam(int32 QueryId, FName ParamName) const;
 
-	/** execute query */
+	/** aborts specific query */
 	bool AbortQuery(int32 RequestID);
 
+	/** outputs active queries to log */
+	void PrintActiveQueryInfo() const;
+
 	/** fail all running queries on cleaning the world */
 	virtual void OnWorldCleanup();
 
@@ -312,7 +315,7 @@ private:
 	mutable double LastQueryCountWarningThresholdTime;
 
 	void CheckQueryCount() const;
-	void LogQueryCountWarning() const;
+	void LogQueryInfo(bool bDisplayThresholdWarning) const;
 #endif
 
 #if WITH_EDITOR
diff --git a/Engine/Source/Runtime/AIModule/Private/EnvironmentQuery/EnvQueryManager.cpp b/Engine/Source/Runtime/AIModule/Private/EnvironmentQuery/EnvQueryManager.cpp
--- a/Engine/Source/Runtime/AIModule/Private/EnvironmentQuery/EnvQueryManager.cpp
+++ b/Engine/Source/Runtime/AIModule/Private/EnvironmentQuery/EnvQueryManager.cpp
@@ -223,6 +223,8 @@ int32 UEnvQueryManager::RunQuery(const TSharedPtr<FEnvQueryInstance>& QueryInsta
 	QueryInstance->FinishDelegate = FinishDelegate;
 	RunningQueries.Add(QueryInstance);
 
+	UE_LOG(LogEQS, Verbose, TEXT("%s: Query: %s - Owner: %s"), ANSI_TO_TCHAR(__FUNCTION__), *QueryInstance->QueryName, *GetNameSafe(QueryInstance->Owner.Get()));
+
 	return QueryInstance->QueryID;
 }
 
@@ -246,6 +248,8 @@ void UEnvQueryManager::RunInstantQuery(const TSharedPtr<FEnvQueryInstance>& Quer
 		return;
 	}
 
+	UE_LOG(LogEQS, Verbose, TEXT("%s: Query: %s - Owner: %s"), ANSI_TO_TCHAR(__FUNCTION__), *QueryInstance->QueryName, *GetNameSafe(QueryInstance->Owner.Get()));
+
 	{
 		CSV_SCOPED_TIMING_STAT_EXCLUSIVE(EnvQueryManager);
 
@@ -276,6 +280,8 @@ void UEnvQueryManager::RemoveAllQueriesByQuerier(const UObject& Querier, bool bE
 			{
 				QueryInstance->MarkAsAborted();
 
+				UE_LOG(LogEQS, Verbose, TEXT("%s: Query: %s - Owner: %s"), ANSI_TO_TCHAR(__FUNCTION__), *QueryInstance->QueryName, *GetNameSafe(QueryInstance->Owner.Get()));
+
 				if (bExecuteFinishDelegate)
 				{
 					QueryInstance->FinishDelegate.ExecuteIfBound(QueryInstance);
@@ -293,6 +299,9 @@ TSharedPtr<FEnvQueryInstance> UEnvQueryManager::PrepareQueryInstance(const FEnvQ
 	TSharedPtr<FEnvQueryInstance> QueryInstance = CreateQueryInstance(Request.QueryTemplate, RunMode);
 	if (!QueryInstance.IsValid())
 	{
+		UE_LOG(LogEQS, Warning, TEXT("Error creating query instance for QueryTemplate: %s - Owner: %s"),
+			Request.QueryTemplate != nullptr ? *Request.QueryTemplate->QueryName.ToString() : TEXT("NONE"),
+			*GetNameSafe(Request.Owner));
 		return NULL;
 	}
 
@@ -320,6 +329,8 @@ bool UEnvQueryManager::AbortQuery(int32 RequestID)
 		if (QueryInstance->QueryID == RequestID &&
 			QueryInstance->IsFinished() == false)
 		{
+			UE_LOG(LogEQS, Verbose, TEXT("%s: Query: %s - Owner: %s"), ANSI_TO_TCHAR(__FUNCTION__), *QueryInstance->QueryName, *GetNameSafe(QueryInstance->Owner.Get()));
+
 			QueryInstance->MarkAsAborted();
 			QueryInstance->FinishDelegate.ExecuteIfBound(QueryInstance);
 			
@@ -403,6 +414,8 @@ void UEnvQueryManager::Tick(float DeltaTime)
 						EQSDebugger.StoreQuery(QueryInstance);
 #endif // USE_EQS_DEBUGGER
 
+						UE_LOG(LogEQS, Verbose, TEXT("%s: Finished Query: %s - Owner: %s"), ANSI_TO_TCHAR(__FUNCTION__), *QueryInstance->QueryName, *GetNameSafe(QueryInstance->Owner.Get()));
+
 						QueryInstancePtr->FinishDelegate.ExecuteIfBound(QueryInstance);
 
 						ResultHandlingDuration = FPlatformTime::Seconds() - ResultHandlingStartTime;
@@ -512,22 +525,29 @@ void UEnvQueryManager::CheckQueryCount() const
 
 		if ((LastQueryCountWarningThresholdTime < 0.0) || ((LastQueryCountWarningThresholdTime + QueryCountWarningInterval) < CurrentTime))
 		{
-			LogQueryCountWarning();
+			LogQueryInfo(true /* bDisplayThresholdWarning */);
 
 			LastQueryCountWarningThresholdTime = CurrentTime;
 		}
 	}
 }
 
-void UEnvQueryManager::LogQueryCountWarning() const
+void UEnvQueryManager::LogQueryInfo(bool bDisplayThresholdWarning) const
 {
-	UE_LOG(LogEQS, Warning, TEXT("The number of EQS queries has reached (%d) the warning threshold (%d).  Logging queries."), RunningQueries.Num(), QueryCountWarningThreshold);
+	if (bDisplayThresholdWarning)
+	{
+		UE_LOG(LogEQS, Warning, TEXT("The number of EQS queries (%d) has reached the warning threshold (%d).  Logging queries."), RunningQueries.Num(), QueryCountWarningThreshold);
+	}
+	else
+	{
+		UE_LOG(LogEQS, Warning, TEXT("The number of EQS queries is (%d).  Logging queries."), RunningQueries.Num());
+	}
 
 	for (const TSharedPtr<FEnvQueryInstance>& RunningQuery : RunningQueries)
 	{
 		if (RunningQuery.IsValid())
 		{
-			UE_LOG(LogEQS, Warning, TEXT("Query: %s - Owner: %s"), *RunningQuery->QueryName, RunningQuery->Owner.IsValid() ? *RunningQuery->Owner->GetName() : TEXT("Invalid"));
+			UE_LOG(LogEQS, Warning, TEXT("Query: %s - Owner: %s - %s"), *RunningQuery->QueryName, RunningQuery->Owner.IsValid() ? *RunningQuery->Owner->GetName() : TEXT("Invalid"), *RunningQuery->GetExecutionTimeDescription());
 		}
 		else
 		{
@@ -537,6 +557,13 @@ void UEnvQueryManager::LogQueryCountWarning() const
 }
 #endif
 
+void UEnvQueryManager::PrintActiveQueryInfo() const
+{
+#if !(UE_BUILD_SHIPPING)
+	LogQueryInfo(false /* bDisplayThresholdWarning */);
+#endif
+}
+
 void UEnvQueryManager::OnWorldCleanup()
 {
 	if (RunningQueries.Num() > 0)
@@ -550,6 +577,7 @@ void UEnvQueryManager::OnWorldCleanup()
 			TSharedPtr<FEnvQueryInstance>& QueryInstance = RunningQueriesCopy[Index];
 			if (QueryInstance->IsFinished() == false)
 			{
+				UE_LOG(LogEQS, Verbose, TEXT("Query failed due to world cleanup: Query: %s - Owner: %s"), *QueryInstance->QueryName, *GetNameSafe(QueryInstance->Owner.Get()));
 				QueryInstance->MarkAsFailed();
 				QueryInstance->FinishDelegate.ExecuteIfBound(QueryInstance);
 			}
