Speculative fix for issue with sub object reuse in the AI Systems components
#rnx

diff --git a/Engine/Source/Runtime/AIModule/Private/AISystem.cpp b/Engine/Source/Runtime/AIModule/Private/AISystem.cpp
--- a/Engine/Source/Runtime/AIModule/Private/AISystem.cpp
+++ b/Engine/Source/Runtime/AIModule/Private/AISystem.cpp
@@ -29,13 +29,7 @@ UAISystem::UAISystem(const FObjectInitializer& ObjectInitializer)
 
 	bEnableBTAITasks = false;
 
-	if (HasAnyFlags(RF_ClassDefaultObject) == false)
-	{
-		BehaviorTreeManager = NewObject<UBehaviorTreeManager>(this, TEXT("BehaviorTreeManager"));
-		EnvironmentQueryManager = NewObject<UEnvQueryManager>(this, TEXT("EQSManager"));
-		NavLocalGrids = NewObject<UNavLocalGridManager>(this, TEXT("NavLocalGrids"));
-	}
-	else
+	if (HasAnyFlags(RF_ClassDefaultObject))
 	{
 		// game-wise config
 		if (FParse::Param(FCommandLine::Get(), TEXT("FixedSeed")) == false)
@@ -62,7 +56,14 @@ void UAISystem::PostInitProperties()
 	if (HasAnyFlags(RF_ClassDefaultObject) == false)
 	{
 		UWorld* WorldOuter = GetOuterWorld();
-		
+
+		BehaviorTreeManager = NewObject<UBehaviorTreeManager>(this, TEXT("BehaviorTreeManager"));
+		ensure(BehaviorTreeManager != nullptr);
+		EnvironmentQueryManager = NewObject<UEnvQueryManager>(this, TEXT("EQSManager"));
+		ensure(BehaviorTreeManager != nullptr);
+		NavLocalGrids = NewObject<UNavLocalGridManager>(this, TEXT("NavLocalGrids"));
+		ensure(BehaviorTreeManager != nullptr);
+
 		TSubclassOf<UAIHotSpotManager> HotSpotManagerClass = HotSpotManagerClassName.IsValid() ? LoadClass<UAIHotSpotManager>(NULL, *HotSpotManagerClassName.ToString(), NULL, LOAD_None, NULL) : nullptr;
 		if (HotSpotManagerClass)
 		{
