Fix typo in ensure logic
#rnx

diff --git a/Engine/Source/Runtime/AIModule/Private/AISystem.cpp b/Engine/Source/Runtime/AIModule/Private/AISystem.cpp
--- a/Engine/Source/Runtime/AIModule/Private/AISystem.cpp
+++ b/Engine/Source/Runtime/AIModule/Private/AISystem.cpp
@@ -60,9 +60,9 @@ void UAISystem::PostInitProperties()
 		BehaviorTreeManager = NewObject<UBehaviorTreeManager>(this, TEXT("BehaviorTreeManager"));
 		ensure(BehaviorTreeManager != nullptr);
 		EnvironmentQueryManager = NewObject<UEnvQueryManager>(this, TEXT("EQSManager"));
-		ensure(BehaviorTreeManager != nullptr);
+		ensure(EnvironmentQueryManager != nullptr);
 		NavLocalGrids = NewObject<UNavLocalGridManager>(this, TEXT("NavLocalGrids"));
-		ensure(BehaviorTreeManager != nullptr);
+		ensure(NavLocalGrids != nullptr);
 
 		TSubclassOf<UAIHotSpotManager> HotSpotManagerClass = HotSpotManagerClassName.IsValid() ? LoadClass<UAIHotSpotManager>(NULL, *HotSpotManagerClassName.ToString(), NULL, LOAD_None, NULL) : nullptr;
 		if (HotSpotManagerClass)
