Additional logging and new fix attempt for JIRA-239935
#rnx

diff --git a/Engine/Source/Runtime/AIModule/Classes/AISubsystem.h b/Engine/Source/Runtime/AIModule/Classes/AISubsystem.h
--- a/Engine/Source/Runtime/AIModule/Classes/AISubsystem.h
+++ b/Engine/Source/Runtime/AIModule/Classes/AISubsystem.h
@@ -22,6 +22,7 @@ private:
 			
 public:
 	UAISubsystem(const FObjectInitializer& ObjectInitializer = FObjectInitializer::Get());
+	~UAISubsystem();
 
 	virtual UWorld* GetWorld() const override;
 
diff --git a/Engine/Source/Runtime/AIModule/Private/AISubsystem.cpp b/Engine/Source/Runtime/AIModule/Private/AISubsystem.cpp
--- a/Engine/Source/Runtime/AIModule/Private/AISubsystem.cpp
+++ b/Engine/Source/Runtime/AIModule/Private/AISubsystem.cpp
@@ -10,6 +10,8 @@ DEFINE_LOG_CATEGORY_STATIC(LogAISub, Log, All);
 UAISubsystem::UAISubsystem(const FObjectInitializer& ObjectInitializer)
 	: Super(ObjectInitializer)
 {
+	UE_LOG(LogAISub, Log, TEXT("%s %s (%p), frame # %llu"), ANSI_TO_TCHAR(__FUNCTION__), *ObjectInitializer.GetObj()->GetClass()->GetName(), this, (uint64)GFrameCounter);
+
 	if (HasAnyFlags(RF_ClassDefaultObject) == false)
 	{
 		AISystem = Cast<UAISystem>(GetOuter());
@@ -44,3 +46,8 @@ TStatId UAISubsystem::GetStatId() const
 {
 	RETURN_QUICK_DECLARE_CYCLE_STAT(UAISubsystem, STATGROUP_Tickables);
 }
+
+UAISubsystem::~UAISubsystem()
+{
+	UE_LOG(LogAISub, Log, TEXT("%s %s (%p), frame # %llu"), ANSI_TO_TCHAR(__FUNCTION__), *GetClass()->GetName(), this, (uint64)GFrameCounter);
+}
diff --git a/Engine/Source/Runtime/AIModule/Private/AISystem.cpp b/Engine/Source/Runtime/AIModule/Private/AISystem.cpp
--- a/Engine/Source/Runtime/AIModule/Private/AISystem.cpp
+++ b/Engine/Source/Runtime/AIModule/Private/AISystem.cpp
@@ -57,11 +57,11 @@ void UAISystem::PostInitProperties()
 	{
 		UWorld* WorldOuter = GetOuterWorld();
 
-		BehaviorTreeManager = NewObject<UBehaviorTreeManager>(this, TEXT("BehaviorTreeManager"));
+		BehaviorTreeManager = NewObject<UBehaviorTreeManager>(this);
 		ensure(BehaviorTreeManager != nullptr);
-		EnvironmentQueryManager = NewObject<UEnvQueryManager>(this, TEXT("EQSManager"));
+		EnvironmentQueryManager = NewObject<UEnvQueryManager>(this);
 		ensure(EnvironmentQueryManager != nullptr);
-		NavLocalGrids = NewObject<UNavLocalGridManager>(this, TEXT("NavLocalGrids"));
+		NavLocalGrids = NewObject<UNavLocalGridManager>(this);
 		ensure(NavLocalGrids != nullptr);
 
 		TSubclassOf<UAIHotSpotManager> HotSpotManagerClass = HotSpotManagerClassName.IsValid() ? LoadClass<UAIHotSpotManager>(NULL, *HotSpotManagerClassName.ToString(), NULL, LOAD_None, NULL) : nullptr;
