Changed the code to allow custom BrainComponents to run, as well as allow BrainComponents (including BehaviorTreeComponent) to be set up as components on AIController and run automatically when Actor is possessed by the AIController.
- Added StartLogic() to BrainComponent
- Added UAIController bStartAILogicOnPosses to let the BrainComponent start when OnPossess is called
- Added DefaultBlackboardAsset to BlackboardComponent which is put in use during component initialized
- Added DefaultBehaviorTreeAsset to BehaviorTreeComponent
- Added StartLogic to BehaviorTreeComponent, runs DefaultBehaviorTreeAsset if set

diff --git a/Engine/Source/Runtime/AIModule/Classes/AIController.h b/Engine/Source/Runtime/AIModule/Classes/AIController.h
--- a/Engine/Source/Runtime/AIModule/Classes/AIController.h
+++ b/Engine/Source/Runtime/AIModule/Classes/AIController.h
@@ -90,8 +90,13 @@ class AIMODULE_API AAIController : public AController, public IAIPerceptionListe
 protected:
 	FFocusKnowledge	FocusInformation;
 
-	/** By default AI's logic gets stopped when controlled Pawn is unpossesed. Setting this flag to false
-	 *	will make AI logic persist past loosing controll over a pawn */
+	/** By default AI's logic does not start when controlled Pawn is possessed. Setting this flag to true
+	 *	will make AI logic start when pawn is possessed */
+	UPROPERTY(EditAnywhere, BlueprintReadWrite, Category = AI)
+	uint32 bStartAILogicOnPossess : 1;
+
+	/** By default AI's logic gets stopped when controlled Pawn is unpossessed. Setting this flag to false
+	 *	will make AI logic persist past losing control over a pawn */
 	UPROPERTY(EditAnywhere, BlueprintReadWrite, Category = AI)
 	uint32 bStopAILogicOnUnposses : 1;
 
diff --git a/Engine/Source/Runtime/AIModule/Classes/BehaviorTree/BehaviorTreeComponent.h b/Engine/Source/Runtime/AIModule/Classes/BehaviorTree/BehaviorTreeComponent.h
--- a/Engine/Source/Runtime/AIModule/Classes/BehaviorTree/BehaviorTreeComponent.h
+++ b/Engine/Source/Runtime/AIModule/Classes/BehaviorTree/BehaviorTreeComponent.h
@@ -77,12 +77,13 @@ struct FBTTreeStartInfo
 	bool HasPendingInitialize() const { return bPendingInitialize && IsSet(); }
 };
 
-UCLASS()
+UCLASS(ClassGroup = AI, meta = (BlueprintSpawnableComponent))
 class AIMODULE_API UBehaviorTreeComponent : public UBrainComponent
 {
 	GENERATED_UCLASS_BODY()
 
 	// Begin UBrainComponent overrides
+	virtual void StartLogic() override;
 	virtual void RestartLogic() override;
 	virtual void StopLogic(const FString& Reason) override;
 	virtual void PauseLogic(const FString& Reason) override;
@@ -391,6 +392,11 @@ protected:
 	friend FBehaviorTreeDebugger;
 	friend FBehaviorTreeInstance;
 	friend FScopedBehaviorTreeLock;
+
+protected:
+	/** data asset defining the tree */
+	UPROPERTY(EditDefaultsOnly, BlueprintReadWrite, Category = AI)
+	UBehaviorTree* DefaultBehaviorTreeAsset;
 };
 
 //////////////////////////////////////////////////////////////////////////
diff --git a/Engine/Source/Runtime/AIModule/Classes/BehaviorTree/BlackboardComponent.h b/Engine/Source/Runtime/AIModule/Classes/BehaviorTree/BlackboardComponent.h
--- a/Engine/Source/Runtime/AIModule/Classes/BehaviorTree/BlackboardComponent.h
+++ b/Engine/Source/Runtime/AIModule/Classes/BehaviorTree/BlackboardComponent.h
@@ -38,7 +38,7 @@ namespace EBlackboardDescription
 }
 
 
-UCLASS(ClassGroup = AI, meta = (BlueprintSpawnableComponent))
+UCLASS(ClassGroup = AI, meta = (BlueprintSpawnableComponent), hidecategories = (Sockets, Collision))
 class AIMODULE_API UBlackboardComponent : public UActorComponent
 {
 	GENERATED_BODY()
@@ -101,7 +101,7 @@ public:
 	/** caches UBrainComponent pointer to be used in communication */
 	void CacheBrainComponent(UBrainComponent& BrainComponent);
 
-	/** setup component for using given blackboard asset */
+	/** setup component for using given blackboard asset, returns true if blackboard is properly initialized for specified blackboard data */
 	bool InitializeBlackboard(UBlackboardData& NewAsset);
 	
 	/** @return true if component can be used with specified blackboard asset */
@@ -234,6 +234,9 @@ protected:
 	UBrainComponent* BrainComp;
 
 	/** data asset defining entries */
+	UPROPERTY(EditDefaultsOnly, Category = AI)
+	UBlackboardData* DefaultBlackboardAsset;
+
 	UPROPERTY(transient)
 	UBlackboardData* BlackboardAsset;
 
diff --git a/Engine/Source/Runtime/AIModule/Classes/BrainComponent.h b/Engine/Source/Runtime/AIModule/Classes/BrainComponent.h
--- a/Engine/Source/Runtime/AIModule/Classes/BrainComponent.h
+++ b/Engine/Source/Runtime/AIModule/Classes/BrainComponent.h
@@ -113,7 +113,7 @@ private:
 	FAIMessageObserver& operator=(const FAIMessageObserver&);
 };
 
-UCLASS(BlueprintType)
+UCLASS(ClassGroup = AI, BlueprintType, hidecategories = (Sockets, Collision))
 class AIMODULE_API UBrainComponent : public UActorComponent, public IAIResourceInterface
 {
 	GENERATED_UCLASS_BODY()
@@ -138,16 +138,26 @@ public:
 	 *	@note this call does nothing if logic is not locked at the moment of call */
 	void RequestLogicRestartOnUnlock();
 
+	/** Starts brain logic. If brain is already running, will not do anything. */
+	UFUNCTION(BlueprintCallable, Category = "AI|Logic")
+	virtual void StartLogic();
+
+	/** Restarts currently running or previously ran brain logic. */
 	UFUNCTION(BlueprintCallable, Category = "AI|Logic")
 	virtual void RestartLogic();
 
+	/** Stops currently running brain logic. */
 	UFUNCTION(BlueprintCallable, Category = "AI|Logic")
 	virtual void StopLogic(const FString& Reason);
 
 	/** AI logic won't be needed anymore, stop all activity and run cleanup */
 	virtual void Cleanup() {}
+
+	/** Pause logic and blackboard updates. */
 	virtual void PauseLogic(const FString& Reason) {}
-	/** MUST be called by child implementations!
+	
+	/** Resumes paused brain logic.
+	 *  MUST be called by child implementations!
 	 *	@return indicates whether child class' ResumeLogic should be called (true) or has it been 
 	 *	handled in a different way and no other actions are required (false)*/
 	virtual EAILogicResuming::Type ResumeLogic(const FString& Reason);
diff --git a/Engine/Source/Runtime/AIModule/Private/AIController.cpp b/Engine/Source/Runtime/AIModule/Private/AIController.cpp
--- a/Engine/Source/Runtime/AIModule/Private/AIController.cpp
+++ b/Engine/Source/Runtime/AIModule/Private/AIController.cpp
@@ -47,6 +47,7 @@ AAIController::AAIController(const FObjectInitializer& ObjectInitializer)
 	bWantsPlayerState = false;
 	TeamID = FGenericTeamId::NoTeam;
 
+	bStartAILogicOnPossess = false;
 	bStopAILogicOnUnposses = true;
 }
 
@@ -66,6 +67,15 @@ void AAIController::PostInitializeComponents()
 		InitPlayerState();
 	}
 
+	if (BrainComponent == nullptr)
+	{
+		BrainComponent = FindComponentByClass<UBrainComponent>();
+	}
+	if (Blackboard == nullptr)
+	{
+		Blackboard = FindComponentByClass<UBlackboardComponent>();
+	}
+
 #if ENABLE_VISUAL_LOG
 	for (UActorComponent* Component : GetComponents())
 	{
@@ -501,6 +511,16 @@ void AAIController::OnPossess(APawn* InPawn)
 
 		REDIRECT_OBJECT_TO_VLOG(CachedGameplayTasksComponent, this);
 	}
+
+	if (Blackboard && Blackboard->GetBlackboardAsset())
+	{
+		InitializeBlackboard(*Blackboard, *Blackboard->GetBlackboardAsset());
+	}
+
+	if (bStartAILogicOnPossess && BrainComponent)
+	{
+		BrainComponent->StartLogic();
+	}
 }
 
 void AAIController::OnUnPossess()
@@ -514,12 +534,9 @@ void AAIController::OnUnPossess()
 		PathFollowingComponent->Cleanup();
 	}
 
-	if (bStopAILogicOnUnposses)
+	if (bStopAILogicOnUnposses && BrainComponent)
 	{
-		if (BrainComponent)
-		{
-			BrainComponent->Cleanup();
-		}
+		BrainComponent->Cleanup();
 	}
 
 	if (CachedGameplayTasksComponent && (CachedGameplayTasksComponent->GetOwner() == CurrentPawn))
@@ -888,7 +905,6 @@ bool AAIController::RunBehaviorTree(UBehaviorTree* BTAsset)
 	}
 
 	bool bSuccess = true;
-	bool bShouldInitializeBlackboard = false;
 
 	// see if need a blackboard component at all
 	UBlackboardComponent* BlackboardComp = Blackboard;
diff --git a/Engine/Source/Runtime/AIModule/Private/BehaviorTree/BehaviorTreeComponent.cpp b/Engine/Source/Runtime/AIModule/Private/BehaviorTree/BehaviorTreeComponent.cpp
--- a/Engine/Source/Runtime/AIModule/Private/BehaviorTree/BehaviorTreeComponent.cpp
+++ b/Engine/Source/Runtime/AIModule/Private/BehaviorTree/BehaviorTreeComponent.cpp
@@ -101,9 +101,35 @@ void UBehaviorTreeComponent::UninitializeComponent()
 	Super::UninitializeComponent();
 }
 
+void UBehaviorTreeComponent::StartLogic()
+{
+	UE_VLOG(GetOwner(), LogBehaviorTree, Log, TEXT("%s"), ANSI_TO_TCHAR(__FUNCTION__));
+
+	if (TreeHasBeenStarted())
+	{
+		UE_VLOG(GetOwner(), LogBehaviorTree, Log, TEXT("%s: Skipping, logic already started."), ANSI_TO_TCHAR(__FUNCTION__));
+		return;
+	}
+
+	if (TreeStartInfo.IsSet() == false)
+	{
+		TreeStartInfo.Asset = DefaultBehaviorTreeAsset;
+	}
+
+	if (TreeStartInfo.IsSet())
+	{
+		TreeStartInfo.bPendingInitialize = true;
+		ProcessPendingInitialize();
+	}
+	else
+	{
+		UE_VLOG(GetOwner(), LogBehaviorTree, Log, TEXT("%s: Could not find BehaviorTree asset to run."), ANSI_TO_TCHAR(__FUNCTION__));
+	}
+}
+
 void UBehaviorTreeComponent::RestartLogic()
 {
-	UE_VLOG(GetOwner(), LogBehaviorTree, Log, TEXT("UBehaviorTreeComponent::RestartLogic"));
+	UE_VLOG(GetOwner(), LogBehaviorTree, Log, TEXT("%s"), ANSI_TO_TCHAR(__FUNCTION__));
 	RestartTree();
 }
 
@@ -375,7 +401,7 @@ void UBehaviorTreeComponent::StopTree(EBTStopMode::Type StopMode)
 
 void UBehaviorTreeComponent::RestartTree()
 {
-	UE_VLOG(GetOwner(), LogBehaviorTree, Log, TEXT("UBehaviorTreeComponent::RestartTree"));
+	UE_VLOG(GetOwner(), LogBehaviorTree, Log, TEXT("%s"), ANSI_TO_TCHAR(__FUNCTION__));
 	
 	if (!bIsRunning)
 	{
@@ -665,7 +691,7 @@ void UBehaviorTreeComponent::RequestExecution(const UBTDecorator* RequestedBy)
 	if (ActiveInstance.ActiveNodeType == EBTActiveNode::ActiveTask)
 	{
 		EBTNodeRelativePriority RelativePriority = CalculateRelativePriority(RequestedBy, ActiveInstance.ActiveNode);
-		UE_CVLOG(RelativePriority < EBTNodeRelativePriority::Same, GetOwner(), LogBehaviorTree, Error, TEXT("UBehaviorTreeComponent::RequestExecution decorator requesting restart has lower priority than Current Task"));
+		UE_CVLOG(RelativePriority < EBTNodeRelativePriority::Same, GetOwner(), LogBehaviorTree, Error, TEXT("%s decorator requesting restart has lower priority than Current Task"), ANSI_TO_TCHAR(__FUNCTION__));
 		ensure(RelativePriority >= EBTNodeRelativePriority::Same);
 	}
 
diff --git a/Engine/Source/Runtime/AIModule/Private/BehaviorTree/BlackboardComponent.cpp b/Engine/Source/Runtime/AIModule/Private/BehaviorTree/BlackboardComponent.cpp
--- a/Engine/Source/Runtime/AIModule/Private/BehaviorTree/BlackboardComponent.cpp
+++ b/Engine/Source/Runtime/AIModule/Private/BehaviorTree/BlackboardComponent.cpp
@@ -41,9 +41,31 @@ void UBlackboardComponent::InitializeComponent()
 		}
 	}
 
+	// Use DefaultBlackboardAsset if available.
+	if (DefaultBlackboardAsset)
+	{
+		if (BlackboardAsset == nullptr)
+		{
+			BlackboardAsset = DefaultBlackboardAsset;
+		}
+		else if (BlackboardAsset != DefaultBlackboardAsset)
+		{
+			UE_LOG(LogBehaviorTree, Warning, TEXT("Ignoring DefaultBlackboardAsset \'%s\', Blackboard asset already set to \'%s\'."),
+				*GetNameSafe(DefaultBlackboardAsset), *GetNameSafe(BlackboardAsset));
+		}
+	}
+
 	if (BlackboardAsset)
 	{
-		InitializeBlackboard(*BlackboardAsset);
+		// Here we expect that BlackboardAsset has been set via DefaultBlackboardAsset above
+		// or somewhere in derived class prior to call this function,
+		// and we expect that the blackboard is not yet initialized.
+		// The trick below ensures that the blackboard gets initialized,
+		// as the function InitializeBlackboard() treats BlackboardAsset
+		// as previously initialized blackboard and early outs if they are the same.
+		UBlackboardData* NewBlackboardAsset = BlackboardAsset;
+		BlackboardAsset = nullptr;
+		InitializeBlackboard(*NewBlackboardAsset);
 	}
 }
 
@@ -105,7 +127,7 @@ bool UBlackboardComponent::InitializeBlackboard(UBlackboardData& NewAsset)
 	// in reseting, since we'd lose all the accumulated knowledge
 	if (&NewAsset == BlackboardAsset)
 	{
-		return false;
+		return true;
 	}
 
 	UAISystem* AISystem = UAISystem::GetCurrentSafe(GetWorld());
diff --git a/Engine/Source/Runtime/AIModule/Private/BrainComponent.cpp b/Engine/Source/Runtime/AIModule/Private/BrainComponent.cpp
--- a/Engine/Source/Runtime/AIModule/Private/BrainComponent.cpp
+++ b/Engine/Source/Runtime/AIModule/Private/BrainComponent.cpp
@@ -328,6 +328,11 @@ EAILogicResuming::Type UBrainComponent::ResumeLogic(const FString& Reason)
 	return EAILogicResuming::Continue;
 }
 
+void UBrainComponent::StartLogic()
+{
+
+}
+
 void UBrainComponent::RestartLogic()
 {
 
