Export LogAIPerception from AIMODULE_API so it can be used by game projects

#rnx

diff --git a/Engine/Source/Runtime/AIModule/Classes/Perception/AIPerceptionSystem.h b/Engine/Source/Runtime/AIModule/Classes/Perception/AIPerceptionSystem.h
--- a/Engine/Source/Runtime/AIModule/Classes/Perception/AIPerceptionSystem.h
+++ b/Engine/Source/Runtime/AIModule/Classes/Perception/AIPerceptionSystem.h
@@ -18,7 +18,7 @@
 class UAIPerceptionComponent;
 class UAISenseEvent;
 
-DECLARE_LOG_CATEGORY_EXTERN(LogAIPerception, Warning, All);
+AIMODULE_API DECLARE_LOG_CATEGORY_EXTERN(LogAIPerception, Warning, All);
 
 class APawn;
 
