Setting default value for UProperty integer.
This is an attempt to remove some deterministic cook warnings.

diff --git a/Engine/Source/Runtime/AIModule/Private/Perception/AIPerceptionStimuliSourceComponent.cpp b/Engine/Source/Runtime/AIModule/Private/Perception/AIPerceptionStimuliSourceComponent.cpp
--- a/Engine/Source/Runtime/AIModule/Private/Perception/AIPerceptionStimuliSourceComponent.cpp
+++ b/Engine/Source/Runtime/AIModule/Private/Perception/AIPerceptionStimuliSourceComponent.cpp
@@ -5,6 +5,7 @@
 
 UAIPerceptionStimuliSourceComponent::UAIPerceptionStimuliSourceComponent(const FObjectInitializer& ObjectInitializer)
 	: Super(ObjectInitializer)
+	, bAutoRegisterAsSource(false)
 {
 	bSuccessfullyRegistered = false;
 }
