PerceptionSystem: UnregisterSource is now called directly from OnPerceptionStimuliSourceEndPlay with the source actor. We no longer need to queue a cleanup for invalid sources for next tick so all related logic has been removed.
This also fixes a missing SensingFailed stimulus from AISense_Sight.

diff --git a/Engine/Source/Runtime/AIModule/Classes/Perception/AIPerceptionSystem.h b/Engine/Source/Runtime/AIModule/Classes/Perception/AIPerceptionSystem.h
--- a/Engine/Source/Runtime/AIModule/Classes/Perception/AIPerceptionSystem.h
+++ b/Engine/Source/Runtime/AIModule/Classes/Perception/AIPerceptionSystem.h
@@ -64,8 +64,10 @@ protected:
 	 *	bSomeListenersNeedUpdateDueToStimuliAging gets reset to false */
 	uint32 bSomeListenersNeedUpdateDueToStimuliAging : 1;
 
-	/** gets set to true when perception system gets notified about a stimuli source's end play */
+#if WITH_EDITORONLY_DATA
+	UE_DEPRECATED(4.25, "This property will be removed in future versions. UnregisterSource is called by AActor.OnEndPlay delegate and will perform the cleanup.")
 	uint32 bStimuliSourcesRefreshRequired : 1;
+#endif
 
 	uint32 bHandlePawnNotification : 1;
 
diff --git a/Engine/Source/Runtime/AIModule/Classes/Perception/AISense.h b/Engine/Source/Runtime/AIModule/Classes/Perception/AISense.h
--- a/Engine/Source/Runtime/AIModule/Classes/Perception/AISense.h
+++ b/Engine/Source/Runtime/AIModule/Classes/Perception/AISense.h
@@ -102,9 +102,10 @@ public:
 	//virtual void RegisterSources(TArray<AActor&> SourceActors) {}
 	virtual void RegisterSource(AActor& SourceActors){}
 	virtual void UnregisterSource(AActor& SourceActors){}
-	// @note this function should not be needed once AActor.OnEndPlay broadcast includes instigator as one of its params
-	// since implementations on this function would end up being more expensive then precise UnregisterSource calls
+
+	UE_DEPRECATED(4.25, "This method is no longer used and will be removed in future versions. UnregisterSource is called by AActor.OnEndPlay delegate and will perform the cleanup.")
 	virtual void CleanseInvalidSources() {}
+
 	virtual void RegisterWrappedEvent(UAISenseEvent& PerceptionEvent);
 	virtual FAISenseID UpdateSenseID();
 
diff --git a/Engine/Source/Runtime/AIModule/Classes/Perception/AISense_Sight.h b/Engine/Source/Runtime/AIModule/Classes/Perception/AISense_Sight.h
--- a/Engine/Source/Runtime/AIModule/Classes/Perception/AISense_Sight.h
+++ b/Engine/Source/Runtime/AIModule/Classes/Perception/AISense_Sight.h
@@ -164,7 +164,6 @@ public:
 
 	virtual void RegisterSource(AActor& SourceActors) override;
 	virtual void UnregisterSource(AActor& SourceActor) override;
-	virtual void CleanseInvalidSources() override;
 	
 	virtual void OnListenerForgetsActor(const FPerceptionListener& Listener, AActor& ActorToForget) override;
 	virtual void OnListenerForgetsAll(const FPerceptionListener& Listener) override;
diff --git a/Engine/Source/Runtime/AIModule/Private/Perception/AIPerceptionSystem.cpp b/Engine/Source/Runtime/AIModule/Private/Perception/AIPerceptionSystem.cpp
--- a/Engine/Source/Runtime/AIModule/Private/Perception/AIPerceptionSystem.cpp
+++ b/Engine/Source/Runtime/AIModule/Private/Perception/AIPerceptionSystem.cpp
@@ -33,10 +33,15 @@ UAIPerceptionSystem::UAIPerceptionSystem(const FObjectInitializer& ObjectInitial
 	: Super(ObjectInitializer)
 	, PerceptionAgingRate(0.3f)
 	, bSomeListenersNeedUpdateDueToStimuliAging(false)
-	, bStimuliSourcesRefreshRequired(false)
 	, bHandlePawnNotification(false)
 	, CurrentTime(0.f)
 {
+PRAGMA_DISABLE_DEPRECATION_WARNINGS
+#if WITH_EDITORONLY_DATA
+	bStimuliSourcesRefreshRequired = false;
+#endif
+PRAGMA_ENABLE_DEPRECATION_WARNINGS
+
 	StimuliSourceEndPlayDelegate.BindDynamic(this, &UAIPerceptionSystem::OnPerceptionStimuliSourceEndPlay);
 }
 
@@ -165,42 +170,13 @@ void UAIPerceptionSystem::Tick(float DeltaSeconds)
 	{
 		// cache it
 		CurrentTime = World->GetTimeSeconds();
-
-		bool bNeedsUpdate = false;
-
-		if (bStimuliSourcesRefreshRequired == true)
-		{
-			// this is a bit naive, but we don't get notifications from the engine which Actor ended play, 
-			// we just know one did. Let's just refresh the map removing all no-longer-valid instances
-			// @todo: we should be just calling senses directly here to unregister specific source
-			// but at this point source actor is invalid already so we can't really pin and use it
-			// this will go away once OnEndPlay broadcast passes in its instigator
-			FPerceptionChannelWhitelist SensesToUpdate;
-			for (auto It = RegisteredStimuliSources.CreateIterator(); It; ++It)
-			{
-				if (It->Value.SourceActor.IsValid() == false)
-				{
-					SensesToUpdate.MergeFilterIn(It->Value.RelevantSenses);
-					It.RemoveCurrent();
-				}
-			}
-
-			for (FPerceptionChannelWhitelist::FConstIterator It(SensesToUpdate); It && Senses.IsValidIndex(*It); ++It)
-			{
-				if (Senses[*It] != nullptr)
-				{
-					Senses[*It]->CleanseInvalidSources();
-				}
-			}
-
-			bStimuliSourcesRefreshRequired = false;
-		}
 		
 		if (SourcesToRegister.Num() > 0)
 		{
 			PerformSourceRegistration();
 		}
 
+		bool bNeedsUpdate = false;
 		for (UAISense* const SenseInstance : Senses)
 		{
 			bNeedsUpdate |= SenseInstance != nullptr && SenseInstance->ProgressTime(DeltaSeconds);
@@ -625,10 +601,7 @@ void UAIPerceptionSystem::RegisterSourceForSenseClass(TSubclassOf<UAISense> Sens
 
 void UAIPerceptionSystem::OnPerceptionStimuliSourceEndPlay(AActor* Actor, EEndPlayReason::Type EndPlayReason)
 {
-	// this tells us just that _a_ source has been removed. We need to parse through all sources and find which one was it
-	// this is a fall-back behavior, if source gets unregistered manually this function won't get called 
-	// note, the actual work will be done on system's tick
-	bStimuliSourcesRefreshRequired = true;
+	UnregisterSource(*Actor);
 }
 
 TSubclassOf<UAISense> UAIPerceptionSystem::GetSenseClassForStimulus(UObject* WorldContextObject, const FAIStimulus& Stimulus)
diff --git a/Engine/Source/Runtime/AIModule/Private/Perception/AISense_Sight.cpp b/Engine/Source/Runtime/AIModule/Private/Perception/AISense_Sight.cpp
--- a/Engine/Source/Runtime/AIModule/Private/Perception/AISense_Sight.cpp
+++ b/Engine/Source/Runtime/AIModule/Private/Perception/AISense_Sight.cpp
@@ -392,38 +392,6 @@ void UAISense_Sight::UnregisterSource(AActor& SourceActor)
 	}
 }
 
-void UAISense_Sight::CleanseInvalidSources()
-{
-	bool bInvalidSourcesFound = false;
-	int32 NumInvalidSourcesFound = 0;
-	for (FTargetsContainer::TIterator ItTarget(ObservedTargets); ItTarget; ++ItTarget)
-	{
-		if (ItTarget->Value.Target.IsValid() == false)
-		{
-			// remove affected queries
-			RemoveAllQueriesToTarget(ItTarget->Key, DontSort);
-			// remove target itself
-			ItTarget.RemoveCurrent();
-
-			bInvalidSourcesFound = true;
-			NumInvalidSourcesFound++;
-		}
-	}
-
-	UE_LOG(LogAIPerception, Verbose, TEXT("UAISense_Sight::CleanseInvalidSources called and removed %d invalid sources"), NumInvalidSourcesFound);
-
-	if (bInvalidSourcesFound)
-	{
-		// remove holes
-		ObservedTargets.Compact();
-		SortQueries();
-	}
-	else
-	{
-		UE_VLOG(GetPerceptionSystem(), LogAIPerception, Error, TEXT("UAISense_Sight::CleanseInvalidSources called and no invalid targets were found"));
-	}
-}
-
 bool UAISense_Sight::RegisterTarget(AActor& TargetActor, FQueriesOperationPostProcess PostProcess)
 {
 	return RegisterTarget(TargetActor, PostProcess, [](FAISightQuery& Query) {});
