PR #6572: Provide a project option to tell the AI system to forget actors that have stale perception (No longer perceived and stimuli past max age) (Contributed by [REDACTED_PERSON])

Minor modifications to PR:
- ProcessStimuli uses list of AActor to forget instead of FActorPerceptionInfo to prevent modifications (now or in the future) to PerceptualData from affecting the pointers in the pending list.
- Tweaked some UProperty specifiers

diff --git a/Engine/Source/Runtime/AIModule/Classes/AISystem.h b/Engine/Source/Runtime/AIModule/Classes/AISystem.h
--- a/Engine/Source/Runtime/AIModule/Classes/AISystem.h
+++ b/Engine/Source/Runtime/AIModule/Classes/AISystem.h
@@ -73,7 +73,13 @@ public:
 	UPROPERTY(globalconfig, EditDefaultsOnly, Category = "AISystem")
 	bool bEnableDebuggerPlugin;
 
-	UPROPERTY(globalconfig, EditAnywhere, Category = "PerceptionSystem")
+	/** If set, actors will be forgotten by the perception system when their stimulus has expired.
+	 *	If not set, the perception system will remember the actor even if they are no longer perceived and their
+	 *	stimuli has exceeded its max age */
+	UPROPERTY(globalconfig, EditDefaultsOnly, Category = "AISystem")
+	bool bForgetStaleActors;
+
+	UPROPERTY(globalconfig, EditDefaultsOnly, Category = "PerceptionSystem")
 	TEnumAsByte<ECollisionChannel> DefaultSightCollisionChannel;
 
 protected:
diff --git a/Engine/Source/Runtime/AIModule/Classes/Perception/AIPerceptionComponent.h b/Engine/Source/Runtime/AIModule/Classes/Perception/AIPerceptionComponent.h
--- a/Engine/Source/Runtime/AIModule/Classes/Perception/AIPerceptionComponent.h
+++ b/Engine/Source/Runtime/AIModule/Classes/Perception/AIPerceptionComponent.h
@@ -238,6 +238,11 @@ protected:
 	TArray<float> MaxActiveAge;
 
 private:
+
+	/** Determines whether all knowledge of previously sensed actors will be removed or not when they become stale.
+		That is, when they are no longer perceived and have exceeded the max age of the sense. */
+	uint32 bForgetStaleActors : 1;
+
 	uint32 bCleanedUp : 1;
 
 public:
diff --git a/Engine/Source/Runtime/AIModule/Private/Perception/AIPerceptionComponent.cpp b/Engine/Source/Runtime/AIModule/Private/Perception/AIPerceptionComponent.cpp
--- a/Engine/Source/Runtime/AIModule/Private/Perception/AIPerceptionComponent.cpp
+++ b/Engine/Source/Runtime/AIModule/Private/Perception/AIPerceptionComponent.cpp
@@ -53,6 +53,7 @@ UAIPerceptionComponent::UAIPerceptionComponent(const FObjectInitializer& ObjectI
 	, PerceptionListenerId(FPerceptionListenerID::InvalidID())
 	, bCleanedUp(false)
 {
+	bForgetStaleActors = GET_AI_CONFIG_VAR(bForgetStaleActors);
 }
 
 void UAIPerceptionComponent::RequestStimuliListenerUpdate()
@@ -415,6 +416,8 @@ void UAIPerceptionComponent::ProcessStimuli()
 	
 	TArray<AActor*> UpdatedActors;
 	UpdatedActors.Reserve(StimuliToProcess.Num());
+	TArray<AActor*> ActorsToForget;
+	ActorsToForget.Reserve(StimuliToProcess.Num());
 
 	for (FStimulusToProcess& SourcedStimulus : StimuliToProcess)
 	{
@@ -470,6 +473,14 @@ void UAIPerceptionComponent::ProcessStimuli()
 		else
 		{
 			HandleExpiredStimulus(StimulusStore);
+
+			if (bForgetStaleActors && !PerceptualInfo->HasAnyCurrentStimulus())
+			{
+				if (AActor* ActorToForget = PerceptualInfo->Target.Get())
+				{
+					ActorsToForget.Add(ActorToForget);
+				}
+			}
 		}
 
 		// if the new stimulus is "valid" or it's info that "no longer sensed" and it used to be sensed successfully
@@ -494,6 +505,11 @@ void UAIPerceptionComponent::ProcessStimuli()
 
 		OnPerceptionUpdated.Broadcast(UpdatedActors);
 	}
+
+	for (AActor* ActorToForget : ActorsToForget)
+	{
+		ForgetActor(ActorToForget);
+	}
 }
 
 void UAIPerceptionComponent::RefreshStimulus(FAIStimulus& StimulusStore, const FAIStimulus& NewStimulus)
