Perception rework & cleanup
Change the calculation of the age of the query
Simplified/deprecated some APIs with TFunctions
Simplified sorting by always doing it at the beginning of the update, help optimizing the removal

#rnx

diff --git a/Engine/Source/Runtime/AIModule/Classes/Perception/AISense_Sight.h b/Engine/Source/Runtime/AIModule/Classes/Perception/AISense_Sight.h
--- a/Engine/Source/Runtime/AIModule/Classes/Perception/AISense_Sight.h
+++ b/Engine/Source/Runtime/AIModule/Classes/Perception/AISense_Sight.h
@@ -72,22 +72,32 @@ struct FAISightQuery
 	FPerceptionListenerID ObserverId;
 	FAISightTarget::FTargetId TargetId;
 
-	float Age;
 	float Score;
 	float Importance;
 
 	FVector LastSeenLocation;
 
-	uint32 bLastResult : 1;
+	uint64 bLastResult:1;
+	uint64 LastProcessedFrameNumber :63;
 
 	FAISightQuery(FPerceptionListenerID ListenerId = FPerceptionListenerID::InvalidID(), FAISightTarget::FTargetId Target = FAISightTarget::InvalidTargetId)
-		: ObserverId(ListenerId), TargetId(Target), Age(0), Score(0), Importance(0), LastSeenLocation(FAISystem::InvalidLocation), bLastResult(false)
+		: ObserverId(ListenerId), TargetId(Target), Score(0), Importance(0), LastSeenLocation(FAISystem::InvalidLocation), bLastResult(false), LastProcessedFrameNumber(GFrameCounter)
 	{
 	}
 
+	float GetAge() const
+	{
+		return (float)(GFrameCounter - LastProcessedFrameNumber);
+	}
+
 	void RecalcScore()
 	{
-		Score = Age + Importance;
+		Score = GetAge() + Importance;
+	}
+
+	void OnProcessed()
+	{
+		LastProcessedFrameNumber = GFrameCounter;
 	}
 
 	void ForgetPreviousResult()
@@ -177,24 +187,38 @@ protected:
 	void OnListenerUpdateImpl(const FPerceptionListener& UpdatedListener);
 	void OnListenerRemovedImpl(const FPerceptionListener& UpdatedListener);	
 	
-	void GenerateQueriesForListener(const FPerceptionListener& Listener, const FDigestedSightProperties& PropertyDigest);
-	void GenerateQueriesForListener(const FPerceptionListener& Listener, const FDigestedSightProperties& PropertyDigest, TFunctionRef<void(FAISightQuery&)> OnAddedFunc);
+	void GenerateQueriesForListener(const FPerceptionListener& Listener, const FDigestedSightProperties& PropertyDigest, const TFunction<void(FAISightQuery&)>& OnAddedFunc = nullptr);
+
+	void RemoveAllQueriesByListener(const FPerceptionListener& Listener, const TFunction<void(const FAISightQuery&)>& OnRemoveFunc = nullptr);
+	void RemoveAllQueriesToTarget(const FAISightTarget::FTargetId& TargetId, const TFunction<void(const FAISightQuery&)>& OnRemoveFunc = nullptr);
+
+	/** returns information whether new LoS queries have been added */
+	bool RegisterTarget(AActor& TargetActor, const TFunction<void(FAISightQuery&)>& OnAddedFunc = nullptr);
 
+	float CalcQueryImportance(const FPerceptionListener& Listener, const FVector& TargetLocation, const float SightRadiusSq) const;
+
+	// Deprecated methods
+public:
+	UE_DEPRECATED(4.25, "Not needed anymore done automatically at the beginning of each update.")
+	FORCEINLINE void SortQueries() {}
+
+protected:
 	enum FQueriesOperationPostProcess
 	{
 		DontSort,
 		Sort
 	};
-	void RemoveAllQueriesByListener(const FPerceptionListener& Listener, FQueriesOperationPostProcess PostProcess);
-	void RemoveAllQueriesByListener(const FPerceptionListener& Listener, FQueriesOperationPostProcess PostProcess, TFunctionRef<void(const FAISightQuery&)> OnRemoveFunc);
-	void RemoveAllQueriesToTarget(const FAISightTarget::FTargetId& TargetId, FQueriesOperationPostProcess PostProcess);
-	void RemoveAllQueriesToTarget(const FAISightTarget::FTargetId& TargetId, FQueriesOperationPostProcess PostProcess, TFunctionRef<void(const FAISightQuery&)> OnRemoveFunc);
+	UE_DEPRECATED(4.25, "Use RemoveAllQueriesByListener without unneeded PostProcess parameter.")
+	void RemoveAllQueriesByListener(const FPerceptionListener& Listener, FQueriesOperationPostProcess PostProcess) { RemoveAllQueriesByListener(Listener); }
+	UE_DEPRECATED(4.25, "Use RemoveAllQueriesByListener without unneeded PostProcess parameter.")
+	void RemoveAllQueriesByListener(const FPerceptionListener& Listener, FQueriesOperationPostProcess PostProcess, TFunctionRef<void(const FAISightQuery&)> OnRemoveFunc) { RemoveAllQueriesByListener(Listener, [&](const FAISightQuery& query) { OnRemoveFunc(query); }); }
+	UE_DEPRECATED(4.25, "Use RemoveAllQueriesToTarget without unneeded PostProcess parameter.")
+	void RemoveAllQueriesToTarget(const FAISightTarget::FTargetId& TargetId, FQueriesOperationPostProcess PostProcess) { RemoveAllQueriesToTarget(TargetId); }
+	UE_DEPRECATED(4.25, "Use RemoveAllQueriesToTarget without unneeded PostProcess parameter.")
+	void RemoveAllQueriesToTarget(const FAISightTarget::FTargetId& TargetId, FQueriesOperationPostProcess PostProcess, TFunctionRef<void(const FAISightQuery&)> OnRemoveFunc) { RemoveAllQueriesToTarget(TargetId, [&](const FAISightQuery& query) { OnRemoveFunc(query); }); }
+	UE_DEPRECATED(4.25, "Use RegisterTarget without unneeded PostProcess parameter.")
+	bool RegisterTarget(AActor& TargetActor, FQueriesOperationPostProcess PostProcess) { return RegisterTarget(TargetActor); }
+	UE_DEPRECATED(4.25, "Use RegisterTarget without unneeded PostProcess parameter.")
+	bool RegisterTarget(AActor& TargetActor, FQueriesOperationPostProcess PostProcess, TFunctionRef<void(FAISightQuery&)> OnAddedFunc) { return RegisterTarget(TargetActor, [&](FAISightQuery& query) { OnAddedFunc(query); }); }
 
-	/** returns information whether new LoS queries have been added */
-	bool RegisterTarget(AActor& TargetActor, FQueriesOperationPostProcess PostProcess);
-	bool RegisterTarget(AActor& TargetActor, FQueriesOperationPostProcess PostProcess, TFunctionRef<void(FAISightQuery&)> OnAddedFunc);
-
-	FORCEINLINE void SortQueries() { SightQueryQueue.Sort(FAISightQuery::FSortPredicate()); }
-
-	float CalcQueryImportance(const FPerceptionListener& Listener, const FVector& TargetLocation, const float SightRadiusSq) const;
 };
diff --git a/Engine/Source/Runtime/AIModule/Private/Perception/AISense_Sight.cpp b/Engine/Source/Runtime/AIModule/Private/Perception/AISense_Sight.cpp
--- a/Engine/Source/Runtime/AIModule/Private/Perception/AISense_Sight.cpp
+++ b/Engine/Source/Runtime/AIModule/Private/Perception/AISense_Sight.cpp
@@ -148,6 +148,16 @@ float UAISense_Sight::Update()
 		return SuspendNextUpdate;
 	}
 
+	// sort Sight Queries
+	{
+		SCOPE_CYCLE_COUNTER(STAT_AI_Sense_Sight_UpdateSort);
+		for (FAISightQuery& SightQuery : SightQueryQueue)
+		{
+			SightQuery.RecalcScore();
+		}
+		SightQueryQueue.Sort(FAISightQuery::FSortPredicate());
+	}
+
 	int32 TracesCount = 0;
 	int32 NumQueriesProcessed = 0;
 	double TimeSliceEnd = FPlatformTime::Seconds() + MaxTimeSlicePerTick;
@@ -284,7 +294,7 @@ float UAISense_Sight::Update()
 				SightQuery->Importance = CalcQueryImportance(Listener, TargetLocation, SightRadiusSq);
 
 				// restart query
-				SightQuery->Age = 0.f;
+				SightQuery->OnProcessed();
 			}
 			else
 			{
@@ -298,11 +308,8 @@ float UAISense_Sight::Update()
 		}
 		else
 		{
-			// age unprocessed queries so that they can advance in the queue during next sort
-			SightQuery->Age += 1.f;
+			break;
 		}
-
-		SightQuery->RecalcScore();
 	}
 #ifdef AISENSE_SIGHT_TIMESLICING_DEBUG
 	UE_LOG(LogAIPerception, VeryVerbose, TEXT("UAISense_Sight::Update processed %d sources in %f seconds [time slice limited? %d]"), NumQueriesProcessed, TimeSpent, bHitTimeSliceLimit ? 1 : 0);
@@ -326,7 +333,7 @@ float UAISense_Sight::Update()
 			for (const auto& TargetId : InvalidTargets)
 			{
 				// remove affected queries
-				RemoveAllQueriesToTarget(TargetId, DontSort);
+				RemoveAllQueriesToTarget(TargetId);
 				// remove target itself
 				ObservedTargets.Remove(TargetId);
 			}
@@ -336,12 +343,6 @@ float UAISense_Sight::Update()
 		}
 	}
 
-	// sort Sight Queries
-	{
-		SCOPE_CYCLE_COUNTER(STAT_AI_Sense_Sight_UpdateSort);
-		SortQueries();
-	}
-
 	//return SightQueryQueue.Num() > 0 ? 1.f/6 : FLT_MAX;
 	return 0.f;
 }
@@ -353,7 +354,7 @@ void UAISense_Sight::RegisterEvent(const FAISightEvent& Event)
 
 void UAISense_Sight::RegisterSource(AActor& SourceActor)
 {
-	RegisterTarget(SourceActor, Sort);
+	RegisterTarget(SourceActor);
 }
 
 void UAISense_Sight::UnregisterSource(AActor& SourceActor)
@@ -384,20 +385,14 @@ void UAISense_Sight::UnregisterSource(AActor& SourceActor)
 						Listener.RegisterStimulus(TargetActor, FAIStimulus(*this, 0.f, SightQuery->LastSeenLocation, Listener.CachedLocation, FAIStimulus::SensingFailed));
 					}
 
-					SightQueryQueue.RemoveAt(QueryIndex, 1, /*bAllowShrinking=*/false);
+					SightQueryQueue.RemoveAtSwap(QueryIndex, 1, /*bAllowShrinking=*/false);
 				}
 			}
-			// no point in sorting, we haven't change the order of other queries
 		}
 	}
 }
 
-bool UAISense_Sight::RegisterTarget(AActor& TargetActor, FQueriesOperationPostProcess PostProcess)
-{
-	return RegisterTarget(TargetActor, PostProcess, [](FAISightQuery& Query) {});
-}
-
-bool UAISense_Sight::RegisterTarget(AActor& TargetActor, FQueriesOperationPostProcess PostProcess, TFunctionRef<void(FAISightQuery&)> OnAddedFunc)
+bool UAISense_Sight::RegisterTarget(AActor& TargetActor, const TFunction<void(FAISightQuery&)>& OnAddedFunc /*= nullptr*/)
 {
 	SCOPE_CYCLE_COUNTER(STAT_AI_Sense_Sight_RegisterTarget);
 	
@@ -443,16 +438,18 @@ bool UAISense_Sight::RegisterTarget(AActor& TargetActor, FQueriesOperationPostPr
 				AddedQuery.TargetId = SightTarget->TargetId;
 				AddedQuery.Importance = CalcQueryImportance(ItListener->Value, TargetLocation, PropDigest.SightRadiusSq);
 				
-				OnAddedFunc(AddedQuery);
+				if (OnAddedFunc)
+				{
+					OnAddedFunc(AddedQuery);
+				}
 				bNewQueriesAdded = true;
 			}
 		}
 	}
 
 	// sort Sight Queries
-	if (PostProcess == Sort && bNewQueriesAdded)
+	if (bNewQueriesAdded)
 	{
-		SortQueries();
 		RequestImmediateUpdate();
 	}
 
@@ -471,12 +468,7 @@ void UAISense_Sight::OnNewListenerImpl(const FPerceptionListener& NewListener)
 	GenerateQueriesForListener(NewListener, PropertyDigest);
 }
 
-void UAISense_Sight::GenerateQueriesForListener(const FPerceptionListener& Listener, const FDigestedSightProperties& PropertyDigest)
-{
-	GenerateQueriesForListener(Listener, PropertyDigest, [](FAISightQuery& Query) {});
-}
-
-void UAISense_Sight::GenerateQueriesForListener(const FPerceptionListener& Listener, const FDigestedSightProperties& PropertyDigest, TFunctionRef<void(FAISightQuery&)> OnAddedFunc)
+void UAISense_Sight::GenerateQueriesForListener(const FPerceptionListener& Listener, const FDigestedSightProperties& PropertyDigest, const TFunction<void(FAISightQuery&)>& OnAddedFunc/*= nullptr */)
 {
 	bool bNewQueriesAdded = false;
 	const IGenericTeamAgentInterface* ListenersTeamAgent = Listener.GetTeamAgent();
@@ -499,7 +491,10 @@ void UAISense_Sight::GenerateQueriesForListener(const FPerceptionListener& Liste
 			AddedQuery.TargetId = ItTarget->Key;
 			AddedQuery.Importance = CalcQueryImportance(Listener, ItTarget->Value.GetLocationSimple(), PropertyDigest.SightRadiusSq);
 
-			OnAddedFunc(AddedQuery);
+			if (OnAddedFunc)
+			{
+				OnAddedFunc(AddedQuery);
+			}
 			bNewQueriesAdded = true;
 		}
 	}
@@ -507,7 +502,6 @@ void UAISense_Sight::GenerateQueriesForListener(const FPerceptionListener& Liste
 	// sort Sight Queries
 	if (bNewQueriesAdded)
 	{
-		SortQueries();
 		RequestImmediateUpdate();
 	}
 }
@@ -529,7 +523,7 @@ void UAISense_Sight::OnListenerUpdateImpl(const FPerceptionListener& UpdatedList
 		{
 			// if still a valid target then backup list of observers for which the listener was visible to restore in the newly created queries
 			TSet<FPerceptionListenerID> LastVisibleObservers;
-			RemoveAllQueriesToTarget(AsTargetId, DontSort, [&LastVisibleObservers](const FAISightQuery& Query)
+			RemoveAllQueriesToTarget(AsTargetId, [&LastVisibleObservers](const FAISightQuery& Query)
 			{
 				if (Query.bLastResult)
 				{
@@ -537,14 +531,14 @@ void UAISense_Sight::OnListenerUpdateImpl(const FPerceptionListener& UpdatedList
 				}
 			});
 
-			RegisterTarget(*(AsTarget->Target.Get()), DontSort, [&LastVisibleObservers](FAISightQuery& Query)
+			RegisterTarget(*(AsTarget->Target.Get()), [&LastVisibleObservers](FAISightQuery& Query)
 			{
 				Query.bLastResult = LastVisibleObservers.Contains(Query.ObserverId);
 			});
 		}
 		else
 		{
-			RemoveAllQueriesToTarget(AsTargetId, DontSort);
+			RemoveAllQueriesToTarget(AsTargetId);
 		}
 	}
 
@@ -554,7 +548,7 @@ void UAISense_Sight::OnListenerUpdateImpl(const FPerceptionListener& UpdatedList
 	{
 		// if still a valid sense then backup list of targets that were visible by the listener to restore in the newly created queries
 		TSet<FAISightTarget::FTargetId> LastVisibleTargets;
-		RemoveAllQueriesByListener(UpdatedListener, DontSort, [&LastVisibleTargets](const FAISightQuery& Query)
+		RemoveAllQueriesByListener(UpdatedListener, [&LastVisibleTargets](const FAISightQuery& Query)
 		{
 			if (Query.bLastResult)
 			{
@@ -575,7 +569,7 @@ void UAISense_Sight::OnListenerUpdateImpl(const FPerceptionListener& UpdatedList
 	else
 	{
 		// remove all queries
-		RemoveAllQueriesByListener(UpdatedListener, DontSort);
+		RemoveAllQueriesByListener(UpdatedListener);
 
 		DigestedProperties.Remove(ListenerID);
 	}
@@ -583,7 +577,7 @@ void UAISense_Sight::OnListenerUpdateImpl(const FPerceptionListener& UpdatedList
 
 void UAISense_Sight::OnListenerRemovedImpl(const FPerceptionListener& UpdatedListener)
 {
-	RemoveAllQueriesByListener(UpdatedListener, DontSort);
+	RemoveAllQueriesByListener(UpdatedListener);
 
 	DigestedProperties.FindAndRemoveChecked(UpdatedListener.GetListenerID());
 
@@ -592,12 +586,7 @@ void UAISense_Sight::OnListenerRemovedImpl(const FPerceptionListener& UpdatedLis
 	// mean it's being removed from the game altogether.
 }
 
-void UAISense_Sight::RemoveAllQueriesByListener(const FPerceptionListener& Listener, FQueriesOperationPostProcess PostProcess)
-{
-	RemoveAllQueriesByListener(Listener, PostProcess, [](const FAISightQuery& Query) {});
-}
-
-void UAISense_Sight::RemoveAllQueriesByListener(const FPerceptionListener& Listener, FQueriesOperationPostProcess PostProcess, TFunctionRef<void(const FAISightQuery&)> OnRemoveFunc)
+void UAISense_Sight::RemoveAllQueriesByListener(const FPerceptionListener& Listener, const TFunction<void(const FAISightQuery&)>& OnRemoveFunc/*= nullptr */)
 {
 	SCOPE_CYCLE_COUNTER(STAT_AI_Sense_Sight_RemoveByListener);
 
@@ -607,7 +596,6 @@ void UAISense_Sight::RemoveAllQueriesByListener(const FPerceptionListener& Liste
 	}
 
 	const uint32 ListenerId = Listener.GetListenerID();
-	bool bQueriesRemoved = false;
 	
 	for (int32 QueryIndex = SightQueryQueue.Num() - 1; QueryIndex >= 0 ; --QueryIndex)
 	{
@@ -615,24 +603,16 @@ void UAISense_Sight::RemoveAllQueriesByListener(const FPerceptionListener& Liste
 
 		if (SightQuery.ObserverId == ListenerId)
 		{
-			OnRemoveFunc(SightQuery);
-			SightQueryQueue.RemoveAt(QueryIndex, 1, /*bAllowShrinking=*/false);
-			bQueriesRemoved = true;
+			if (OnRemoveFunc)
+			{
+				OnRemoveFunc(SightQuery);
+			}
+			SightQueryQueue.RemoveAtSwap(QueryIndex, 1, /*bAllowShrinking=*/false);
 		}
 	}
-
-	if (PostProcess == Sort && bQueriesRemoved)
-	{
-		SortQueries();
-	}
-}
-
-void UAISense_Sight::RemoveAllQueriesToTarget(const FAISightTarget::FTargetId& TargetId, FQueriesOperationPostProcess PostProcess)
-{
-	RemoveAllQueriesToTarget(TargetId, PostProcess, [](const FAISightQuery& Query) {});
 }
 
-void UAISense_Sight::RemoveAllQueriesToTarget(const FAISightTarget::FTargetId& TargetId, FQueriesOperationPostProcess PostProcess, TFunctionRef<void(const FAISightQuery&)> OnRemoveFunc)
+void UAISense_Sight::RemoveAllQueriesToTarget(const FAISightTarget::FTargetId& TargetId, const TFunction<void(const FAISightQuery&)>& OnRemoveFunc/*= nullptr */)
 {
 	SCOPE_CYCLE_COUNTER(STAT_AI_Sense_Sight_RemoveToTarget);
 
@@ -641,24 +621,19 @@ void UAISense_Sight::RemoveAllQueriesToTarget(const FAISightTarget::FTargetId& T
 		return;
 	}
 
-	bool bQueriesRemoved = false;
-
 	for (int32 QueryIndex = SightQueryQueue.Num() - 1; QueryIndex >= 0; --QueryIndex)
 	{
 		const FAISightQuery& SightQuery = SightQueryQueue[QueryIndex];
 
 		if (SightQuery.TargetId == TargetId)
 		{
-			OnRemoveFunc(SightQuery);
-			SightQueryQueue.RemoveAt(QueryIndex, 1, /*bAllowShrinking=*/false);
-			bQueriesRemoved = true;
+			if (OnRemoveFunc)
+			{
+				OnRemoveFunc(SightQuery);
+			}
+			SightQueryQueue.RemoveAtSwap(QueryIndex, 1, /*bAllowShrinking=*/false);
 		}
 	}
-
-	if (PostProcess == Sort && bQueriesRemoved)
-	{
-		SortQueries();
-	}
 }
 
 void UAISense_Sight::OnListenerForgetsActor(const FPerceptionListener& Listener, AActor& ActorToForget)
