PR #6274: Initialize spawned AI with Owner (Contributed by [REDACTED_PERSON])

Extended the original PR with function comment and made the newly added parameter "AdvancedDisplay"

diff --git a/Engine/Source/Runtime/AIModule/Classes/Blueprint/AIBlueprintHelperLibrary.h b/Engine/Source/Runtime/AIModule/Classes/Blueprint/AIBlueprintHelperLibrary.h
--- a/Engine/Source/Runtime/AIModule/Classes/Blueprint/AIBlueprintHelperLibrary.h
+++ b/Engine/Source/Runtime/AIModule/Classes/Blueprint/AIBlueprintHelperLibrary.h
@@ -32,8 +32,16 @@ class AIMODULE_API UAIBlueprintHelperLibrary : public UBlueprintFunctionLibrary
 	UFUNCTION(BlueprintCallable, Category="AI", meta=(DefaultToSelf="MessageSource"))
 	static void SendAIMessage(APawn* Target, FName Message, UObject* MessageSource, bool bSuccess = true);
 
-	UFUNCTION(BlueprintCallable, Category="AI", meta=(WorldContext="WorldContextObject", UnsafeDuringActorConstruction="true"))
-	static APawn* SpawnAIFromClass(UObject* WorldContextObject, TSubclassOf<APawn> PawnClass, UBehaviorTree* BehaviorTree, FVector Location, FRotator Rotation=FRotator::ZeroRotator, bool bNoCollisionFail=false);
+	/**	Spawns AI agent of a given class. The PawnClass needs to have AIController 
+	 *	set for the function to spawn a controller as well.
+	 *	@param BehaviorTree if set, and the function has successfully spawned 
+	 *		and AI controller, this BehaviorTree asset will be assigned to the AI 
+	 *		controller, and run.
+	 *	@param Owner lets you spawn the AI in a sublevel rather than in the 
+	 *		persistent level (which is the default behavior).
+	 */
+	UFUNCTION(BlueprintCallable, Category="AI", meta=(WorldContext="WorldContextObject", UnsafeDuringActorConstruction="true", AdvancedDisplay = "Owner"))
+	static APawn* SpawnAIFromClass(UObject* WorldContextObject, TSubclassOf<APawn> PawnClass, UBehaviorTree* BehaviorTree, FVector Location, FRotator Rotation = FRotator::ZeroRotator, bool bNoCollisionFail = false, AActor* Owner = nullptr);
 
 	/** The way it works exactly is if the actor passed in is a pawn, then the function retrieves 
 	 *	pawn's controller cast to AIController. Otherwise the function returns actor cast to AIController. */
diff --git a/Engine/Source/Runtime/AIModule/Private/Blueprint/AIBlueprintHelperLibrary.cpp b/Engine/Source/Runtime/AIModule/Private/Blueprint/AIBlueprintHelperLibrary.cpp
--- a/Engine/Source/Runtime/AIModule/Private/Blueprint/AIBlueprintHelperLibrary.cpp
+++ b/Engine/Source/Runtime/AIModule/Private/Blueprint/AIBlueprintHelperLibrary.cpp
@@ -181,7 +181,7 @@ void UAIBlueprintHelperLibrary::SendAIMessage(APawn* Target, FName Message, UObj
 	FAIMessage::Send(Target, FAIMessage(Message, MessageSource, bSuccess));
 }
 
-APawn* UAIBlueprintHelperLibrary::SpawnAIFromClass(UObject* WorldContextObject, TSubclassOf<APawn> PawnClass, UBehaviorTree* BehaviorTree, FVector Location, FRotator Rotation, bool bNoCollisionFail)
+APawn* UAIBlueprintHelperLibrary::SpawnAIFromClass(UObject* WorldContextObject, TSubclassOf<APawn> PawnClass, UBehaviorTree* BehaviorTree, FVector Location, FRotator Rotation, bool bNoCollisionFail, AActor *Owner)
 {
 	APawn* NewPawn = NULL;
 
@@ -189,6 +189,7 @@ APawn* UAIBlueprintHelperLibrary::SpawnAIFromClass(UObject* WorldContextObject,
 	if (World && *PawnClass)
 	{
 		FActorSpawnParameters ActorSpawnParams;
+		ActorSpawnParams.Owner = Owner;
 		ActorSpawnParams.SpawnCollisionHandlingOverride = bNoCollisionFail ? ESpawnActorCollisionHandlingMethod::AlwaysSpawn : ESpawnActorCollisionHandlingMethod::AdjustIfPossibleButDontSpawnIfColliding;
 
 		NewPawn = World->SpawnActor<APawn>(*PawnClass, Location, Rotation, ActorSpawnParams);
