PR #6290: Changed OnQueryFinishedEvent in UEnvQueryInstanceBlueprintWrapper from protected to public to allow C++ access. (Contributed by [REDACTED_PERSON])

Updated the original PR to use an inline getter function rather than directly making the member property public

diff --git a/Engine/Source/Runtime/AIModule/Classes/EnvironmentQuery/EnvQueryInstanceBlueprintWrapper.h b/Engine/Source/Runtime/AIModule/Classes/EnvironmentQuery/EnvQueryInstanceBlueprintWrapper.h
--- a/Engine/Source/Runtime/AIModule/Classes/EnvironmentQuery/EnvQueryInstanceBlueprintWrapper.h
+++ b/Engine/Source/Runtime/AIModule/Classes/EnvironmentQuery/EnvQueryInstanceBlueprintWrapper.h
@@ -77,6 +77,8 @@ public:
 
 	void SetInstigator(const UObject* Object);
 
+	FEQSQueryDoneSignature& GetOnQueryFinishedEvent() { return OnQueryFinishedEvent; }
+
 protected:
 	void OnQueryFinished(TSharedPtr<FEnvQueryResult> Result);
 
