Added a BP-configurable display name to BP-implemented BT nodes

diff --git a/Engine/Source/Runtime/AIModule/Classes/BehaviorTree/Decorators/BTDecorator_BlueprintBase.h b/Engine/Source/Runtime/AIModule/Classes/BehaviorTree/Decorators/BTDecorator_BlueprintBase.h
--- a/Engine/Source/Runtime/AIModule/Classes/BehaviorTree/Decorators/BTDecorator_BlueprintBase.h
+++ b/Engine/Source/Runtime/AIModule/Classes/BehaviorTree/Decorators/BTDecorator_BlueprintBase.h
@@ -79,6 +79,11 @@ protected:
 	/** properties with runtime values, stored only in class default object */
 	TArray<UProperty*> PropertyData;
 
+#if WITH_EDITORONLY_DATA
+	UPROPERTY(EditDefaultsOnly, BlueprintReadWrite, Category = Description)
+	FString CustomDescription;
+#endif // WITH_EDITORONLY_DATA
+
 	/** show detailed information about properties */
 	UPROPERTY(EditInstanceOnly, Category=Description)
 	uint32 bShowPropertyDetails : 1;
diff --git a/Engine/Source/Runtime/AIModule/Classes/BehaviorTree/Services/BTService_BlueprintBase.h b/Engine/Source/Runtime/AIModule/Classes/BehaviorTree/Services/BTService_BlueprintBase.h
--- a/Engine/Source/Runtime/AIModule/Classes/BehaviorTree/Services/BTService_BlueprintBase.h
+++ b/Engine/Source/Runtime/AIModule/Classes/BehaviorTree/Services/BTService_BlueprintBase.h
@@ -50,6 +50,11 @@ protected:
 	/** properties with runtime values, stored only in class default object */
 	TArray<UProperty*> PropertyData;
 
+#if WITH_EDITORONLY_DATA
+	UPROPERTY(EditDefaultsOnly, BlueprintReadWrite, Category = Description)
+	FString CustomDescription;
+#endif // WITH_EDITORONLY_DATA
+
 	/** show detailed information about properties */
 	UPROPERTY(EditInstanceOnly, Category=Description)
 	uint32 bShowPropertyDetails : 1;
diff --git a/Engine/Source/Runtime/AIModule/Classes/BehaviorTree/Tasks/BTTask_BlueprintBase.h b/Engine/Source/Runtime/AIModule/Classes/BehaviorTree/Tasks/BTTask_BlueprintBase.h
--- a/Engine/Source/Runtime/AIModule/Classes/BehaviorTree/Tasks/BTTask_BlueprintBase.h
+++ b/Engine/Source/Runtime/AIModule/Classes/BehaviorTree/Tasks/BTTask_BlueprintBase.h
@@ -60,6 +60,11 @@ protected:
 	/** properties that should be copied */
 	TArray<UProperty*> PropertyData;
 
+#if WITH_EDITORONLY_DATA
+	UPROPERTY(EditDefaultsOnly, BlueprintReadWrite, Category = Description)
+	FString CustomDescription;
+#endif // WITH_EDITORONLY_DATA
+
 	/** show detailed information about properties */
 	UPROPERTY(EditInstanceOnly, Category=Description)
 	uint32 bShowPropertyDetails : 1;
diff --git a/Engine/Source/Runtime/AIModule/Private/BehaviorTree/Decorators/BTDecorator_BlueprintBase.cpp b/Engine/Source/Runtime/AIModule/Private/BehaviorTree/Decorators/BTDecorator_BlueprintBase.cpp
--- a/Engine/Source/Runtime/AIModule/Private/BehaviorTree/Decorators/BTDecorator_BlueprintBase.cpp
+++ b/Engine/Source/Runtime/AIModule/Private/BehaviorTree/Decorators/BTDecorator_BlueprintBase.cpp
@@ -244,7 +244,11 @@ bool UBTDecorator_BlueprintBase::IsDecoratorObserverActive() const
 
 FString UBTDecorator_BlueprintBase::GetStaticDescription() const
 {
-	FString ReturnDesc = Super::GetStaticDescription();
+	FString ReturnDesc =
+#if WITH_EDITORONLY_DATA
+		CustomDescription.Len() ? CustomDescription :
+#endif // WITH_EDITORONLY_DATA
+		Super::GetStaticDescription();
 
 	UBTDecorator_BlueprintBase* CDO = (UBTDecorator_BlueprintBase*)(GetClass()->GetDefaultObject());
 	if (bShowPropertyDetails && CDO)
diff --git a/Engine/Source/Runtime/AIModule/Private/BehaviorTree/Services/BTService_BlueprintBase.cpp b/Engine/Source/Runtime/AIModule/Private/BehaviorTree/Services/BTService_BlueprintBase.cpp
--- a/Engine/Source/Runtime/AIModule/Private/BehaviorTree/Services/BTService_BlueprintBase.cpp
+++ b/Engine/Source/Runtime/AIModule/Private/BehaviorTree/Services/BTService_BlueprintBase.cpp
@@ -128,14 +128,18 @@ bool UBTService_BlueprintBase::IsServiceActive() const
 
 FString UBTService_BlueprintBase::GetStaticServiceDescription() const
 {
-	FString ReturnDesc;
+	FString ReturnDesc
+#if WITH_EDITORONLY_DATA
+		= CustomDescription
+#endif // WITH_EDITORONLY_DATA
+		;
 
 	UBTService_BlueprintBase* CDO = (UBTService_BlueprintBase*)(GetClass()->GetDefaultObject());
 	if (CDO)
 	{
 		if (bShowEventDetails)
 		{
-			ReturnDesc = FString::Printf(TEXT("%s, %s, %s, %s\n"),
+			ReturnDesc += FString::Printf(TEXT("%s, %s, %s, %s\n"),
 				ReceiveTickImplementations != 0 ? *GetStaticTickIntervalDescription() : TEXT("No tick"),
 				ReceiveActivationImplementations != 0 ? TEXT("Activation") : TEXT("No Activation"),
 				ReceiveDeactivationImplementations != 0 ? TEXT("Deactivation") : TEXT("No Deactivation"),
@@ -143,7 +147,7 @@ FString UBTService_BlueprintBase::GetStaticServiceDescription() const
 		}
 		else
 		{
-			ReturnDesc = Super::GetStaticServiceDescription();
+			ReturnDesc += Super::GetStaticServiceDescription();
 			ReturnDesc += TEXT('\n');
 		}
 						
diff --git a/Engine/Source/Runtime/AIModule/Private/BehaviorTree/Tasks/BTTask_BlueprintBase.cpp b/Engine/Source/Runtime/AIModule/Private/BehaviorTree/Tasks/BTTask_BlueprintBase.cpp
--- a/Engine/Source/Runtime/AIModule/Private/BehaviorTree/Tasks/BTTask_BlueprintBase.cpp
+++ b/Engine/Source/Runtime/AIModule/Private/BehaviorTree/Tasks/BTTask_BlueprintBase.cpp
@@ -190,7 +190,11 @@ void UBTTask_BlueprintBase::SetFinishOnMessageWithId(FName MessageName, int32 Re
 
 FString UBTTask_BlueprintBase::GetStaticDescription() const
 {
-	FString ReturnDesc = Super::GetStaticDescription();
+	FString ReturnDesc =
+#if WITH_EDITORONLY_DATA
+		CustomDescription.Len() ? CustomDescription :
+#endif // WITH_EDITORONLY_DATA
+		Super::GetStaticDescription();
 
 	UBTTask_BlueprintBase* CDO = (UBTTask_BlueprintBase*)(GetClass()->GetDefaultObject());
 	if (bShowPropertyDetails && CDO)
