PR #5855: Fixed path following. (Contributed by [REDACTED_PERSON])

More detail: the change fixes the case when UPathFollowingComponent::RequestMoveWithImmediateFinish gets called with the original move request having bStopMovementOnFinish == false. Setting bStopMovementOnFinish in RequestMoveWithImmediateFinish makes it consistent with UPathFollowingComponent::AbortMove.

diff --git a/Engine/Source/Runtime/AIModule/Private/Navigation/PathFollowingComponent.cpp b/Engine/Source/Runtime/AIModule/Private/Navigation/PathFollowingComponent.cpp
--- a/Engine/Source/Runtime/AIModule/Private/Navigation/PathFollowingComponent.cpp
+++ b/Engine/Source/Runtime/AIModule/Private/Navigation/PathFollowingComponent.cpp
@@ -457,6 +457,9 @@ void UPathFollowingComponent::AbortMove(const UObject& Instigator, FPathFollowin
 
 FAIRequestID UPathFollowingComponent::RequestMoveWithImmediateFinish(EPathFollowingResult::Type Result, EPathFollowingVelocityMode VelocityMode)
 {
+	// make sure we stop (or not) depending on VelocityMode's value
+	bStopMovementOnFinish = (VelocityMode == EPathFollowingVelocityMode::Reset);
+
 	if (Status != EPathFollowingStatus::Idle)
 	{
 		OnPathFinished(EPathFollowingResult::Aborted, FPathFollowingResultFlags::NewRequest);
