PR #6526: Add public setter for TagQueryToMatch in UEnvQueryTest_GameplayTags. (Contributed by [REDACTED_PERSON])

diff --git a/Engine/Source/Runtime/AIModule/Classes/EnvironmentQuery/Tests/EnvQueryTest_GameplayTags.h b/Engine/Source/Runtime/AIModule/Classes/EnvironmentQuery/Tests/EnvQueryTest_GameplayTags.h
--- a/Engine/Source/Runtime/AIModule/Classes/EnvironmentQuery/Tests/EnvQueryTest_GameplayTags.h
+++ b/Engine/Source/Runtime/AIModule/Classes/EnvironmentQuery/Tests/EnvQueryTest_GameplayTags.h
@@ -15,6 +15,15 @@ class UEnvQueryTest_GameplayTags : public UEnvQueryTest
 {
 	GENERATED_UCLASS_BODY()
 
+public:
+	/**
+	 * @note Calling function only makes sense before first run of given query
+	 * by the EQS manager. The query gets preprocessed and cached then so the query 
+	 * value will get stored and calling this function will not change it (unless 
+	 * you call it on the cached test's instance, see UEnvQueryManager::CreateQueryInstance).
+	 */
+	void SetTagQueryToMatch(FGameplayTagQuery& GameplayTagQuery);
+
 protected:
 	virtual void RunTest(FEnvQueryInstance& QueryInstance) const override;
 
diff --git a/Engine/Source/Runtime/AIModule/Private/EnvironmentQuery/Tests/EnvQueryTest_GameplayTags.cpp b/Engine/Source/Runtime/AIModule/Private/EnvironmentQuery/Tests/EnvQueryTest_GameplayTags.cpp
--- a/Engine/Source/Runtime/AIModule/Private/EnvironmentQuery/Tests/EnvQueryTest_GameplayTags.cpp
+++ b/Engine/Source/Runtime/AIModule/Private/EnvironmentQuery/Tests/EnvQueryTest_GameplayTags.cpp
@@ -114,3 +114,8 @@ FText UEnvQueryTest_GameplayTags::GetDescriptionDetails() const
 {
 	return FText::FromString(TagQueryToMatch.GetDescription());
 }
+
+void UEnvQueryTest_GameplayTags::SetTagQueryToMatch(FGameplayTagQuery& GameplayTagQuery)
+{
+	TagQueryToMatch = GameplayTagQuery;
+}
