PR #6543: Check that sensing is enabled before setting timer in UPawnSensingCom. (Contributed by [REDACTED_PERSON])

diff --git a/Engine/Source/Runtime/AIModule/Private/Perception/PawnSensingComponent.cpp b/Engine/Source/Runtime/AIModule/Private/Perception/PawnSensingComponent.cpp
--- a/Engine/Source/Runtime/AIModule/Private/Perception/PawnSensingComponent.cpp
+++ b/Engine/Source/Runtime/AIModule/Private/Perception/PawnSensingComponent.cpp
@@ -139,7 +139,10 @@ void UPawnSensingComponent::OnTimer()
 		UpdateAISensing();
 	}
 	
-	SetTimer(SensingInterval);
+	if (bEnableSensingUpdates)
+	{
+		SetTimer(SensingInterval);
+	}
 };
 
 
