Added a handly callable-in-editor function for copying endpoints from NavLinkProxy's simple links over to smart link

This improves usability since there's no gizmo for smart link end points.

diff --git a/Engine/Source/Runtime/AIModule/Classes/Navigation/NavLinkProxy.h b/Engine/Source/Runtime/AIModule/Classes/Navigation/NavLinkProxy.h
--- a/Engine/Source/Runtime/AIModule/Classes/Navigation/NavLinkProxy.h
+++ b/Engine/Source/Runtime/AIModule/Classes/Navigation/NavLinkProxy.h
@@ -26,7 +26,7 @@ class AIMODULE_API ANavLinkProxy : public AActor, public INavLinkHostInterface,
 	GENERATED_UCLASS_BODY()
 
 	/** Navigation links (point to point) added to navigation data */
-	UPROPERTY(EditAnywhere, Category=SimpleLink)
+	UPROPERTY(EditAnywhere, BlueprintReadOnly, Category=SimpleLink)
 	TArray<FNavigationLink> PointLinks;
 	
 	/** Navigation links (segment to segment) added to navigation data
@@ -106,6 +106,14 @@ public:
 	UFUNCTION(BlueprintCallable, Category="AI|Navigation")
 	bool HasMovingAgents() const;
 
+#if WITH_EDITOR
+	/** Copies navlink end points from the first entry in PointLinks array. This 
+	 *	function is a helper function making up for smart links not drawing
+	 *	the FVector widgets in the editor. */
+	UFUNCTION(CallInEditor, Category = SmartLink, meta = (DisplayName="CopyEndPointsFromSimpleLink"))
+	void CopyEndPointsFromSimpleLinkToSmartLink();
+#endif // WITH_EDITOR
+
 protected:
 
 	UPROPERTY(BlueprintAssignable)
diff --git a/Engine/Source/Runtime/AIModule/Private/Navigation/NavLinkProxy.cpp b/Engine/Source/Runtime/AIModule/Private/Navigation/NavLinkProxy.cpp
--- a/Engine/Source/Runtime/AIModule/Private/Navigation/NavLinkProxy.cpp
+++ b/Engine/Source/Runtime/AIModule/Private/Navigation/NavLinkProxy.cpp
@@ -289,3 +289,13 @@ bool ANavLinkProxy::HasMovingAgents() const
 {
 	return SmartLinkComp->HasMovingAgents();
 }
+
+#if WITH_EDITOR
+void ANavLinkProxy::CopyEndPointsFromSimpleLinkToSmartLink()
+{
+	if (PointLinks.Num() && SmartLinkComp)
+	{
+		SmartLinkComp->SetLinkData(PointLinks[0].Left, PointLinks[0].Right, PointLinks[0].Direction);
+	}
+}
+#endif // WITH_EDITOR
\ No newline at end of file
