Optimized sight config update, by allowing the recreation of the queries only when certain values are changed.

#rnx

diff --git a/Engine/Source/Runtime/AIModule/Classes/Perception/AIPerceptionSystem.h b/Engine/Source/Runtime/AIModule/Classes/Perception/AIPerceptionSystem.h
--- a/Engine/Source/Runtime/AIModule/Classes/Perception/AIPerceptionSystem.h
+++ b/Engine/Source/Runtime/AIModule/Classes/Perception/AIPerceptionSystem.h
@@ -163,6 +163,7 @@ public:
 
 	void OnListenerForgetsActor(const UAIPerceptionComponent& Listener, AActor& ActorToForget);
 	void OnListenerForgetsAll(const UAIPerceptionComponent& Listener);
+	void OnListenerConfigUpdated(FAISenseID SenseID, const UAIPerceptionComponent& Listener);
 
 	void RegisterDelayedStimulus(FPerceptionListenerID ListenerId, float Delay, AActor* Instigator, const FAIStimulus& Stimulus);
 
diff --git a/Engine/Source/Runtime/AIModule/Classes/Perception/AISense.h b/Engine/Source/Runtime/AIModule/Classes/Perception/AISense.h
--- a/Engine/Source/Runtime/AIModule/Classes/Perception/AISense.h
+++ b/Engine/Source/Runtime/AIModule/Classes/Perception/AISense.h
@@ -114,8 +114,9 @@ public:
 	virtual void OnListenerForgetsAll(const FPerceptionListener& Listener) {}
 
 	FORCEINLINE void OnNewListener(const FPerceptionListener& NewListener) { OnNewListenerDelegate.ExecuteIfBound(NewListener); }
-	FORCEINLINE void OnListenerUpdate(const FPerceptionListener& NewListener) { OnListenerUpdateDelegate.ExecuteIfBound(NewListener); }
-	FORCEINLINE void OnListenerRemoved(const FPerceptionListener& NewListener) { OnListenerRemovedDelegate.ExecuteIfBound(NewListener); }
+	FORCEINLINE void OnListenerUpdate(const FPerceptionListener& UpdatedListener) { OnListenerUpdateDelegate.ExecuteIfBound(UpdatedListener); }
+	FORCEINLINE void OnListenerRemoved(const FPerceptionListener& RemovedListener) { OnListenerRemovedDelegate.ExecuteIfBound(RemovedListener); }
+	virtual void OnListenerConfigUpdated(const FPerceptionListener& UpdatedListener) { OnListenerUpdate(UpdatedListener); }
 
 	UE_DEPRECATED(4.23, "This method will be removed in future versions. Perception relies on AISenseConfig::MaxAge so the value returned is no longer used by the perception system.")
 	FORCEINLINE float GetDefaultExpirationAge() const
diff --git a/Engine/Source/Runtime/AIModule/Classes/Perception/AISense_Sight.h b/Engine/Source/Runtime/AIModule/Classes/Perception/AISense_Sight.h
--- a/Engine/Source/Runtime/AIModule/Classes/Perception/AISense_Sight.h
+++ b/Engine/Source/Runtime/AIModule/Classes/Perception/AISense_Sight.h
@@ -185,7 +185,8 @@ protected:
 
 	void OnNewListenerImpl(const FPerceptionListener& NewListener);
 	void OnListenerUpdateImpl(const FPerceptionListener& UpdatedListener);
-	void OnListenerRemovedImpl(const FPerceptionListener& UpdatedListener);	
+	void OnListenerRemovedImpl(const FPerceptionListener& RemovedListener);
+	virtual void OnListenerConfigUpdated(const FPerceptionListener& UpdatedListener) override;
 	
 	void GenerateQueriesForListener(const FPerceptionListener& Listener, const FDigestedSightProperties& PropertyDigest, const TFunction<void(FAISightQuery&)>& OnAddedFunc = nullptr);
 
diff --git a/Engine/Source/Runtime/AIModule/Private/Perception/AIPerceptionComponent.cpp b/Engine/Source/Runtime/AIModule/Private/Perception/AIPerceptionComponent.cpp
--- a/Engine/Source/Runtime/AIModule/Private/Perception/AIPerceptionComponent.cpp
+++ b/Engine/Source/Runtime/AIModule/Private/Perception/AIPerceptionComponent.cpp
@@ -128,7 +128,11 @@ void UAIPerceptionComponent::ConfigureSense(UAISenseConfig& Config)
 
 	if (IsRegistered())
 	{
-		RequestStimuliListenerUpdate();
+	    UAIPerceptionSystem* AIPerceptionSys = UAIPerceptionSystem::GetCurrent(GetWorld());
+	    if (AIPerceptionSys != nullptr)
+	    {
+		    AIPerceptionSys->OnListenerConfigUpdated(Config.GetSenseID(), *this);
+	    }
 	}
 	// else the sense will be auto-configured during OnRegister
 }
diff --git a/Engine/Source/Runtime/AIModule/Private/Perception/AIPerceptionSystem.cpp b/Engine/Source/Runtime/AIModule/Private/Perception/AIPerceptionSystem.cpp
--- a/Engine/Source/Runtime/AIModule/Private/Perception/AIPerceptionSystem.cpp
+++ b/Engine/Source/Runtime/AIModule/Private/Perception/AIPerceptionSystem.cpp
@@ -307,6 +307,24 @@ void UAIPerceptionSystem::UpdateListener(UAIPerceptionComponent& Listener)
 	}
 }
 
+void UAIPerceptionSystem::OnListenerConfigUpdated(FAISenseID SenseID, const UAIPerceptionComponent& Listener)
+{
+	SCOPE_CYCLE_COUNTER(STAT_AI_PerceptionSys);
+
+	if (ensureMsgf(SenseID.IsValid() && Senses.IsValidIndex(SenseID) && Senses[SenseID], TEXT("Sense must exist to update its sense config")))
+	{
+		const FPerceptionListenerID ListenerId = Listener.GetListenerId();
+		if (ensureMsgf(ListenerId != FPerceptionListenerID::InvalidID(), TEXT("Listener must have a valid it to update the its sense config")))
+		{
+			FPerceptionListener& ListenerEntry = ListenerContainer[ListenerId];
+			if (ensureMsgf(ListenerEntry.Listener.IsValid() && ListenerEntry.Listener.Get() == &Listener, TEXT("Listener is not a valid registered listener")))
+			{
+				Senses[SenseID]->OnListenerConfigUpdated(ListenerEntry);
+			}
+		}
+	}
+}
+
 void UAIPerceptionSystem::UnregisterListener(UAIPerceptionComponent& Listener)
 {
 	SCOPE_CYCLE_COUNTER(STAT_AI_PerceptionSys);
diff --git a/Engine/Source/Runtime/AIModule/Private/Perception/AISense_Sight.cpp b/Engine/Source/Runtime/AIModule/Private/Perception/AISense_Sight.cpp
--- a/Engine/Source/Runtime/AIModule/Private/Perception/AISense_Sight.cpp
+++ b/Engine/Source/Runtime/AIModule/Private/Perception/AISense_Sight.cpp
@@ -575,11 +575,33 @@ void UAISense_Sight::OnListenerUpdateImpl(const FPerceptionListener& UpdatedList
 	}
 }
 
-void UAISense_Sight::OnListenerRemovedImpl(const FPerceptionListener& UpdatedListener)
+void UAISense_Sight::OnListenerConfigUpdated(const FPerceptionListener& UpdatedListener)
 {
-	RemoveAllQueriesByListener(UpdatedListener);
+	bool bSkipListenerUpdate = false;
+	const FPerceptionListenerID ListenerID = UpdatedListener.GetListenerID();
+
+	FDigestedSightProperties* PropertiesDigest = DigestedProperties.Find(ListenerID);
+	if (PropertiesDigest)
+	{
+		// The only parameter we need to rebuild all the queries for this listener is if the affiliation mask changed, otherwise there is nothing to update.
+		const UAISenseConfig_Sight* SenseConfig = CastChecked<const UAISenseConfig_Sight>(UpdatedListener.Listener->GetSenseConfig(GetSenseID()));
+		FDigestedSightProperties NewPropertiesDigest(*SenseConfig);
+		bSkipListenerUpdate = NewPropertiesDigest.AffiliationFlags == PropertiesDigest->AffiliationFlags;
+		*PropertiesDigest = NewPropertiesDigest;
+	}
+
+	if (!bSkipListenerUpdate)
+	{
+		Super::OnListenerConfigUpdated(UpdatedListener);
+	}
+}
+
+
+void UAISense_Sight::OnListenerRemovedImpl(const FPerceptionListener& RemovedListener)
+{
+	RemoveAllQueriesByListener(RemovedListener);
 
-	DigestedProperties.FindAndRemoveChecked(UpdatedListener.GetListenerID());
+	DigestedProperties.FindAndRemoveChecked(RemovedListener.GetListenerID());
 
 	// note: there use to be code to remove all queries _to_ listener here as well
 	// but that was wrong - the fact that a listener gets unregistered doesn't have to
