Fix CIS ensure on autotests, thx Bob

#rnx

diff --git a/Engine/Source/Runtime/AIModule/Private/Perception/AIPerceptionSystem.cpp b/Engine/Source/Runtime/AIModule/Private/Perception/AIPerceptionSystem.cpp
--- a/Engine/Source/Runtime/AIModule/Private/Perception/AIPerceptionSystem.cpp
+++ b/Engine/Source/Runtime/AIModule/Private/Perception/AIPerceptionSystem.cpp
@@ -311,18 +311,23 @@ void UAIPerceptionSystem::OnListenerConfigUpdated(FAISenseID SenseID, const UAIP
 {
 	SCOPE_CYCLE_COUNTER(STAT_AI_PerceptionSys);
 
-	if (ensureMsgf(SenseID.IsValid() && Senses.IsValidIndex(SenseID) && Senses[SenseID], TEXT("Sense must exist to update its sense config")))
+	if (!IsSenseInstantiated(SenseID))
 	{
-		const FPerceptionListenerID ListenerId = Listener.GetListenerId();
-		if (ensureMsgf(ListenerId != FPerceptionListenerID::InvalidID(), TEXT("Listener must have a valid it to update the its sense config")))
-		{
-			FPerceptionListener& ListenerEntry = ListenerContainer[ListenerId];
-			if (ensureMsgf(ListenerEntry.Listener.IsValid() && ListenerEntry.Listener.Get() == &Listener, TEXT("Listener is not a valid registered listener")))
-			{
-				Senses[SenseID]->OnListenerConfigUpdated(ListenerEntry);
-			}
-		}
+		UE_LOG(LogAIPerception, Warning, TEXT("Sense must exist to update its sense config"));
+		return;
+	}
+
+	const FPerceptionListenerID ListenerId = Listener.GetListenerId();
+	if (ListenerId == FPerceptionListenerID::InvalidID() || !ListenerContainer.Contains(ListenerId))
+	{
+		UE_LOG(LogAIPerception, Warning, TEXT("Listener must have a valid id to update its sense config"));
+		return;
 	}
+
+	FPerceptionListener& ListenerEntry = ListenerContainer[ListenerId];
+	check(ListenerEntry.Listener.IsValid() && ListenerEntry.Listener.Get() == &Listener);
+
+	Senses[SenseID]->OnListenerConfigUpdated(ListenerEntry);
 }
 
 void UAIPerceptionSystem::UnregisterListener(UAIPerceptionComponent& Listener)
