AIPerceptionComponent: fixed MaxActiveAge array indexing by always using AISenseID.
Additional fix in ConfigureSenses to register AISenseConfig added after component registration to the perception system.

#rnx

diff --git a/Engine/Source/Runtime/AIModule/Classes/Perception/AIPerceptionComponent.h b/Engine/Source/Runtime/AIModule/Classes/Perception/AIPerceptionComponent.h
--- a/Engine/Source/Runtime/AIModule/Classes/Perception/AIPerceptionComponent.h
+++ b/Engine/Source/Runtime/AIModule/Classes/Perception/AIPerceptionComponent.h
@@ -371,7 +371,8 @@ protected:
 private:
 	friend UAIPerceptionSystem;
 
+	void RegisterSenseConfig(UAISenseConfig& SenseConfig, UAIPerceptionSystem& AIPerceptionSys);
 	void StoreListenerId(FPerceptionListenerID InListenerId) { PerceptionListenerId = InListenerId; }
-	void SetMaxStimulusAge(int32 ConfigIndex, float MaxAge);
+	void SetMaxStimulusAge(FAISenseID SenseId, float MaxAge);
 };
 
diff --git a/Engine/Source/Runtime/AIModule/Private/Perception/AIPerceptionComponent.cpp b/Engine/Source/Runtime/AIModule/Private/Perception/AIPerceptionComponent.cpp
--- a/Engine/Source/Runtime/AIModule/Private/Perception/AIPerceptionComponent.cpp
+++ b/Engine/Source/Runtime/AIModule/Private/Perception/AIPerceptionComponent.cpp
@@ -106,24 +106,21 @@ void UAIPerceptionComponent::PostInitProperties()
 
 void UAIPerceptionComponent::ConfigureSense(UAISenseConfig& Config)
 {
-	int32 ConfigIndex = INDEX_NONE;
 	// first check if we're reconfiguring a sense
-	for (int32 Index = 0; Index < SensesConfig.Num(); ++Index)
+	bool bIsNewConfig = true;
+	for (UAISenseConfig*& SenseConfig : SensesConfig)
 	{
-		if (SensesConfig[Index] != nullptr
-			&& SensesConfig[Index]->GetClass() == Config.GetClass())
+		if (SenseConfig != nullptr && SenseConfig->GetClass() == Config.GetClass())
 		{
-			ConfigIndex = Index;
-			SensesConfig[Index] = &Config;
-			SetMaxStimulusAge(Index, Config.GetMaxAge());
+			SenseConfig = &Config;
+			bIsNewConfig = false;
 			break;
 		}
 	}
 
-	if (ConfigIndex == INDEX_NONE)
+	if (bIsNewConfig)
 	{
-		ConfigIndex = SensesConfig.Add(&Config);
-		SetMaxStimulusAge(ConfigIndex, Config.GetMaxAge());
+		SensesConfig.Add(&Config);
 	}
 
 	if (IsRegistered())
@@ -131,6 +128,14 @@ void UAIPerceptionComponent::ConfigureSense(UAISenseConfig& Config)
 	    UAIPerceptionSystem* AIPerceptionSys = UAIPerceptionSystem::GetCurrent(GetWorld());
 	    if (AIPerceptionSys != nullptr)
 	    {
+			if (bIsNewConfig)
+			{
+				RegisterSenseConfig(Config, *AIPerceptionSys);
+			}
+			else
+			{
+				SetMaxStimulusAge(Config.GetSenseID(), Config.GetMaxAge());
+			}
 		    AIPerceptionSys->OnListenerConfigUpdated(Config.GetSenseID(), *this);
 	    }
 	}
@@ -142,13 +147,18 @@ UAIPerceptionComponent::TAISenseConfigConstIterator UAIPerceptionComponent::GetS
 	return SensesConfig.CreateConstIterator();
 }
 
-void UAIPerceptionComponent::SetMaxStimulusAge(int32 ConfigIndex, float MaxAge)
+void UAIPerceptionComponent::SetMaxStimulusAge(FAISenseID SenseID, float MaxAge)
 {
-	if (MaxActiveAge.IsValidIndex(ConfigIndex) == false)
+	if (!ensureMsgf(SenseID.IsValid(), TEXT("Sense must exist to update max age")))
+	{
+		return;
+	}
+
+	if (MaxActiveAge.IsValidIndex(SenseID) == false)
 	{
-		MaxActiveAge.AddUninitialized(ConfigIndex - MaxActiveAge.Num() + 1);
+		MaxActiveAge.AddUninitialized(SenseID - MaxActiveAge.Num() + 1);
 	}
-	MaxActiveAge[ConfigIndex] = MaxAge;
+	MaxActiveAge[SenseID] = MaxAge;
 
 	// @todo process all data already gathered and see if any _still_active_ stimuli
 	// got it's expiration prolonged, with SetExpirationAge
@@ -179,22 +189,7 @@ void UAIPerceptionComponent::OnRegister()
 			{
 				if (SenseConfig)
 				{
-					TSubclassOf<UAISense> SenseImplementation = SenseConfig->GetSenseImplementation();
-
-					if (SenseImplementation)
-					{
-						// make sure it's registered with perception system
-						AIPerceptionSys->RegisterSenseClass(SenseImplementation);
-
-						if (SenseConfig->IsEnabled())
-						{
-							PerceptionFilter.AcceptChannel(UAISense::GetSenseID(SenseImplementation));
-						}
-
-						const FAISenseID SenseID = UAISense::GetSenseID(SenseImplementation);
-						check(SenseID.IsValid());
-						SetMaxStimulusAge(SenseID, SenseConfig->GetMaxAge());
-					}
+					RegisterSenseConfig(*SenseConfig, *AIPerceptionSys);
 				}
 			}
 
@@ -213,6 +208,24 @@ void UAIPerceptionComponent::OnRegister()
 	}
 }
 
+void UAIPerceptionComponent::RegisterSenseConfig(UAISenseConfig& SenseConfig, UAIPerceptionSystem& AIPerceptionSys)
+{
+	const TSubclassOf<UAISense> SenseImplementation = SenseConfig.GetSenseImplementation();
+	if (SenseImplementation)
+	{
+		// make sure it's registered with perception system
+		const FAISenseID SenseID = AIPerceptionSys.RegisterSenseClass(SenseImplementation);
+		check(SenseID.IsValid());
+
+		if (SenseConfig.IsEnabled())
+		{
+			PerceptionFilter.AcceptChannel(SenseID);
+		}
+
+		SetMaxStimulusAge(SenseID, SenseConfig.GetMaxAge());
+	}
+}
+
 void UAIPerceptionComponent::OnUnregister()
 {
 	CleanUp();
