Fix static analysis warning dereferencing null pointer (NodeB)

diff --git a/Engine/Source/Runtime/AIModule/Private/BehaviorTree/BehaviorTreeComponent.cpp b/Engine/Source/Runtime/AIModule/Private/BehaviorTree/BehaviorTreeComponent.cpp
--- a/Engine/Source/Runtime/AIModule/Private/BehaviorTree/BehaviorTreeComponent.cpp
+++ b/Engine/Source/Runtime/AIModule/Private/BehaviorTree/BehaviorTreeComponent.cpp
@@ -778,6 +778,8 @@ static void FindCommonParent(const TArray<FBehaviorTreeInstance>& Instances, con
 			*UBehaviorTreeTypes::DescribeNodeHelper(InNodeB), InstanceIdxB, *AssetBName, *UBehaviorTreeTypes::DescribeNodeHelper(NodeB),
 			CommonInstanceIdx, *AssetCName, *UBehaviorTreeTypes::DescribeNodeHelper(Instances[CommonInstanceIdx].ActiveNode),
 			(Instances[CommonInstanceIdx].ActiveNode == Instances[CommonInstanceIdx].RootNode) ? TEXT(" (root)") : TEXT(""));
+
+		return;
 	}
 
 	// find common parent of two nodes
