Optimization: split perception query list into 2 list, in range and out of range.

#rnx

diff --git a/Engine/Source/Runtime/AIModule/Classes/Perception/AISense_Sight.h b/Engine/Source/Runtime/AIModule/Classes/Perception/AISense_Sight.h
--- a/Engine/Source/Runtime/AIModule/Classes/Perception/AISense_Sight.h
+++ b/Engine/Source/Runtime/AIModule/Classes/Perception/AISense_Sight.h
@@ -141,7 +141,13 @@ public:
 	FTargetsContainer ObservedTargets;
 	TMap<FPerceptionListenerID, FDigestedSightProperties> DigestedProperties;
 
-	TArray<FAISightQuery> SightQueryQueue;
+	/** The SightQueries are a n^2 problem and to reduce the sort time, they are now split between in range and out of range */
+	/** Since the out of range queries only age as the distance component of the score is always 0, there is few need to sort them */
+	/** In the majority of the cases most of the queries are out of range, so the sort time is greatly reduced as we only sort the in range queries */
+	int32 NextOutOfRangeIndex = 0;
+	bool bSightQueriesOutOfRangeDirty = true;
+	TArray<FAISightQuery> SightQueriesOutOfRange;
+	TArray<FAISightQuery> SightQueriesInRange;
 
 protected:
 	UPROPERTY(EditDefaultsOnly, Category = "AI Perception", config)
diff --git a/Engine/Source/Runtime/AIModule/Private/Perception/AISense_Sight.cpp b/Engine/Source/Runtime/AIModule/Private/Perception/AISense_Sight.cpp
--- a/Engine/Source/Runtime/AIModule/Private/Perception/AISense_Sight.cpp
+++ b/Engine/Source/Runtime/AIModule/Private/Perception/AISense_Sight.cpp
@@ -46,6 +46,45 @@ FORCEINLINE_DEBUGGABLE bool CheckIsTargetInSightPie(const FPerceptionListener& L
 	return false;
 }
 
+enum class EForEachResult : uint8
+{
+	Break,
+	Continue,
+};
+
+template <typename T, class PREDICATE_CLASS>
+EForEachResult ForEach(T& Array, const PREDICATE_CLASS& Predicate)
+{
+	for (T::ElementType& Element : Array)
+	{
+		if (Predicate(Element) == EForEachResult::Break)
+		{
+			return EForEachResult::Break;
+		}
+	}
+	return EForEachResult::Continue;
+}
+
+enum EReverseForEachResult : uint8
+{
+	UnTouched,
+	Modified,
+};
+
+template <typename T, class PREDICATE_CLASS>
+EReverseForEachResult ReverseForEach(T& Array, const PREDICATE_CLASS& Predicate)
+{
+	EReverseForEachResult RetVal = EReverseForEachResult::UnTouched;
+	for (int32 Index = Array.Num()-1; Index >= 0; --Index)
+	{
+		if (Predicate(Array, Index) == EReverseForEachResult::Modified)
+		{
+			RetVal = EReverseForEachResult::Modified;
+		}
+	}
+	return RetVal;
+}
+
 //----------------------------------------------------------------------//
 // FAISightTarget
 //----------------------------------------------------------------------//
@@ -150,12 +189,25 @@ float UAISense_Sight::Update()
 
 	// sort Sight Queries
 	{
-		SCOPE_CYCLE_COUNTER(STAT_AI_Sense_Sight_UpdateSort);
-		for (FAISightQuery& SightQuery : SightQueryQueue)
+		auto RecalcScore = [](FAISightQuery& SightQuery)->EForEachResult
 		{
 			SightQuery.RecalcScore();
+			return EForEachResult::Continue;
+		};
+
+		SCOPE_CYCLE_COUNTER(STAT_AI_Sense_Sight_UpdateSort);
+        // Sort out of range queries
+    	if (bSightQueriesOutOfRangeDirty)
+		{
+			ForEach(SightQueriesOutOfRange, RecalcScore);
+			SightQueriesOutOfRange.Sort(FAISightQuery::FSortPredicate());
+			NextOutOfRangeIndex = 0;
+			bSightQueriesOutOfRangeDirty = false;
 		}
-		SightQueryQueue.Sort(FAISightQuery::FSortPredicate());
+
+        // Sort in range queries
+		ForEach(SightQueriesInRange, RecalcScore);
+		SightQueriesInRange.Sort(FAISightQuery::FSortPredicate());
 	}
 
 	int32 TracesCount = 0;
@@ -168,16 +220,45 @@ float UAISense_Sight::Update()
 	double LastTime = FPlatformTime::Seconds();
 #endif // AISENSE_SIGHT_TIMESLICING_DEBUG
 	static const int32 InitialInvalidItemsSize = 16;
-	TArray<int32> InvalidQueries;
+	enum class EOperationType : uint8
+	{
+		Remove,
+		SwapList
+	};
+	struct FQueryOperation
+	{
+		FQueryOperation(bool bInInRange, EOperationType InOpType, int32 InIndex) : bInRange(bInInRange), OpType(InOpType), Index(InIndex) {}
+		bool bInRange;
+		EOperationType OpType;
+		int32 Index;
+	};
+	TArray<FQueryOperation> QueryOperations;
 	TArray<FAISightTarget::FTargetId> InvalidTargets;
-	InvalidQueries.Reserve(InitialInvalidItemsSize);
+	QueryOperations.Reserve(InitialInvalidItemsSize);
 	InvalidTargets.Reserve(InitialInvalidItemsSize);
 
 	AIPerception::FListenerMap& ListenersMap = *GetListeners();
 
-	FAISightQuery* SightQuery = SightQueryQueue.GetData();
-	for (int32 QueryIndex = 0; QueryIndex < SightQueryQueue.Num(); ++QueryIndex, ++SightQuery)
+	int32 InRangeItr = 0;
+	int32 OutOfRangeItr = 0;
+	for (int32 QueryIndex = 0; QueryIndex < SightQueriesInRange.Num() + SightQueriesOutOfRange.Num(); ++QueryIndex)
 	{
+		// Calculate next in range query
+		int32 InRangeIndex = SightQueriesInRange.IsValidIndex(InRangeItr) ? InRangeItr : INDEX_NONE;
+		FAISightQuery* InRangeQuery = InRangeIndex != INDEX_NONE ? &SightQueriesInRange[InRangeIndex] : nullptr;
+
+		// Calculate next out of range query
+		int32 OutOfRangeIndex = SightQueriesOutOfRange.IsValidIndex(OutOfRangeItr) ? (NextOutOfRangeIndex + OutOfRangeItr) % SightQueriesOutOfRange.Num() : INDEX_NONE;
+		FAISightQuery* OutOfRangeQuery = OutOfRangeIndex != INDEX_NONE ? &SightQueriesOutOfRange[OutOfRangeIndex] : nullptr;
+		if (OutOfRangeQuery)
+		{
+			OutOfRangeQuery->RecalcScore();
+		}
+
+		// Compare to real find next query
+		const bool bIsInRangeQuery = (InRangeQuery && OutOfRangeQuery) ? FAISightQuery::FSortPredicate()(*InRangeQuery,*OutOfRangeQuery) : !OutOfRangeQuery;
+		FAISightQuery* SightQuery = bIsInRangeQuery ? InRangeQuery : OutOfRangeQuery;
+
 		// Time slice limit check - spread out checks to every N queries so we don't spend more time checking timer than doing work
 		NumQueriesProcessed++;
 #ifdef AISENSE_SIGHT_TIMESLICING_DEBUG
@@ -192,6 +273,8 @@ float UAISense_Sight::Update()
 
 		if (TracesCount < MaxTracesPerTick && bHitTimeSliceLimit == false)
 		{
+			bIsInRangeQuery ? ++InRangeItr : ++OutOfRangeItr;
+
 			FPerceptionListener& Listener = ListenersMap[SightQuery->ObserverId];
 			FAISightTarget& Target = ObservedTargets[SightQuery->TargetId];
 
@@ -292,6 +375,11 @@ float UAISense_Sight::Update()
 				}
 
 				SightQuery->Importance = CalcQueryImportance(Listener, TargetLocation, SightRadiusSq);
+				const bool bShouldBeInRange = SightQuery->Importance > 0.0f;
+				if (bIsInRangeQuery != bShouldBeInRange)
+				{
+					QueryOperations.Add(FQueryOperation(bIsInRangeQuery, EOperationType::SwapList, bIsInRangeQuery ? InRangeIndex : OutOfRangeIndex));
+				}
 
 				// restart query
 				SightQuery->OnProcessed();
@@ -299,7 +387,7 @@ float UAISense_Sight::Update()
 			else
 			{
 				// put this index to "to be removed" array
-				InvalidQueries.Add(QueryIndex);
+				QueryOperations.Add( FQueryOperation(bIsInRangeQuery, EOperationType::Remove, bIsInRangeQuery ? InRangeIndex : OutOfRangeIndex) );
 				if (TargetActor == nullptr)
 				{
 					InvalidTargets.AddUnique(SightQuery->TargetId);
@@ -311,18 +399,59 @@ float UAISense_Sight::Update()
 			break;
 		}
 	}
+	NextOutOfRangeIndex = SightQueriesOutOfRange.Num() > 0 ? (NextOutOfRangeIndex + OutOfRangeItr) % SightQueriesOutOfRange.Num() : 0;
+
 #ifdef AISENSE_SIGHT_TIMESLICING_DEBUG
 	UE_LOG(LogAIPerception, VeryVerbose, TEXT("UAISense_Sight::Update processed %d sources in %f seconds [time slice limited? %d]"), NumQueriesProcessed, TimeSpent, bHitTimeSliceLimit ? 1 : 0);
 #else
 	UE_LOG(LogAIPerception, VeryVerbose, TEXT("UAISense_Sight::Update processed %d sources [time slice limited? %d]"), NumQueriesProcessed, bHitTimeSliceLimit ? 1 : 0);
 #endif // AISENSE_SIGHT_TIMESLICING_DEBUG
 
-	if (InvalidQueries.Num() > 0)
+	if (QueryOperations.Num() > 0)
 	{
-		for (int32 Index = InvalidQueries.Num() - 1; Index >= 0; --Index)
+		// Sort by InRange and by descending Index 
+		QueryOperations.Sort([](const FQueryOperation& LHS, const FQueryOperation& RHS)->bool
 		{
-			// removing with swapping here, since queue is going to be sorted anyway
-			SightQueryQueue.RemoveAtSwap(InvalidQueries[Index], 1, /*bAllowShrinking*/false);
+			if (LHS.bInRange != RHS.bInRange)
+				return LHS.bInRange;
+			return LHS.Index > RHS.Index;
+		});
+        // Do all the removes first and save the out of range swaps because we will insert them at the right location to prevent sorting
+		TArray<FAISightQuery> SightQueriesOutOfRangeToInsert;
+		for (FQueryOperation& Operation : QueryOperations)
+		{
+			if (Operation.OpType == EOperationType::SwapList)
+			{
+				if (Operation.bInRange)
+				{
+					SightQueriesOutOfRangeToInsert.Push(SightQueriesInRange[Operation.Index]);
+				}
+				else
+				{
+					SightQueriesInRange.Add(SightQueriesOutOfRange[Operation.Index]);
+				}
+			}
+
+			if (Operation.bInRange)
+			{
+				// In range queries are always sorted at the beginning of the update
+				SightQueriesInRange.RemoveAtSwap(Operation.Index, 1, /*bAllowShrinking*/false);
+			}
+			else
+			{
+				// Preserve the list ordered
+				SightQueriesOutOfRange.RemoveAt(Operation.Index, 1, /*bAllowShrinking*/false);
+				if (Operation.Index < NextOutOfRangeIndex)
+				{
+					NextOutOfRangeIndex--;
+				}
+			}
+		}
+        // Reinsert the saved out of range swaps
+		if (SightQueriesOutOfRangeToInsert.Num() > 0)
+		{
+			SightQueriesOutOfRange.Insert(SightQueriesOutOfRangeToInsert.GetData(), SightQueriesOutOfRangeToInsert.Num(), NextOutOfRangeIndex);
+			NextOutOfRangeIndex += SightQueriesOutOfRangeToInsert.Num();
 		}
 
 		if (InvalidTargets.Num() > 0)
@@ -343,7 +472,7 @@ float UAISense_Sight::Update()
 		}
 	}
 
-	//return SightQueryQueue.Num() > 0 ? 1.f/6 : FLT_MAX;
+	//return SightQueries.Num() > 0 ? 1.f/6 : FLT_MAX;
 	return 0.f;
 }
 
@@ -362,8 +491,8 @@ void UAISense_Sight::UnregisterSource(AActor& SourceActor)
 	const FAISightTarget::FTargetId AsTargetId = SourceActor.GetUniqueID();
 	FAISightTarget AsTarget;
 	
-	if (ObservedTargets.RemoveAndCopyValue(AsTargetId, AsTarget)
-		&& SightQueryQueue.Num() > 0)
+	if (ObservedTargets.RemoveAndCopyValue(AsTargetId, AsTarget) 
+		&& (SightQueriesInRange.Num() + SightQueriesOutOfRange.Num()) > 0)
 	{
 		AActor* TargetActor = AsTarget.Target.Get();
 
@@ -372,9 +501,9 @@ void UAISense_Sight::UnregisterSource(AActor& SourceActor)
 			// notify all interested observers that this source is no longer
 			// visible		
 			AIPerception::FListenerMap& ListenersMap = *GetListeners();
-			const FAISightQuery* SightQuery = &SightQueryQueue[SightQueryQueue.Num() - 1];
-			for (int32 QueryIndex = SightQueryQueue.Num() - 1; QueryIndex >= 0; --QueryIndex, --SightQuery)
+			auto RemoveQuery = [this,&ListenersMap,&AsTargetId,&TargetActor](TArray<FAISightQuery>& SightQueries, const int32 QueryIndex)->EReverseForEachResult
 			{
+				FAISightQuery* SightQuery = &SightQueries[QueryIndex];
 				if (SightQuery->TargetId == AsTargetId)
 				{
 					if (SightQuery->bLastResult == true)
@@ -385,8 +514,15 @@ void UAISense_Sight::UnregisterSource(AActor& SourceActor)
 						Listener.RegisterStimulus(TargetActor, FAIStimulus(*this, 0.f, SightQuery->LastSeenLocation, Listener.CachedLocation, FAIStimulus::SensingFailed));
 					}
 
-					SightQueryQueue.RemoveAtSwap(QueryIndex, 1, /*bAllowShrinking=*/false);
+					SightQueries.RemoveAtSwap(QueryIndex, 1, /*bAllowShrinking=*/false);
+					return EReverseForEachResult::Modified;
 				}
+				return EReverseForEachResult::UnTouched;
+			};
+			ReverseForEach(SightQueriesInRange, RemoveQuery);
+			if (ReverseForEach(SightQueriesOutOfRange, RemoveQuery) == EReverseForEachResult::Modified)
+			{
+				bSightQueriesOutOfRangeDirty = true;
 			}
 		}
 	}
@@ -433,10 +569,16 @@ bool UAISense_Sight::RegisterTarget(AActor& TargetActor, const TFunction<void(FA
 			if (FAISenseAffiliationFilter::ShouldSenseTeam(ListenersTeamAgent, TargetActor, PropDigest.AffiliationFlags))
 			{
 				// create a sight query		
-				FAISightQuery& AddedQuery = SightQueryQueue.AddDefaulted_GetRef();
+				const float Importance = CalcQueryImportance(ItListener->Value, TargetLocation, PropDigest.SightRadiusSq);
+				const bool bInRange = Importance > 0.0f;
+				if (!bInRange)
+				{
+					bSightQueriesOutOfRangeDirty = true;
+				}
+				FAISightQuery& AddedQuery = bInRange ? SightQueriesInRange.AddDefaulted_GetRef() : SightQueriesOutOfRange.AddDefaulted_GetRef();
 				AddedQuery.ObserverId = ItListener->Key;
 				AddedQuery.TargetId = SightTarget->TargetId;
-				AddedQuery.Importance = CalcQueryImportance(ItListener->Value, TargetLocation, PropDigest.SightRadiusSq);
+				AddedQuery.Importance = Importance;
 				
 				if (OnAddedFunc)
 				{
@@ -486,10 +628,16 @@ void UAISense_Sight::GenerateQueriesForListener(const FPerceptionListener& Liste
 		if (FAISenseAffiliationFilter::ShouldSenseTeam(ListenersTeamAgent, *TargetActor, PropertyDigest.AffiliationFlags))
 		{
 			// create a sight query		
-			FAISightQuery& AddedQuery = SightQueryQueue.AddDefaulted_GetRef();
+			const float Importance = CalcQueryImportance(Listener, ItTarget->Value.GetLocationSimple(), PropertyDigest.SightRadiusSq);
+			const bool bInRange = Importance > 0.0f;
+			if (!bInRange)
+			{
+				bSightQueriesOutOfRangeDirty = true;
+			}
+			FAISightQuery& AddedQuery = bInRange ? SightQueriesInRange.AddDefaulted_GetRef() : SightQueriesOutOfRange.AddDefaulted_GetRef();
 			AddedQuery.ObserverId = Listener.GetListenerID();
 			AddedQuery.TargetId = ItTarget->Key;
-			AddedQuery.Importance = CalcQueryImportance(Listener, ItTarget->Value.GetLocationSimple(), PropertyDigest.SightRadiusSq);
+			AddedQuery.Importance = Importance;
 
 			if (OnAddedFunc)
 			{
@@ -612,16 +760,16 @@ void UAISense_Sight::RemoveAllQueriesByListener(const FPerceptionListener& Liste
 {
 	SCOPE_CYCLE_COUNTER(STAT_AI_Sense_Sight_RemoveByListener);
 
-	if (SightQueryQueue.Num() == 0)
+	if ((SightQueriesInRange.Num() + SightQueriesOutOfRange.Num()) == 0)
 	{
 		return;
 	}
 
 	const uint32 ListenerId = Listener.GetListenerID();
 	
-	for (int32 QueryIndex = SightQueryQueue.Num() - 1; QueryIndex >= 0 ; --QueryIndex)
+	auto RemoveQuery = [&ListenerId, &OnRemoveFunc](TArray<FAISightQuery>& SightQueries, const int32 QueryIndex)->EReverseForEachResult
 	{
-		const FAISightQuery& SightQuery = SightQueryQueue[QueryIndex];
+		const FAISightQuery& SightQuery = SightQueries[QueryIndex];
 
 		if (SightQuery.ObserverId == ListenerId)
 		{
@@ -629,8 +777,15 @@ void UAISense_Sight::RemoveAllQueriesByListener(const FPerceptionListener& Liste
 			{
 				OnRemoveFunc(SightQuery);
 			}
-			SightQueryQueue.RemoveAtSwap(QueryIndex, 1, /*bAllowShrinking=*/false);
+			SightQueries.RemoveAtSwap(QueryIndex, 1, /*bAllowShrinking=*/false);
+			return EReverseForEachResult::Modified;
 		}
+		return EReverseForEachResult::UnTouched;
+	};
+	ReverseForEach(SightQueriesInRange, RemoveQuery);
+	if(ReverseForEach(SightQueriesOutOfRange, RemoveQuery) == EReverseForEachResult::Modified)
+	{
+		bSightQueriesOutOfRangeDirty = true;
 	}
 }
 
@@ -638,14 +793,9 @@ void UAISense_Sight::RemoveAllQueriesToTarget(const FAISightTarget::FTargetId& T
 {
 	SCOPE_CYCLE_COUNTER(STAT_AI_Sense_Sight_RemoveToTarget);
 
-	if (SightQueryQueue.Num() == 0)
-	{
-		return;
-	}
-
-	for (int32 QueryIndex = SightQueryQueue.Num() - 1; QueryIndex >= 0; --QueryIndex)
+	auto RemoveQuery = [&TargetId, &OnRemoveFunc](TArray<FAISightQuery>& SightQueries, const int32 QueryIndex)->EReverseForEachResult
 	{
-		const FAISightQuery& SightQuery = SightQueryQueue[QueryIndex];
+		const FAISightQuery& SightQuery = SightQueries[QueryIndex];
 
 		if (SightQuery.TargetId == TargetId)
 		{
@@ -653,8 +803,15 @@ void UAISense_Sight::RemoveAllQueriesToTarget(const FAISightTarget::FTargetId& T
 			{
 				OnRemoveFunc(SightQuery);
 			}
-			SightQueryQueue.RemoveAtSwap(QueryIndex, 1, /*bAllowShrinking=*/false);
+			SightQueries.RemoveAtSwap(QueryIndex, 1, /*bAllowShrinking=*/false);
+			return EReverseForEachResult::Modified;
 		}
+		return EReverseForEachResult::UnTouched;
+	};
+	ReverseForEach(SightQueriesInRange, RemoveQuery);
+	if (ReverseForEach(SightQueriesOutOfRange, RemoveQuery) == EReverseForEachResult::Modified)
+	{
+		bSightQueriesOutOfRangeDirty = true;
 	}
 }
 
@@ -663,14 +820,20 @@ void UAISense_Sight::OnListenerForgetsActor(const FPerceptionListener& Listener,
 	const uint32 ListenerId = Listener.GetListenerID();
 	const uint32 TargetId = ActorToForget.GetUniqueID();
 	
-	for (FAISightQuery& SightQuery : SightQueryQueue)
+	auto ForgetPreviousResult = [&ListenerId, &TargetId](FAISightQuery& SightQuery)->EForEachResult
 	{
 		if (SightQuery.ObserverId == ListenerId && SightQuery.TargetId == TargetId)
 		{
 			// assuming one query per observer-target pair
 			SightQuery.ForgetPreviousResult();
-			break;
+			return EForEachResult::Break;
 		}
+		return EForEachResult::Continue;
+	};
+
+	if (ForEach(SightQueriesInRange, ForgetPreviousResult) == EForEachResult::Continue)
+	{
+		ForEach(SightQueriesOutOfRange, ForgetPreviousResult);
 	}
 }
 
@@ -678,11 +841,15 @@ void UAISense_Sight::OnListenerForgetsAll(const FPerceptionListener& Listener)
 {
 	const uint32 ListenerId = Listener.GetListenerID();
 
-	for (FAISightQuery& SightQuery : SightQueryQueue)
+	auto ForgetPreviousResult = [&ListenerId](FAISightQuery& SightQuery)->EForEachResult
 	{
 		if (SightQuery.ObserverId == ListenerId)
 		{
 			SightQuery.ForgetPreviousResult();
 		}
-	}
+		return EForEachResult::Continue;
+	};
+
+	ForEach(SightQueriesInRange, ForgetPreviousResult);
+	ForEach(SightQueriesOutOfRange, ForgetPreviousResult);
 }
