Merging CL#11314775

Fix build break

#rnx

diff --git a/Engine/Source/Runtime/AIModule/Private/Perception/AISense_Sight.cpp b/Engine/Source/Runtime/AIModule/Private/Perception/AISense_Sight.cpp
--- a/Engine/Source/Runtime/AIModule/Private/Perception/AISense_Sight.cpp
+++ b/Engine/Source/Runtime/AIModule/Private/Perception/AISense_Sight.cpp
@@ -55,7 +55,7 @@ enum class EForEachResult : uint8
 template <typename T, class PREDICATE_CLASS>
 EForEachResult ForEach(T& Array, const PREDICATE_CLASS& Predicate)
 {
-	for (T::ElementType& Element : Array)
+	for (typename T::ElementType& Element : Array)
 	{
 		if (Predicate(Element) == EForEachResult::Break)
 		{
