New sight perception features to help close by awareness
* Added backward offset to the point of view
* Added a near clipping radius

#rnx

diff --git a/Engine/Source/Runtime/AIModule/Classes/Perception/AISenseConfig_Sight.h b/Engine/Source/Runtime/AIModule/Classes/Perception/AISenseConfig_Sight.h
--- a/Engine/Source/Runtime/AIModule/Classes/Perception/AISenseConfig_Sight.h
+++ b/Engine/Source/Runtime/AIModule/Classes/Perception/AISenseConfig_Sight.h
@@ -43,7 +43,15 @@ public:
 	/** If not an InvalidRange (which is the default), we will always be able to see the target that has already been seen if they are within this range of their last seen location. */
 	UPROPERTY(EditDefaultsOnly, BlueprintReadOnly, Category = "Sense", config)
 	float AutoSuccessRangeFromLastSeenLocation;
-		
+
+	/** Point of view move back distance for cone calculation. In conjunction with near clipping distance, this will act as a close by awareness and peripheral vision. */
+	UPROPERTY(EditDefaultsOnly, BlueprintReadOnly, Category = "Sense", config, meta = (UIMin = 0.0, ClampMin = 0.0))
+	float PointOfViewBackwardOffset;
+
+	/** Near clipping distance, to be used with point of view backward offset. Will act as a close by awareness and peripheral vision */
+	UPROPERTY(EditDefaultsOnly, BlueprintReadOnly, Category = "Sense", config, meta = (UIMin = 0.0, ClampMin = 0.0))
+	float NearClippingRadius;
+
 	virtual TSubclassOf<UAISense> GetSenseImplementation() const override;
 
 #if WITH_EDITOR
diff --git a/Engine/Source/Runtime/AIModule/Classes/Perception/AISense_Sight.h b/Engine/Source/Runtime/AIModule/Classes/Perception/AISense_Sight.h
--- a/Engine/Source/Runtime/AIModule/Classes/Perception/AISense_Sight.h
+++ b/Engine/Source/Runtime/AIModule/Classes/Perception/AISense_Sight.h
@@ -131,6 +131,8 @@ public:
 		float SightRadiusSq;
 		float AutoSuccessRangeSqFromLastSeenLocation;
 		float LoseSightRadiusSq;
+		float PointOfViewBackwardOffset;
+		float NearClippingRadiusSq;
 		uint8 AffiliationFlags;
 
 		FDigestedSightProperties();
diff --git a/Engine/Source/Runtime/AIModule/Private/Perception/AISense.cpp b/Engine/Source/Runtime/AIModule/Private/Perception/AISense.cpp
--- a/Engine/Source/Runtime/AIModule/Private/Perception/AISense.cpp
+++ b/Engine/Source/Runtime/AIModule/Private/Perception/AISense.cpp
@@ -226,10 +226,15 @@ void UAISenseConfig_Sight::DescribeSelfToGameplayDebugger(const UAIPerceptionCom
 		DebuggerCategory->AddShape(FGameplayDebuggerShape::MakeCylinder(BodyLocation, LoseSightRadius, 25.0f, LoseSightRangeColor));
 		DebuggerCategory->AddShape(FGameplayDebuggerShape::MakeCylinder(BodyLocation, SightRadius, 25.0f, SightRangeColor));
 
-		const float SightPieLength = FMath::Max(LoseSightRadius, SightRadius);
-		DebuggerCategory->AddShape(FGameplayDebuggerShape::MakeSegment(BodyLocation, BodyLocation + (BodyFacing * SightPieLength), SightRangeColor));
-		DebuggerCategory->AddShape(FGameplayDebuggerShape::MakeSegment(BodyLocation, BodyLocation + (BodyFacing.RotateAngleAxis(PeripheralVisionAngleDegrees, FVector::UpVector) * SightPieLength), SightRangeColor));
-		DebuggerCategory->AddShape(FGameplayDebuggerShape::MakeSegment(BodyLocation, BodyLocation + (BodyFacing.RotateAngleAxis(-PeripheralVisionAngleDegrees, FVector::UpVector) * SightPieLength), SightRangeColor));
+		const float SightPieLength = FMath::Max(LoseSightRadius, SightRadius) + PointOfViewBackwardOffset;
+		const FVector RootLocation = BodyLocation - (BodyFacing * PointOfViewBackwardOffset);
+		const FVector LeftDirection = BodyFacing.RotateAngleAxis(PeripheralVisionAngleDegrees, FVector::UpVector);
+		const FVector RightDirection = BodyFacing.RotateAngleAxis(-PeripheralVisionAngleDegrees, FVector::UpVector);
+		DebuggerCategory->AddShape(FGameplayDebuggerShape::MakeSegment(RootLocation + (BodyFacing * NearClippingRadius), RootLocation + (BodyFacing * SightPieLength), SightRangeColor));
+		DebuggerCategory->AddShape(FGameplayDebuggerShape::MakeSegment(RootLocation + (LeftDirection * NearClippingRadius), RootLocation + (LeftDirection * SightPieLength), SightRangeColor));
+		DebuggerCategory->AddShape(FGameplayDebuggerShape::MakeSegment(RootLocation + (RightDirection * NearClippingRadius), RootLocation + (RightDirection * SightPieLength), SightRangeColor));
+		DebuggerCategory->AddShape(FGameplayDebuggerShape::MakeSegment(RootLocation + (LeftDirection * NearClippingRadius), RootLocation + (BodyFacing * NearClippingRadius), SightRangeColor));
+		DebuggerCategory->AddShape(FGameplayDebuggerShape::MakeSegment(RootLocation + (BodyFacing * NearClippingRadius), RootLocation + (RightDirection * NearClippingRadius), SightRangeColor));
 	}
 }
 #endif // WITH_GAMEPLAY_DEBUGGER
diff --git a/Engine/Source/Runtime/AIModule/Private/Perception/AISense_Sight.cpp b/Engine/Source/Runtime/AIModule/Private/Perception/AISense_Sight.cpp
--- a/Engine/Source/Runtime/AIModule/Private/Perception/AISense_Sight.cpp
+++ b/Engine/Source/Runtime/AIModule/Private/Perception/AISense_Sight.cpp
@@ -33,13 +33,35 @@ static const int32 DefaultMaxTracesPerTick = 6;
 static const int32 DefaultMinQueriesPerTimeSliceCheck = 40;
 
 //----------------------------------------------------------------------//
-// helpers
+// CheckIsTargetInSightCone
+//                     R
+//                   *****  
+//              *             *
+//          *                     *
+//       *                           *
+//     *                               *
+// R *                                   * R
+//    \                                 /
+//     \                               /
+//      \                             /
+//       \             X             /
+//        \                         /
+//         \          ***          /
+//          \     *    N    *     /
+//           \ *               * /
+//            N                 N
+// X = Actor Location
+// R = Sight Radius
+// N = Near Clipping Radius
 //----------------------------------------------------------------------//
-FORCEINLINE_DEBUGGABLE bool CheckIsTargetInSightPie(const FPerceptionListener& Listener, const UAISense_Sight::FDigestedSightProperties& DigestedProps, const FVector& TargetLocation, const float SightRadiusSq)
+FORCEINLINE_DEBUGGABLE bool CheckIsTargetInSightCone(const FPerceptionListener& Listener, const UAISense_Sight::FDigestedSightProperties& DigestedProps, const FVector& TargetLocation, const float SightRadiusSq)
 {
-	if (FVector::DistSquared(Listener.CachedLocation, TargetLocation) <= SightRadiusSq) 
+	const FVector& BaseLocation = FMath::IsNearlyZero(DigestedProps.PointOfViewBackwardOffset) ? Listener.CachedLocation : Listener.CachedLocation - Listener.CachedDirection * DigestedProps.PointOfViewBackwardOffset;
+	const FVector ActorToTarget = TargetLocation - BaseLocation;
+	const float DistToTargetSq = ActorToTarget.SizeSquared();
+	if (DistToTargetSq <= SightRadiusSq && DistToTargetSq >= DigestedProps.NearClippingRadiusSq)
 	{
-		const FVector DirectionToTarget = (TargetLocation - Listener.CachedLocation).GetUnsafeNormal();
+		const FVector DirectionToTarget = ActorToTarget.GetUnsafeNormal();
 		return FVector::DotProduct(DirectionToTarget, Listener.CachedDirection) > DigestedProps.PeripheralVisionAngleCos;
 	}
 
@@ -108,8 +130,10 @@ FAISightTarget::FAISightTarget(AActor* InTarget, FGenericTeamId InTeamId)
 //----------------------------------------------------------------------//
 UAISense_Sight::FDigestedSightProperties::FDigestedSightProperties(const UAISenseConfig_Sight& SenseConfig)
 {
-	SightRadiusSq = FMath::Square(SenseConfig.SightRadius);
-	LoseSightRadiusSq = FMath::Square(SenseConfig.LoseSightRadius);
+	SightRadiusSq = FMath::Square(SenseConfig.SightRadius + SenseConfig.PointOfViewBackwardOffset);
+	LoseSightRadiusSq = FMath::Square(SenseConfig.LoseSightRadius + SenseConfig.PointOfViewBackwardOffset);
+	PointOfViewBackwardOffset = SenseConfig.PointOfViewBackwardOffset;
+	NearClippingRadiusSq = FMath::Square(SenseConfig.NearClippingRadius);
 	PeripheralVisionAngleCos = FMath::Cos(FMath::Clamp(FMath::DegreesToRadians(SenseConfig.PeripheralVisionAngleDegrees), 0.f, PI));
 	AffiliationFlags = SenseConfig.DetectionByAffiliation.GetAsFlags();
 	// keep the special value of FAISystem::InvalidRange (-1.f) if it's set.
@@ -117,7 +141,7 @@ UAISense_Sight::FDigestedSightProperties::FDigestedSightProperties(const UAISens
 }
 
 UAISense_Sight::FDigestedSightProperties::FDigestedSightProperties()
-	: PeripheralVisionAngleCos(0.f), SightRadiusSq(-1.f), AutoSuccessRangeSqFromLastSeenLocation(FAISystem::InvalidRange), LoseSightRadiusSq(-1.f), AffiliationFlags(-1)
+	: PeripheralVisionAngleCos(0.f), SightRadiusSq(-1.f), AutoSuccessRangeSqFromLastSeenLocation(FAISystem::InvalidRange), LoseSightRadiusSq(-1.f), PointOfViewBackwardOffset(0.0f), NearClippingRadiusSq(0.0f), AffiliationFlags(-1)
 {}
 
 //----------------------------------------------------------------------//
@@ -299,7 +323,7 @@ float UAISense_Sight::Update()
 					Listener.RegisterStimulus(TargetActor, FAIStimulus(*this, StimulusStrength, SightQuery->LastSeenLocation, Listener.CachedLocation));
 					SightQuery->bLastResult = true;
 				}
-				else if (CheckIsTargetInSightPie(Listener, PropDigest, TargetLocation, SightRadiusSq))
+				else if (CheckIsTargetInSightCone(Listener, PropDigest, TargetLocation, SightRadiusSq))
 				{
 					SIGHT_LOG_SEGMENT(ListenerPtr->GetOwner(), Listener.CachedLocation, TargetLocation, FColor::Green, TEXT("%s"), *(Target.TargetId.ToString()));
 
