Adding default constructor to FAIDataProviderTypedValue to null its members to prevent a GC crash

diff --git a/Engine/Source/Runtime/AIModule/Classes/DataProviders/AIDataProvider.h b/Engine/Source/Runtime/AIModule/Classes/DataProviders/AIDataProvider.h
--- a/Engine/Source/Runtime/AIModule/Classes/DataProviders/AIDataProvider.h
+++ b/Engine/Source/Runtime/AIModule/Classes/DataProviders/AIDataProvider.h
@@ -75,6 +75,11 @@ struct AIMODULE_API FAIDataProviderTypedValue : public FAIDataProviderValue
 {
 	GENERATED_USTRUCT_BODY()
 
+	FAIDataProviderTypedValue()
+		: PropertyType_DEPRECATED(nullptr)
+		, PropertyType(nullptr)
+	{}
+
 	/** type of value */
 	UPROPERTY()
 	UClass* PropertyType_DEPRECATED;
diff --git a/Engine/Source/Runtime/AIModule/Private/DataProviders/AIDataProvider.cpp b/Engine/Source/Runtime/AIModule/Private/DataProviders/AIDataProvider.cpp
--- a/Engine/Source/Runtime/AIModule/Private/DataProviders/AIDataProvider.cpp
+++ b/Engine/Source/Runtime/AIModule/Private/DataProviders/AIDataProvider.cpp
@@ -77,6 +77,7 @@ bool FAIDataProviderTypedValue::Serialize(FArchive& Ar)
 		if (PropertyType_DEPRECATED)
 		{
 			PropertyType = FFieldClass::GetNameToFieldClassMap().FindRef(PropertyType_DEPRECATED->GetFName());
+			PropertyType_DEPRECATED = nullptr;
 		}
 		else
 		{
