AIPerceptionComponent:  fixed issues with perceptual data associated to actors that are no longer valid and that were broadcasted through delegates. Now using TObjectKey to index the data and added new delegate to allow proper bookkeeping in blueprints (since TWeakObjetPtr are converted to normal pointers).

diff --git a/Engine/Source/Runtime/AIModule/Classes/Perception/AIPerceptionComponent.h b/Engine/Source/Runtime/AIModule/Classes/Perception/AIPerceptionComponent.h
--- a/Engine/Source/Runtime/AIModule/Classes/Perception/AIPerceptionComponent.h
+++ b/Engine/Source/Runtime/AIModule/Classes/Perception/AIPerceptionComponent.h
@@ -4,6 +4,7 @@
 
 #include "CoreMinimal.h"
 #include "UObject/ObjectMacros.h"
+#include "UObject/ObjectKey.h"
 #include "Templates/SubclassOf.h"
 #include "Components/ActorComponent.h"
 #include "EngineDefines.h"
@@ -21,6 +22,29 @@ struct FVisualLogEntry;
 DECLARE_DYNAMIC_MULTICAST_DELEGATE_OneParam(FPerceptionUpdatedDelegate, const TArray<AActor*>&, UpdatedActors);
 DECLARE_DYNAMIC_MULTICAST_DELEGATE_TwoParams(FActorPerceptionUpdatedDelegate, AActor*, Actor, FAIStimulus, Stimulus);
 
+USTRUCT(BlueprintType, meta = (DisplayName = "Sensed Actor's Update Data"))
+struct FActorPerceptionUpdateInfo
+{
+	GENERATED_USTRUCT_BODY()
+
+	/** Id of to the stimulus source */
+	UPROPERTY(BlueprintReadWrite, Category = "AI|Perception")
+	int32 TargetId;
+
+	/** Actor associated to the stimulus (can be null) */
+	UPROPERTY(BlueprintReadWrite, Category = "AI|Perception")
+	TWeakObjectPtr<AActor> Target;
+
+	/** Updated stimulus */
+	UPROPERTY(BlueprintReadWrite, Category = "AI|Perception")
+	FAIStimulus Stimulus;
+
+	FActorPerceptionUpdateInfo() = default;
+	FActorPerceptionUpdateInfo(const int32 TargetId, const TWeakObjectPtr<AActor>& Target, const FAIStimulus& Stimulus);
+};
+
+DECLARE_DYNAMIC_MULTICAST_DELEGATE_OneParam(FActorPerceptionInfoUpdatedDelegate, const FActorPerceptionUpdateInfo&, UpdateInfo);
+
 struct AIMODULE_API FActorPerceptionInfo
 {
 	TWeakObjectPtr<AActor> Target;
@@ -186,10 +210,7 @@ class AIMODULE_API UAIPerceptionComponent : public UActorComponent
 	
 	static const int32 InitialStimuliToProcessArraySize;
 
-	/** The uint64 is the address of the Actor. As we are using the address raw as a key to the TMap there's potential the actor it 
-	 *  points to will be GCed in future (we aren't tagging PerceptualData UPROPERTY intentionally for optimization purposes).
-	 */
-	typedef TMap<uint64, FActorPerceptionInfo> TActorPerceptionContainer;
+	typedef TMap<TObjectKey<AActor>, FActorPerceptionInfo> TActorPerceptionContainer;
 	typedef TActorPerceptionContainer FActorPerceptionContainer;
 
 protected:
@@ -211,18 +232,12 @@ protected:
 
 private:
 	FPerceptionListenerID PerceptionListenerId;
-
-	/**@TODO there is a rare but possible issue here. Actors could be set to Endplay() and GCed between calls to RemoveDeadData, 
-	 * infact EndPlay and GC can occur in same frame. Either we need to take the hit and make PerceptualData a UPROPERTY or we need to come up
-	 * with a different indexing scheme. Currently we could add an new Actor (at the same address as a GCed one), if we then add a new 
-	 * PerceptualData record we could end up merging the new results with the GCed Actor's record instead of making a fresh one.
-	 */
 	FActorPerceptionContainer PerceptualData;
 		
 protected:	
 	struct FStimulusToProcess
 	{
-		AActor* Source;
+		TObjectKey<AActor> Source;
 		FAIStimulus Stimulus;
 
 		FStimulusToProcess(AActor* InSource, const FAIStimulus& InStimulus)
@@ -265,7 +280,7 @@ public:
 	FORCEINLINE FPerceptionListenerID GetListenerId() const { return PerceptionListenerId; }
 
 	FVector GetActorLocation(const AActor& Actor) const;
-	FORCEINLINE const FActorPerceptionInfo* GetActorInfo(const AActor& Actor) const { return PerceptualData.Find(reinterpret_cast<uint64>(&Actor)); }
+	FORCEINLINE const FActorPerceptionInfo* GetActorInfo(const AActor& Actor) const { return PerceptualData.Find(&Actor); }
 	FORCEINLINE FActorPerceptionContainer::TIterator GetPerceptualDataIterator() { return FActorPerceptionContainer::TIterator(PerceptualData); }
 	FORCEINLINE FActorPerceptionContainer::TConstIterator GetPerceptualDataConstIterator() const { return FActorPerceptionContainer::TConstIterator(PerceptualData); }
 
@@ -347,9 +362,34 @@ public:
 	UPROPERTY(BlueprintAssignable)
 	FPerceptionUpdatedDelegate OnPerceptionUpdated;
 
+	/**
+	 * Notifies all bound objects that perception info has been updated for a given target.
+	 * The notification is broadcasted for any received stimulus or on change of state
+	 * according to the stimulus configuration.
+	 * 
+	 * Note - This delegate will not be called if source actor is no longer valid 
+	 * by the time a stimulus gets processed. 
+	 * Use OnTargetPerceptionInfoUpdated providing a source id to handle those cases.
+	 *
+	 * @param	SourceActor	Actor associated to the stimulus (can not be null)
+	 * @param	Stimulus	Updated stimulus
+	 */
 	UPROPERTY(BlueprintAssignable)
 	FActorPerceptionUpdatedDelegate OnTargetPerceptionUpdated;
 
+	/**
+	 * Notifies all bound objects that perception info has been updated for a given target.
+	 * The notification is broadcasted for any received stimulus or on change of state
+	 * according to the stimulus configuration.
+	 *
+	 * Note - This delegate will be called even if source actor is no longer valid 
+	 * by the time a stimulus gets processed so it is better to use source id for bookkeeping.
+	 *
+	 * @param	UpdateInfo	Data structure providing information related to the updated perceptual data
+	 */
+	UPROPERTY(BlueprintAssignable)
+	FActorPerceptionInfoUpdatedDelegate OnTargetPerceptionInfoUpdated;
+
 protected:
 	UE_DEPRECATED(4.11, "Function has been renamed and made public. Please use UpdatePerceptionWhitelist instead")
 	void UpdatePerceptionFilter(FAISenseID Channel, bool bNewValue);
diff --git a/Engine/Source/Runtime/AIModule/Private/Perception/AIPerceptionComponent.cpp b/Engine/Source/Runtime/AIModule/Private/Perception/AIPerceptionComponent.cpp
--- a/Engine/Source/Runtime/AIModule/Private/Perception/AIPerceptionComponent.cpp
+++ b/Engine/Source/Runtime/AIModule/Private/Perception/AIPerceptionComponent.cpp
@@ -43,6 +43,16 @@ FActorPerceptionBlueprintInfo::FActorPerceptionBlueprintInfo(const FActorPercept
 	bIsHostile = Info.bIsHostile;
 }
 
+//----------------------------------------------------------------------//
+// 
+//----------------------------------------------------------------------//
+FActorPerceptionUpdateInfo::FActorPerceptionUpdateInfo(const int32 InTargetId, const TWeakObjectPtr<AActor>& InTarget, const FAIStimulus& InStimulus)
+	: TargetId(InTargetId)
+	, Target(InTarget)
+	, Stimulus(InStimulus)
+{
+}
+
 //----------------------------------------------------------------------//
 // 
 //----------------------------------------------------------------------//
@@ -300,9 +310,9 @@ void UAIPerceptionComponent::GetHostileActors(TArray<AActor*>& OutActors) const
 
 		if (ActorPerceptionInfo.bIsHostile && ActorPerceptionInfo.HasAnyKnownStimulus())
 		{
-			if (ActorPerceptionInfo.Target.IsValid())
+			if (AActor* Actor = ActorPerceptionInfo.Target.Get())
 			{
-				OutActors.Add(ActorPerceptionInfo.Target.Get());
+				OutActors.Add(Actor);
 			}
 			else
 			{
@@ -430,17 +440,20 @@ void UAIPerceptionComponent::ProcessStimuli()
 	}
 
 	const bool bBroadcastEveryTargetUpdate = OnTargetPerceptionUpdated.IsBound();
+	const bool bBroadcastEveryTargetInfoUpdate = OnTargetPerceptionInfoUpdated.IsBound();
 	
 	TArray<AActor*> UpdatedActors;
 	UpdatedActors.Reserve(StimuliToProcess.Num());
 	TArray<AActor*> ActorsToForget;
 	ActorsToForget.Reserve(StimuliToProcess.Num());
+	TArray<TObjectKey<AActor>, TInlineAllocator<8>> DataToRemove;
 
 	for (FStimulusToProcess& SourcedStimulus : StimuliToProcess)
 	{
-		const uint64 SourceAddr = reinterpret_cast<uint64>(SourcedStimulus.Source);
+		const TObjectKey<AActor>& SourceKey = SourcedStimulus.Source;
 
-		FActorPerceptionInfo* PerceptualInfo = PerceptualData.Find(SourceAddr);
+		FActorPerceptionInfo* PerceptualInfo = PerceptualData.Find(SourceKey);
+		AActor* SourceActor = nullptr;
 
 		if (PerceptualInfo == NULL)
 		{
@@ -452,12 +465,20 @@ void UAIPerceptionComponent::ProcessStimuli()
 			}
 			else
 			{
+				SourceActor = CastChecked<AActor>(SourceKey.ResolveObjectPtr(), ECastCheckedType::NullAllowed);
+
+				// no existing perceptual data and source no longer valid: nothing to do with this stimulus
+				if (SourceActor == nullptr)
+				{
+					continue;
+				}
+				
 				// create an entry
-				PerceptualInfo = &PerceptualData.Add(SourceAddr, FActorPerceptionInfo(SourcedStimulus.Source));
+				PerceptualInfo = &PerceptualData.Add(SourceKey, FActorPerceptionInfo(SourceActor));
 				// tell it what's our dominant sense
 				PerceptualInfo->DominantSense = DominantSenseID;
 
-				PerceptualInfo->bIsHostile = AIOwner != NULL && FGenericTeamId::GetAttitude(AIOwner, SourcedStimulus.Source) == ETeamAttitude::Hostile;
+				PerceptualInfo->bIsHostile = AIOwner != NULL && FGenericTeamId::GetAttitude(AIOwner, SourceActor) == ETeamAttitude::Hostile;
 			}
 		}
 
@@ -503,10 +524,24 @@ void UAIPerceptionComponent::ProcessStimuli()
 		// if the new stimulus is "valid" or it's info that "no longer sensed" and it used to be sensed successfully
 		if (bActorInfoUpdated)
 		{
-			UpdatedActors.AddUnique(SourcedStimulus.Source);
-			if (bBroadcastEveryTargetUpdate)
+			// Source Actor is only resolved from SourceKey when required but might already have been resolved for new entry
+			SourceActor = (SourceActor == nullptr) ? CastChecked<AActor>(SourceKey.ResolveObjectPtr(), ECastCheckedType::NullAllowed) : SourceActor;
+			if (SourceActor == nullptr)
+			{
+				DataToRemove.Add(SourceKey);
+			}
+			else
+			{
+				UpdatedActors.AddUnique(SourceActor);
+				if (bBroadcastEveryTargetUpdate)
+				{
+					OnTargetPerceptionUpdated.Broadcast(SourceActor, StimulusStore);
+				}
+			}
+
+			if (bBroadcastEveryTargetInfoUpdate)
 			{
-				OnTargetPerceptionUpdated.Broadcast(SourcedStimulus.Source, StimulusStore);
+				OnTargetPerceptionInfoUpdated.Broadcast(FActorPerceptionUpdateInfo(GetTypeHash(SourceKey), PerceptualInfo->Target, StimulusStore));
 			}
 		}
 	}
@@ -523,10 +558,17 @@ void UAIPerceptionComponent::ProcessStimuli()
 		OnPerceptionUpdated.Broadcast(UpdatedActors);
 	}
 
+	// forget actors that are no longer perceived
 	for (AActor* ActorToForget : ActorsToForget)
 	{
 		ForgetActor(ActorToForget);
 	}
+
+	// remove perceptual info related to stale actors
+	for (const TObjectKey<AActor> SourceKey : DataToRemove)
+	{
+		PerceptualData.Remove(SourceKey);
+	}
 }
 
 void UAIPerceptionComponent::RefreshStimulus(FAIStimulus& StimulusStore, const FAIStimulus& NewStimulus)
@@ -585,7 +627,7 @@ void UAIPerceptionComponent::ForgetActor(AActor* ActorToForget)
 			AIPerceptionSys->OnListenerForgetsActor(*this, *ActorToForget);
 		}
 
-		const int32 NumRemoved = PerceptualData.Remove(reinterpret_cast<uint64>(ActorToForget));
+		const int32 NumRemoved = PerceptualData.Remove(ActorToForget);
 	}
 }
 
@@ -688,9 +730,9 @@ void UAIPerceptionComponent::GetCurrentlyPerceivedActors(TSubclassOf<UAISense> S
 		const bool bCurrentlyPerceived = (SenseToUse == nullptr) ? DataIt->Value.HasAnyCurrentStimulus() : DataIt->Value.IsSenseActive(SenseID);
 		if (bCurrentlyPerceived)
 		{
-			if (DataIt->Value.Target.IsValid())
+			if (AActor* Actor = DataIt->Value.Target.Get())
 			{
-				OutActors.Add(DataIt->Value.Target.Get());
+				OutActors.Add(Actor);
 			}
 		}
 	}
@@ -753,9 +795,10 @@ void UAIPerceptionComponent::DescribeSelfToGameplayDebugger(FGameplayDebuggerCat
 	for (UAIPerceptionComponent::FActorPerceptionContainer::TConstIterator It(GetPerceptualDataConstIterator()); It; ++It)
 	{
 		const FActorPerceptionInfo& ActorPerceptionInfo = It->Value;
-		if (ActorPerceptionInfo.Target.IsValid() && It->Key)
+		const AActor* Target = ActorPerceptionInfo.Target.Get();
+		if (Target != nullptr)
 		{
-			const FVector TargetLocation = ActorPerceptionInfo.Target->GetActorLocation();
+			const FVector TargetLocation = Target->GetActorLocation();
 			for (const FAIStimulus& Stimulus : ActorPerceptionInfo.LastSensedStimuli)
 			{
 				const UAISenseConfig* SenseConfig = GetSenseConfig(Stimulus.Type);
