PR #6792: Change the CrowdFollowingComponent dependencies (Contributed by [REDACTED_PERSON])

- Using IRVOAvoidanceInterface from the movement component instead of UCharacterMovementComponent

diff --git a/Engine/Plugins/Runtime/PhysXVehicles/Source/PhysXVehicles/Private/WheeledVehicleMovementComponent.cpp b/Engine/Plugins/Runtime/PhysXVehicles/Source/PhysXVehicles/Private/WheeledVehicleMovementComponent.cpp
--- a/Engine/Plugins/Runtime/PhysXVehicles/Source/PhysXVehicles/Private/WheeledVehicleMovementComponent.cpp
+++ b/Engine/Plugins/Runtime/PhysXVehicles/Source/PhysXVehicles/Private/WheeledVehicleMovementComponent.cpp
@@ -2005,32 +2005,32 @@ void UWheeledVehicleMovementComponent::CalculateAvoidanceVelocity(float DeltaTim
 
 void UWheeledVehicleMovementComponent::SetAvoidanceGroup(int32 GroupFlags)
 {
-	AvoidanceGroup.SetFlagsDirectly(GroupFlags);
+	SetAvoidanceGroupMask(GroupFlags);
 }
 
 void UWheeledVehicleMovementComponent::SetAvoidanceGroupMask(const FNavAvoidanceMask& GroupMask)
 {
-	AvoidanceGroup.SetFlagsDirectly(GroupMask.Packed);
+	SetAvoidanceGroupMask(GroupMask.Packed);
 }
 
 void UWheeledVehicleMovementComponent::SetGroupsToAvoid(int32 GroupFlags)
 {
-	GroupsToAvoid.SetFlagsDirectly(GroupFlags);
+	SetGroupsToAvoidMask(GroupFlags);
 }
 
 void UWheeledVehicleMovementComponent::SetGroupsToAvoidMask(const FNavAvoidanceMask& GroupMask)
 {
-	GroupsToAvoid.SetFlagsDirectly(GroupMask.Packed);
+	SetGroupsToAvoidMask(GroupMask.Packed);
 }
 
 void UWheeledVehicleMovementComponent::SetGroupsToIgnore(int32 GroupFlags)
 {
-	GroupsToIgnore.SetFlagsDirectly(GroupFlags);
+	SetGroupsToIgnoreMask(GroupFlags);
 }
 
 void UWheeledVehicleMovementComponent::SetGroupsToIgnoreMask(const FNavAvoidanceMask& GroupMask)
 {
-	GroupsToIgnore.SetFlagsDirectly(GroupMask.Packed);
+	SetGroupsToIgnoreMask(GroupMask.Packed);
 }
 
 void UWheeledVehicleMovementComponent::SetAvoidanceEnabled(bool bEnable)
@@ -2099,16 +2099,31 @@ FVector UWheeledVehicleMovementComponent::GetVelocityForRVOConsideration()
 	return Velocity2D;
 }
 
+void UWheeledVehicleMovementComponent::SetAvoidanceGroupMask(int32 GroupFlags)
+{
+	AvoidanceGroup.SetFlagsDirectly(GroupFlags);
+}
+
 int32 UWheeledVehicleMovementComponent::GetAvoidanceGroupMask()
 {
 	return AvoidanceGroup.Packed;
 }
 
+void UWheeledVehicleMovementComponent::SetGroupsToAvoidMask(int32 GroupFlags)
+{
+	GroupsToAvoid.SetFlagsDirectly(GroupFlags);
+}
+
 int32 UWheeledVehicleMovementComponent::GetGroupsToAvoidMask()
 {
 	return GroupsToAvoid.Packed;
 }
 
+void UWheeledVehicleMovementComponent::SetGroupsToIgnoreMask(int32 GroupFlags)
+{
+	GroupsToIgnore.SetFlagsDirectly(GroupFlags);
+}
+
 int32 UWheeledVehicleMovementComponent::GetGroupsToIgnoreMask()
 {
 	return GroupsToIgnore.Packed;
diff --git a/Engine/Plugins/Runtime/PhysXVehicles/Source/PhysXVehicles/Public/WheeledVehicleMovementComponent.h b/Engine/Plugins/Runtime/PhysXVehicles/Source/PhysXVehicles/Public/WheeledVehicleMovementComponent.h
--- a/Engine/Plugins/Runtime/PhysXVehicles/Source/PhysXVehicles/Public/WheeledVehicleMovementComponent.h
+++ b/Engine/Plugins/Runtime/PhysXVehicles/Source/PhysXVehicles/Public/WheeledVehicleMovementComponent.h
@@ -629,8 +629,11 @@ protected:
 	virtual float GetRVOAvoidanceHeight() override;
 	virtual float GetRVOAvoidanceConsiderationRadius() override;
 	virtual FVector GetVelocityForRVOConsideration() override;
+	virtual void SetAvoidanceGroupMask(int32 GroupFlags) override;
 	virtual int32 GetAvoidanceGroupMask() override;
+	virtual void SetGroupsToAvoidMask(int32 GroupFlags) override;
 	virtual int32 GetGroupsToAvoidMask() override;
+	virtual void SetGroupsToIgnoreMask(int32 GroupFlags) override;
 	virtual int32 GetGroupsToIgnoreMask() override;
 	/** END IRVOAvoidanceInterface */
 
diff --git a/Engine/Source/Runtime/AIModule/Classes/Navigation/CrowdFollowingComponent.h b/Engine/Source/Runtime/AIModule/Classes/Navigation/CrowdFollowingComponent.h
--- a/Engine/Source/Runtime/AIModule/Classes/Navigation/CrowdFollowingComponent.h
+++ b/Engine/Source/Runtime/AIModule/Classes/Navigation/CrowdFollowingComponent.h
@@ -4,6 +4,7 @@
 
 #include "CoreMinimal.h"
 #include "UObject/ObjectMacros.h"
+#include "UObject/WeakInterfacePtr.h"
 #include "EngineDefines.h"
 #include "AITypes.h"
 #include "Navigation/PathFollowingComponent.h"
@@ -12,7 +13,7 @@
 #include "CrowdFollowingComponent.generated.h"
 
 class INavLinkCustomInterface;
-class UCharacterMovementComponent;
+class IRVOAvoidanceInterface;
 class UCrowdManager;
 
 namespace ECrowdAvoidanceQuality
@@ -150,8 +151,7 @@ class AIMODULE_API UCrowdFollowingComponent : public UPathFollowingComponent, pu
 
 protected:
 
-	UPROPERTY(transient)
-	UCharacterMovementComponent* CharacterMovement;
+	TWeakInterfacePtr<IRVOAvoidanceInterface> AvoidanceInterface;
 
 public:
 	UPROPERTY()
@@ -159,15 +159,19 @@ public:
 
 protected:
 #if WITH_EDITORONLY_DATA
-	/** DEPRECATED: Group mask for this agent - use property from CharacterMovementComponent instead */
+	UE_DEPRECATED(4.26, "This member is deprecated and no longer used. Use AvoidanceInterface instead.")
+	UPROPERTY(transient)
+	class UCharacterMovementComponent* CharacterMovement;
+
+	/** DEPRECATED: Group mask for this agent - use IRVOAvoidanceInterface instead */
 	UPROPERTY()
 	FNavAvoidanceMask AvoidanceGroup_DEPRECATED;
 
-	/** DEPRECATED: Will avoid other agents if they are in one of specified groups - use property from CharacterMovementComponent instead */
+	/** DEPRECATED: Will avoid other agents if they are in one of specified groups - use IRVOAvoidanceInterface instead */
 	UPROPERTY()
 	FNavAvoidanceMask GroupsToAvoid_DEPRECATED;
 
-	/** DEPRECATED: Will NOT avoid other agents if they are in one of specified groups, higher priority than GroupsToAvoid - use property from CharacterMovementComponent instead */
+	/** DEPRECATED: Will NOT avoid other agents if they are in one of specified groups, higher priority than GroupsToAvoid - use IRVOAvoidanceInterface instead */
 	UPROPERTY()
 	FNavAvoidanceMask GroupsToIgnore_DEPRECATED;
 #endif
diff --git a/Engine/Source/Runtime/AIModule/Private/Navigation/CrowdFollowingComponent.cpp b/Engine/Source/Runtime/AIModule/Private/Navigation/CrowdFollowingComponent.cpp
--- a/Engine/Source/Runtime/AIModule/Private/Navigation/CrowdFollowingComponent.cpp
+++ b/Engine/Source/Runtime/AIModule/Private/Navigation/CrowdFollowingComponent.cpp
@@ -252,9 +252,9 @@ void UCrowdFollowingComponent::SetCrowdRotateToVelocity(bool bEnable)
 
 void UCrowdFollowingComponent::SetAvoidanceGroup(int32 GroupFlags, bool bUpdateAgent)
 {
-	if (CharacterMovement && CharacterMovement->GetAvoidanceGroupMask() != GroupFlags)
+	if (AvoidanceInterface.IsValid() && AvoidanceInterface->GetAvoidanceGroupMask() != GroupFlags)
 	{
-		CharacterMovement->SetAvoidanceGroup(GroupFlags);
+		AvoidanceInterface->SetAvoidanceGroupMask(GroupFlags);
 
 		if (bUpdateAgent)
 		{
@@ -265,9 +265,9 @@ void UCrowdFollowingComponent::SetAvoidanceGroup(int32 GroupFlags, bool bUpdateA
 
 void UCrowdFollowingComponent::SetGroupsToAvoid(int32 GroupFlags, bool bUpdateAgent)
 {
-	if (CharacterMovement && CharacterMovement->GetGroupsToAvoidMask() != GroupFlags)
+	if (AvoidanceInterface.IsValid() && AvoidanceInterface->GetGroupsToAvoidMask() != GroupFlags)
 	{
-		CharacterMovement->SetGroupsToAvoid(GroupFlags);
+		AvoidanceInterface->SetGroupsToAvoidMask(GroupFlags);
 
 		if (bUpdateAgent)
 		{
@@ -278,9 +278,9 @@ void UCrowdFollowingComponent::SetGroupsToAvoid(int32 GroupFlags, bool bUpdateAg
 
 void UCrowdFollowingComponent::SetGroupsToIgnore(int32 GroupFlags, bool bUpdateAgent)
 {
-	if (CharacterMovement && CharacterMovement->GetGroupsToIgnoreMask() != GroupFlags)
+	if (AvoidanceInterface.IsValid() && AvoidanceInterface->GetGroupsToIgnoreMask() != GroupFlags)
 	{
-		CharacterMovement->SetGroupsToIgnore(GroupFlags);
+		AvoidanceInterface->SetGroupsToIgnoreMask(GroupFlags);
 
 		if (bUpdateAgent)
 		{
@@ -291,17 +291,17 @@ void UCrowdFollowingComponent::SetGroupsToIgnore(int32 GroupFlags, bool bUpdateA
 
 int32 UCrowdFollowingComponent::GetAvoidanceGroup() const
 {
-	return CharacterMovement ? CharacterMovement->GetAvoidanceGroupMask() : 0;
+	return AvoidanceInterface.IsValid() ? AvoidanceInterface->GetAvoidanceGroupMask() : 0;
 }
 
 int32 UCrowdFollowingComponent::GetGroupsToAvoid() const
 {
-	return CharacterMovement ? CharacterMovement->GetGroupsToAvoidMask() : 0;
+	return AvoidanceInterface.IsValid() ? AvoidanceInterface->GetGroupsToAvoidMask() : 0;
 }
 
 int32 UCrowdFollowingComponent::GetGroupsToIgnore() const
 {
-	return CharacterMovement ? CharacterMovement->GetGroupsToIgnoreMask() : 0;
+	return AvoidanceInterface.IsValid() ? AvoidanceInterface->GetGroupsToIgnoreMask() : 0;
 }
 
 void UCrowdFollowingComponent::SuspendCrowdSteering(bool bSuspend)
@@ -334,7 +334,7 @@ void UCrowdFollowingComponent::UpdateCachedDirections(const FVector& NewVelocity
 	}
 
 	// CrowdAgentMoveDirection either direction on path or aligned with current velocity
-	const bool bIsNotFalling = (CharacterMovement == nullptr || CharacterMovement->MovementMode != MOVE_Falling);
+	const bool bIsNotFalling = (MovementComp == nullptr || !MovementComp->IsFalling());
 	if (bIsNotFalling)
 	{
 		if (bUpdateDirectMoveVelocity)
@@ -385,7 +385,7 @@ void UCrowdFollowingComponent::ApplyCrowdAgentVelocity(const FVector& NewVelocit
 	bCanCheckMovingTooFar = !bTraversingLink && bIsNearEndOfPath;
 	if (IsCrowdSimulationEnabled() && Status == EPathFollowingStatus::Moving && MovementComp)
 	{
-		const bool bIsNotFalling = (CharacterMovement == nullptr || CharacterMovement->MovementMode != MOVE_Falling);
+		const bool bIsNotFalling = (MovementComp == nullptr || !MovementComp->IsFalling());
 		if (bAffectFallingVelocity || bIsNotFalling)
 		{
 			UpdateCachedDirections(NewVelocity, DestPathCorner, bTraversingLink);
@@ -576,7 +576,8 @@ void UCrowdFollowingComponent::Reset()
 bool UCrowdFollowingComponent::UpdateMovementComponent(bool bForce)
 {
 	bool bRet = Super::UpdateMovementComponent(bForce);
-	CharacterMovement = Cast<UCharacterMovementComponent>(MovementComp);
+
+	AvoidanceInterface = TWeakInterfacePtr<IRVOAvoidanceInterface>(MovementComp);
 
 	return bRet;
 }
@@ -1048,7 +1049,7 @@ FVector UCrowdFollowingComponent::GetMoveFocus(bool bAllowStrafe) const
 
 		// if we're not moving, falling, or don't have a crowd agent move direction, set our focus to ahead of the rotation of our owner to keep the same rotation,
 		// otherwise use the Crowd Agent Move Direction to move in the direction we're supposed to be going
-		const FVector ForwardDir = MovementComp->GetOwner() && ((Status != EPathFollowingStatus::Moving) || (CharacterMovement && (CharacterMovement->MovementMode == MOVE_Falling)) || CrowdAgentMoveDirection.IsNearlyZero()) ?
+		const FVector ForwardDir = MovementComp->GetOwner() && ((Status != EPathFollowingStatus::Moving) || (MovementComp->IsFalling()) || CrowdAgentMoveDirection.IsNearlyZero()) ?
 			MovementComp->GetOwner()->GetActorForwardVector() :
 			CrowdAgentMoveDirection.GetSafeNormal2D();
 
diff --git a/Engine/Source/Runtime/Engine/Classes/GameFramework/CharacterMovementComponent.h b/Engine/Source/Runtime/Engine/Classes/GameFramework/CharacterMovementComponent.h
--- a/Engine/Source/Runtime/Engine/Classes/GameFramework/CharacterMovementComponent.h
+++ b/Engine/Source/Runtime/Engine/Classes/GameFramework/CharacterMovementComponent.h
@@ -2481,8 +2481,11 @@ public:
 	virtual float GetRVOAvoidanceHeight() override;
 	virtual float GetRVOAvoidanceConsiderationRadius() override;
 	virtual FVector GetVelocityForRVOConsideration() override;
+	virtual void SetAvoidanceGroupMask(int32 GroupFlags) override;
 	virtual int32 GetAvoidanceGroupMask() override;
+	virtual void SetGroupsToAvoidMask(int32 GroupFlags) override;
 	virtual int32 GetGroupsToAvoidMask() override;
+	virtual void SetGroupsToIgnoreMask(int32 GroupFlags) override;
 	virtual int32 GetGroupsToIgnoreMask() override;
 	/** END IRVOAvoidanceInterface */
 
diff --git a/Engine/Source/Runtime/Engine/Private/Components/CharacterMovementComponent.cpp b/Engine/Source/Runtime/Engine/Private/Components/CharacterMovementComponent.cpp
--- a/Engine/Source/Runtime/Engine/Private/Components/CharacterMovementComponent.cpp
+++ b/Engine/Source/Runtime/Engine/Private/Components/CharacterMovementComponent.cpp
@@ -3582,16 +3582,31 @@ FVector UCharacterMovementComponent::GetVelocityForRVOConsideration()
 	return Velocity;
 }
 
+void UCharacterMovementComponent::SetAvoidanceGroupMask(int32 GroupFlags)
+{
+	AvoidanceGroup.SetFlagsDirectly(GroupFlags);
+}
+
 int32 UCharacterMovementComponent::GetAvoidanceGroupMask()
 {
 	return AvoidanceGroup.Packed;
 }
 
+void UCharacterMovementComponent::SetGroupsToAvoidMask(int32 GroupFlags)
+{
+	GroupsToAvoid.SetFlagsDirectly(GroupFlags);
+}
+
 int32 UCharacterMovementComponent::GetGroupsToAvoidMask()
 {
 	return GroupsToAvoid.Packed;
 }
 
+void UCharacterMovementComponent::SetGroupsToIgnoreMask(int32 GroupFlags)
+{
+	GroupsToIgnore.SetFlagsDirectly(GroupFlags);
+}
+
 int32 UCharacterMovementComponent::GetGroupsToIgnoreMask()
 {
 	return GroupsToIgnore.Packed;
@@ -9651,32 +9666,32 @@ void UCharacterMovementComponent::CapsuleTouched(UPrimitiveComponent* Overlapped
 
 void UCharacterMovementComponent::SetAvoidanceGroup(int32 GroupFlags)
 {
-	AvoidanceGroup.SetFlagsDirectly(GroupFlags);
+	SetAvoidanceGroupMask(GroupFlags);
 }
 
 void UCharacterMovementComponent::SetAvoidanceGroupMask(const FNavAvoidanceMask& GroupMask)
 {
-	AvoidanceGroup.SetFlagsDirectly(GroupMask.Packed);
+	SetAvoidanceGroupMask(GroupMask.Packed);
 }
 
 void UCharacterMovementComponent::SetGroupsToAvoid(int32 GroupFlags)
 {
-	GroupsToAvoid.SetFlagsDirectly(GroupFlags);
+	SetGroupsToAvoidMask(GroupFlags);
 }
 
 void UCharacterMovementComponent::SetGroupsToAvoidMask(const FNavAvoidanceMask& GroupMask)
 {
-	GroupsToAvoid.SetFlagsDirectly(GroupMask.Packed);
+	SetGroupsToAvoidMask(GroupMask.Packed);
 }
 
 void UCharacterMovementComponent::SetGroupsToIgnore(int32 GroupFlags)
 {
-	GroupsToIgnore.SetFlagsDirectly(GroupFlags);
+	SetGroupsToIgnoreMask(GroupFlags);
 }
 
 void UCharacterMovementComponent::SetGroupsToIgnoreMask(const FNavAvoidanceMask& GroupMask)
 {
-	GroupsToIgnore.SetFlagsDirectly(GroupMask.Packed);
+	SetGroupsToIgnoreMask(GroupMask.Packed);
 }
 
 void UCharacterMovementComponent::SetAvoidanceEnabled(bool bEnable)
@@ -11245,4 +11260,4 @@ void UCharacterMovementComponent::FlushServerMoves()
 ETeleportType UCharacterMovementComponent::GetTeleportType() const
 { 
 	return bJustTeleported || bNetworkLargeClientCorrection ? ETeleportType::TeleportPhysics : ETeleportType::None;
-}
+}
\ No newline at end of file
diff --git a/Engine/Source/Runtime/Engine/Public/AI/RVOAvoidanceInterface.h b/Engine/Source/Runtime/Engine/Public/AI/RVOAvoidanceInterface.h
--- a/Engine/Source/Runtime/Engine/Public/AI/RVOAvoidanceInterface.h
+++ b/Engine/Source/Runtime/Engine/Public/AI/RVOAvoidanceInterface.h
@@ -50,12 +50,21 @@ class ENGINE_API IRVOAvoidanceInterface
 	/** The velocity of the avoiding entity **/
 	virtual FVector GetVelocityForRVOConsideration() = 0;
 
+	/** Store the avoidance group mask **/
+	virtual void SetAvoidanceGroupMask(int32 GroupFlags) = 0;
+
 	/** This actor's avoidance group mask **/
 	virtual int32 GetAvoidanceGroupMask() = 0;
 
+	/** Store the groups to avoid mask **/
+	virtual void SetGroupsToAvoidMask(int32 GroupFlags) = 0;
+
 	/** Agent groups to avoid mask **/
 	virtual int32 GetGroupsToAvoidMask() = 0;
 
+	/** Store the groups to ignore mask **/
+	virtual void SetGroupsToIgnoreMask(int32 GroupFlags) = 0;
+
 	/** Agent groups to ignore **/
 	virtual int32 GetGroupsToIgnoreMask() = 0;
 };
