Workaround for EQS debug drawing throwing asserts for not being on the game thread

diff --git a/Engine/Source/Runtime/AIModule/Private/EnvironmentQuery/EQSRenderingComponent.cpp b/Engine/Source/Runtime/AIModule/Private/EnvironmentQuery/EQSRenderingComponent.cpp
--- a/Engine/Source/Runtime/AIModule/Private/EnvironmentQuery/EQSRenderingComponent.cpp
+++ b/Engine/Source/Runtime/AIModule/Private/EnvironmentQuery/EQSRenderingComponent.cpp
@@ -289,8 +289,15 @@ FPrimitiveSceneProxy* UEQSRenderingComponent::CreateSceneProxy()
 #if  USE_EQS_DEBUGGER
 	if (NewSceneProxy)
 	{
-		EQSRenderingDebugDrawDelegateHelper.InitDelegateHelper(NewSceneProxy);
-		EQSRenderingDebugDrawDelegateHelper.ReregisterDebugDrawDelgate();
+		if (IsInGameThread())
+		{
+			EQSRenderingDebugDrawDelegateHelper.InitDelegateHelper(NewSceneProxy);
+			EQSRenderingDebugDrawDelegateHelper.ReregisterDebugDrawDelgate();
+		}
+		else
+		{
+			UE_LOG(LogEQS, Warning, TEXT("Couldn't register delegate helper on non-game thread"));
+		}
 	}
 #endif
 
