Enable AI Path draw in the gameplay debugger

diff --git a/Engine/Source/Runtime/AIModule/Private/GameplayDebugger/GameplayDebuggerCategory_AI.cpp b/Engine/Source/Runtime/AIModule/Private/GameplayDebugger/GameplayDebuggerCategory_AI.cpp
--- a/Engine/Source/Runtime/AIModule/Private/GameplayDebugger/GameplayDebuggerCategory_AI.cpp
+++ b/Engine/Source/Runtime/AIModule/Private/GameplayDebugger/GameplayDebuggerCategory_AI.cpp
@@ -28,12 +28,14 @@
 namespace FGameplayDebuggerCategoryTweakables
 {
 	int32 bDrawOverheadIcons = 1;
+	bool bDrawPaths = false;
 }
 
 namespace
 {
 	FAutoConsoleVariableRef CVars_GameplayDebuggerCategory_AI[] = {
-		FAutoConsoleVariableRef(TEXT("ai.debug.DrawOverheadIcons"), FGameplayDebuggerCategoryTweakables::bDrawOverheadIcons, TEXT("Should default AI overhead icons be drawn"), ECVF_Default)
+		FAutoConsoleVariableRef(TEXT("ai.debug.DrawOverheadIcons"), FGameplayDebuggerCategoryTweakables::bDrawOverheadIcons, TEXT("Should default AI overhead icons be drawn"), ECVF_Default),
+		FAutoConsoleVariableRef(TEXT("ai.debug.DrawPaths"), FGameplayDebuggerCategoryTweakables::bDrawPaths, TEXT("Should AI paths be drawn"), ECVF_Default)
 	};
 }
 
@@ -315,6 +317,11 @@ void FGameplayDebuggerCategory_AI::DrawData(APlayerController* OwnerPC, FGamepla
 		DrawOverheadInfo(*SelectedActor, CanvasContext);
 	}
 
+	if (FGameplayDebuggerCategoryTweakables::bDrawPaths && DataPack.bIsUsingPathFollowing)
+	{
+		DrawPath(MyWorld);
+	}
+
 	if (DataPack.bHasController)
 	{
 		CanvasContext.Printf(TEXT("Controller Name: {yellow}%s"), *DataPack.ControllerName);
