#UE4  Remove some debug logging

diff --git a/Engine/Source/Runtime/AIModule/Classes/AISubsystem.h b/Engine/Source/Runtime/AIModule/Classes/AISubsystem.h
--- a/Engine/Source/Runtime/AIModule/Classes/AISubsystem.h
+++ b/Engine/Source/Runtime/AIModule/Classes/AISubsystem.h
@@ -22,7 +22,6 @@ private:
 			
 public:
 	UAISubsystem(const FObjectInitializer& ObjectInitializer = FObjectInitializer::Get());
-	~UAISubsystem();
 
 	virtual UWorld* GetWorld() const override;
 
diff --git a/Engine/Source/Runtime/AIModule/Private/AISubsystem.cpp b/Engine/Source/Runtime/AIModule/Private/AISubsystem.cpp
--- a/Engine/Source/Runtime/AIModule/Private/AISubsystem.cpp
+++ b/Engine/Source/Runtime/AIModule/Private/AISubsystem.cpp
@@ -10,8 +10,6 @@ DEFINE_LOG_CATEGORY_STATIC(LogAISub, Log, All);
 UAISubsystem::UAISubsystem(const FObjectInitializer& ObjectInitializer)
 	: Super(ObjectInitializer)
 {
-	UE_LOG(LogAISub, Log, TEXT("%s %s (%p), frame # %llu"), ANSI_TO_TCHAR(__FUNCTION__), *ObjectInitializer.GetObj()->GetClass()->GetName(), this, (uint64)GFrameCounter);
-
 	if (HasAnyFlags(RF_ClassDefaultObject) == false)
 	{
 		AISystem = Cast<UAISystem>(GetOuter());
@@ -45,9 +43,4 @@ ETickableTickType UAISubsystem::GetTickableTickType() const
 TStatId UAISubsystem::GetStatId() const
 {
 	RETURN_QUICK_DECLARE_CYCLE_STAT(UAISubsystem, STATGROUP_Tickables);
-}
-
-UAISubsystem::~UAISubsystem()
-{
-	UE_LOG(LogAISub, Log, TEXT("%s %s (%p), frame # %llu"), ANSI_TO_TCHAR(__FUNCTION__), *GetClass()->GetName(), this, (uint64)GFrameCounter);
-}
+}
\ No newline at end of file
