Moved CheckIsTargetInSightCone() to AIHelpers.h and created a cpp file for helper function implementations.

diff --git a/Engine/Source/Runtime/AIModule/Private/AIHelpers.cpp b/Engine/Source/Runtime/AIModule/Private/AIHelpers.cpp
new file mode 100644
--- /dev/null
+++ b/Engine/Source/Runtime/AIModule/Private/AIHelpers.cpp
@@ -0,0 +1,70 @@
+// Copyright Epic Games, Inc. All Rights Reserved.
+
+#include "AIHelpers.h"
+
+namespace FAISystem
+{
+	FVector FindClosestLocation(const FVector& Origin, const TArray<FVector>& Locations)
+	{
+		FVector BestLocation = FAISystem::InvalidLocation;
+		float BestDistance = FLT_MAX;
+
+		for (const FVector& Candidate : Locations)
+		{
+			const float CurrentDistanceSq = FVector::DistSquared(Origin, Candidate);
+			if (CurrentDistanceSq < BestDistance)
+			{
+				BestDistance = CurrentDistanceSq;
+				BestLocation = Candidate;
+			}
+		}
+
+		return BestLocation;
+	}
+
+	//----------------------------------------------------------------------//
+	// CheckIsTargetInSightCone
+	//                     F
+	//                   *****  
+	//              *             *
+	//          *                     *
+	//       *                           *
+	//     *                               *
+	//   *                                   * 
+	//    \                                 /
+	//     \                               /
+	//      \                             /
+	//       \             X             /
+	//        \                         /
+	//         \          ***          /
+	//          \     *    N    *     /
+	//           \ *               * /
+	//            N                 N
+	//            
+	//           
+	//           
+	//           
+	//
+	// 
+	//                     B 
+	//
+	// X = StartLocation
+	// B = Backward offset
+	// N = Near Clipping Radius (from the StartLocation adjusted by Backward offset)
+	// F = Far Clipping Radius (from the StartLocation adjusted by Backward offset)
+	//----------------------------------------------------------------------//
+	bool CheckIsTargetInSightCone(const FVector& StartLocation, const FVector& ConeDirectionNormal, float PeripheralVisionAngleCos,
+		float ConeDirectionBackwardOffset, float NearClippingRadiusSq, float const FarClippingRadiusSq, const FVector& TargetLocation)
+	{
+		const FVector BaseLocation = FMath::IsNearlyZero(ConeDirectionBackwardOffset) ? StartLocation : StartLocation - ConeDirectionNormal * ConeDirectionBackwardOffset;
+		const FVector ActorToTarget = TargetLocation - BaseLocation;
+		const float DistToTargetSq = ActorToTarget.SizeSquared();
+		if (DistToTargetSq <= FarClippingRadiusSq && DistToTargetSq >= NearClippingRadiusSq)
+		{
+			const FVector DirectionToTarget = ActorToTarget.GetUnsafeNormal();
+			return FVector::DotProduct(DirectionToTarget, ConeDirectionNormal) > PeripheralVisionAngleCos;
+		}
+
+		return false;
+	}
+}
diff --git a/Engine/Source/Runtime/AIModule/Private/Perception/AISense_Sight.cpp b/Engine/Source/Runtime/AIModule/Private/Perception/AISense_Sight.cpp
--- a/Engine/Source/Runtime/AIModule/Private/Perception/AISense_Sight.cpp
+++ b/Engine/Source/Runtime/AIModule/Private/Perception/AISense_Sight.cpp
@@ -6,6 +6,7 @@
 #include "CollisionQueryParams.h"
 #include "Engine/Engine.h"
 #include "AISystem.h"
+#include "AIHelpers.h"
 #include "Perception/AIPerceptionComponent.h"
 #include "VisualLogger/VisualLogger.h"
 #include "Perception/AISightTargetInterface.h"
@@ -32,42 +33,6 @@ DECLARE_CYCLE_STAT(TEXT("Perception Sense: Sight, Remove To Target"), STAT_AI_Se
 static const int32 DefaultMaxTracesPerTick = 6;
 static const int32 DefaultMinQueriesPerTimeSliceCheck = 40;
 
-//----------------------------------------------------------------------//
-// CheckIsTargetInSightCone
-//                     R
-//                   *****  
-//              *             *
-//          *                     *
-//       *                           *
-//     *                               *
-// R *                                   * R
-//    \                                 /
-//     \                               /
-//      \                             /
-//       \             X             /
-//        \                         /
-//         \          ***          /
-//          \     *    N    *     /
-//           \ *               * /
-//            N                 N
-// X = Actor Location
-// R = Sight Radius
-// N = Near Clipping Radius
-//----------------------------------------------------------------------//
-FORCEINLINE_DEBUGGABLE bool CheckIsTargetInSightCone(const FPerceptionListener& Listener, const UAISense_Sight::FDigestedSightProperties& DigestedProps, const FVector& TargetLocation, const float SightRadiusSq)
-{
-	const FVector& BaseLocation = FMath::IsNearlyZero(DigestedProps.PointOfViewBackwardOffset) ? Listener.CachedLocation : Listener.CachedLocation - Listener.CachedDirection * DigestedProps.PointOfViewBackwardOffset;
-	const FVector ActorToTarget = TargetLocation - BaseLocation;
-	const float DistToTargetSq = ActorToTarget.SizeSquared();
-	if (DistToTargetSq <= SightRadiusSq && DistToTargetSq >= DigestedProps.NearClippingRadiusSq)
-	{
-		const FVector DirectionToTarget = ActorToTarget.GetUnsafeNormal();
-		return FVector::DotProduct(DirectionToTarget, Listener.CachedDirection) > DigestedProps.PeripheralVisionAngleCos;
-	}
-
-	return false;
-}
-
 enum class EForEachResult : uint8
 {
 	Break,
@@ -323,7 +288,7 @@ float UAISense_Sight::Update()
 					Listener.RegisterStimulus(TargetActor, FAIStimulus(*this, StimulusStrength, SightQuery->LastSeenLocation, Listener.CachedLocation));
 					SightQuery->bLastResult = true;
 				}
-				else if (CheckIsTargetInSightCone(Listener, PropDigest, TargetLocation, SightRadiusSq))
+				else if (FAISystem::CheckIsTargetInSightCone(Listener.CachedLocation, Listener.CachedDirection, PropDigest.PeripheralVisionAngleCos, PropDigest.PointOfViewBackwardOffset, PropDigest.NearClippingRadiusSq, SightRadiusSq, TargetLocation))
 				{
 					SIGHT_LOG_SEGMENT(ListenerPtr->GetOwner(), Listener.CachedLocation, TargetLocation, FColor::Green, TEXT("%s"), *(Target.TargetId.ToString()));
 
diff --git a/Engine/Source/Runtime/AIModule/Public/AIHelpers.h b/Engine/Source/Runtime/AIModule/Public/AIHelpers.h
--- a/Engine/Source/Runtime/AIModule/Public/AIHelpers.h
+++ b/Engine/Source/Runtime/AIModule/Public/AIHelpers.h
@@ -2,25 +2,44 @@
 
 #pragma once
 
+#include "CoreMinimal.h"
 #include "AITypes.h"
 
 namespace FAISystem
 {
-	FVector FindClosestLocation(const FVector& Origin, const TArray<FVector>& Locations)
-	{
-		FVector BestLocation = FAISystem::InvalidLocation;
-		float BestDistance = FLT_MAX;
+	AIMODULE_API FVector FindClosestLocation(const FVector& Origin, const TArray<FVector>& Locations);
 
-		for (const FVector& Candidate : Locations)
-		{
-			const float CurrentDistanceSq = FVector::DistSquared(Origin, Candidate);
-			if (CurrentDistanceSq < BestDistance)
-			{
-				BestDistance = CurrentDistanceSq;
-				BestLocation = Candidate;
-			}
-		}
-
-		return BestLocation;
-	}
+	//----------------------------------------------------------------------//
+	// CheckIsTargetInSightCone
+	//                     F
+	//                   *****  
+	//              *             *
+	//          *                     *
+	//       *                           *
+	//     *                               *
+	//   *                                   * 
+	//    \                                 /
+	//     \                               /
+	//      \                             /
+	//       \             X             /
+	//        \                         /
+	//         \          ***          /
+	//          \     *    N    *     /
+	//           \ *               * /
+	//            N                 N
+	//            
+	//           
+	//           
+	//           
+	//
+	// 
+	//                     B 
+	//
+	// X = StartLocation
+	// B = Backward offset
+	// N = Near Clipping Radius (from the StartLocation adjusted by Backward offset)
+	// F = Far Clipping Radius (from the StartLocation adjusted by Backward offset)
+	//----------------------------------------------------------------------//
+	AIMODULE_API bool CheckIsTargetInSightCone(const FVector& StartLocation, const FVector& ConeDirectionNormal, float PeripheralVisionAngleCos,
+		float ConeDirectionBackwardOffset, float NearClippingRadiusSq, float const FarClippingRadiusSq, const FVector& TargetLocation);
 }
