[AISubsystem] Updated how we cache the parent AISystem for subsystems that might have been created from a different Outer

#rnx

diff --git a/Engine/Source/Runtime/AIModule/Private/AISubsystem.cpp b/Engine/Source/Runtime/AIModule/Private/AISubsystem.cpp
--- a/Engine/Source/Runtime/AIModule/Private/AISubsystem.cpp
+++ b/Engine/Source/Runtime/AIModule/Private/AISubsystem.cpp
@@ -13,6 +13,14 @@ UAISubsystem::UAISubsystem(const FObjectInitializer& ObjectInitializer)
 	if (HasAnyFlags(RF_ClassDefaultObject) == false)
 	{
 		AISystem = Cast<UAISystem>(GetOuter());
+
+		// Some AISubsystems might have been created from a different Outer
+		// so we fall back on the world to cache the parent AISystem
+		if (AISystem == nullptr)
+		{
+			AISystem = UAISystem::GetCurrentSafe(GetWorld());
+		}
+
 		UE_CLOG(AISystem == nullptr, LogAISub, Error, TEXT("%s is an invalid outer for UAISubsystem instance %s")
 			, *GetName(), *GetNameSafe(GetOuter()));
 #if !(UE_BUILD_SHIPPING || UE_BUILD_TEST)
