Fix reverb build break

#rnx

diff --git a/Engine/Source/Runtime/AIModule/Public/GraphAStar.h b/Engine/Source/Runtime/AIModule/Public/GraphAStar.h
--- a/Engine/Source/Runtime/AIModule/Public/GraphAStar.h
+++ b/Engine/Source/Runtime/AIModule/Public/GraphAStar.h
@@ -154,7 +154,7 @@ struct FGraphAStar
 
 	using FNodeArray = TArray<FSearchNode, FRangeChecklessAllocator<DoRangeCheck>>;
 	using FRangeChecklessSetAllocator = TSetAllocator<TSparseArrayAllocator<FRangeChecklessAllocator<DoRangeCheck>, TInlineAllocator<4, FRangeChecklessAllocator<DoRangeCheck>>>, TInlineAllocator<1, FRangeChecklessAllocator<DoRangeCheck>>>;
-	using FNodeMap = TMap<uint64, int32, FRangeChecklessSetAllocator>;
+	using FNodeMap = TMap<FGraphNodeRef, int32, FRangeChecklessSetAllocator>;
 	using FIndexArray = TArray<int32, FRangeChecklessAllocator<DoRangeCheck>>;
 
 	struct FNodeSorter
