Renamed FEnvQueryResult::IsSuccsessful to FEnvQueryResult::IsSuccessful.

diff --git a/Engine/Source/Runtime/AIModule/Classes/EnvironmentQuery/EnvQueryTypes.h b/Engine/Source/Runtime/AIModule/Classes/EnvironmentQuery/EnvQueryTypes.h
--- a/Engine/Source/Runtime/AIModule/Classes/EnvironmentQuery/EnvQueryTypes.h
+++ b/Engine/Source/Runtime/AIModule/Classes/EnvironmentQuery/EnvQueryTypes.h
@@ -550,7 +550,7 @@ public:
 
 	FORCEINLINE bool IsFinished() const { return Status != EEnvQueryStatus::Processing; }
 	FORCEINLINE bool IsAborted() const { return Status == EEnvQueryStatus::Aborted; }
-	FORCEINLINE bool IsSuccsessful() const { return Status == EEnvQueryStatus::Success; }
+	FORCEINLINE bool IsSuccessful() const { return Status == EEnvQueryStatus::Success; }
 	FORCEINLINE void MarkAsMissingParam() { Status = EEnvQueryStatus::MissingParam; }
 	FORCEINLINE void MarkAsAborted() { Status = EEnvQueryStatus::Aborted; }
 	FORCEINLINE void MarkAsFailed() { Status = EEnvQueryStatus::Failed; }
@@ -558,6 +558,9 @@ public:
 	FORCEINLINE void MarkAsOwnerLost() { Status = EEnvQueryStatus::OwnerLost; }
 
 	FORCEINLINE EEnvQueryStatus::Type GetRawStatus() const { return Status; }
+
+	UE_DEPRECATED(5.0, "FEnvQueryResult::IsSuccsessful is deprecated. Use FEnvQueryResult::IsSuccessful instead.")
+	FORCEINLINE bool IsSuccsessful() const { return IsSuccessful(); }
 };
 
 
