Switched FEQSParametrizedQueryExecutionRequest::Execute over to using UObject& rather than AActor& as the request owner. It has absolutely no impact on how it works since EQS has been switched over to UObject-typed owners a long time ago.

diff --git a/Engine/Source/Runtime/AIModule/Classes/EnvironmentQuery/EnvQueryTypes.h b/Engine/Source/Runtime/AIModule/Classes/EnvironmentQuery/EnvQueryTypes.h
--- a/Engine/Source/Runtime/AIModule/Classes/EnvironmentQuery/EnvQueryTypes.h
+++ b/Engine/Source/Runtime/AIModule/Classes/EnvironmentQuery/EnvQueryTypes.h
@@ -1314,7 +1314,7 @@ struct AIMODULE_API FEQSParametrizedQueryExecutionRequest
 
 	bool IsValid() const { return bInitialized; }
 
-	int32 Execute(AActor& QueryOwner, const UBlackboardComponent* BlackboardComponent, FQueryFinishedSignature& QueryFinishedDelegate);
+	int32 Execute(UObject& QueryOwner, const UBlackboardComponent* BlackboardComponent, FQueryFinishedSignature& QueryFinishedDelegate);
 
 #if WITH_EDITOR
 	void PostEditChangeProperty(UObject& Owner, struct FPropertyChangedEvent& PropertyChangedEvent);
diff --git a/Engine/Source/Runtime/AIModule/Private/EnvironmentQuery/EnvQueryTypes.cpp b/Engine/Source/Runtime/AIModule/Private/EnvironmentQuery/EnvQueryTypes.cpp
--- a/Engine/Source/Runtime/AIModule/Private/EnvironmentQuery/EnvQueryTypes.cpp
+++ b/Engine/Source/Runtime/AIModule/Private/EnvironmentQuery/EnvQueryTypes.cpp
@@ -411,7 +411,7 @@ void FEQSParametrizedQueryExecutionRequest::InitForOwnerAndBlackboard(UObject& O
 	}
 }
 
-int32 FEQSParametrizedQueryExecutionRequest::Execute(AActor& QueryOwner, const UBlackboardComponent* BlackboardComponent, FQueryFinishedSignature& QueryFinishedDelegate)
+int32 FEQSParametrizedQueryExecutionRequest::Execute(UObject& QueryOwner, const UBlackboardComponent* BlackboardComponent, FQueryFinishedSignature& QueryFinishedDelegate)
 {
 	if (bUseBBKeyForQueryTemplate)
 	{
