Made the EQS manager tickable in the editor.
Required to be able to use EQS-based pipelines in the editor-mode

diff --git a/Engine/Source/Runtime/AIModule/Classes/EnvironmentQuery/EnvQueryManager.h b/Engine/Source/Runtime/AIModule/Classes/EnvironmentQuery/EnvQueryManager.h
--- a/Engine/Source/Runtime/AIModule/Classes/EnvironmentQuery/EnvQueryManager.h
+++ b/Engine/Source/Runtime/AIModule/Classes/EnvironmentQuery/EnvQueryManager.h
@@ -165,6 +165,7 @@ class AIMODULE_API UEnvQueryManager : public UAISubsystem, public FSelfRegisteri
 	// FTickableGameObject begin
 	virtual void Tick(float DeltaTime) override;
 	virtual TStatId GetStatId() const override;
+	virtual bool IsTickableInEditor() const override { return true; }
 	// FTickableGameObject end
 
 	/** execute query */
