Fixed comments on EBTNodeResult to make them display properly in the editor

diff --git a/Engine/Source/Runtime/AIModule/Classes/BehaviorTree/BehaviorTreeTypes.h b/Engine/Source/Runtime/AIModule/Classes/BehaviorTree/BehaviorTreeTypes.h
--- a/Engine/Source/Runtime/AIModule/Classes/BehaviorTree/BehaviorTreeTypes.h
+++ b/Engine/Source/Runtime/AIModule/Classes/BehaviorTree/BehaviorTreeTypes.h
@@ -76,10 +76,14 @@ namespace EBTNodeResult
 	// keep in sync with DescribeNodeResult()
 	enum Type
 	{
-		Succeeded,		// finished as success
-		Failed,			// finished as failure
-		Aborted,		// finished aborting = failure
-		InProgress,		// not finished yet
+		// finished as success
+		Succeeded,
+		// finished as failure
+		Failed,
+		// finished aborting = failure
+		Aborted,
+		// not finished yet
+		InProgress,
 	};
 }
 
