Prototype Wrapped Object Pointers (ObjectHandle/ObjectPtr)

Call-site changes for non-scalar wrapped object pointer upgrade that overlap with automated pointer changes.

diff --git a/Engine/Plugins/FX/Niagara/Source/Niagara/Classes/NiagaraSystem.h b/Engine/Plugins/FX/Niagara/Source/Niagara/Classes/NiagaraSystem.h
--- a/Engine/Plugins/FX/Niagara/Source/Niagara/Classes/NiagaraSystem.h
+++ b/Engine/Plugins/FX/Niagara/Source/Niagara/Classes/NiagaraSystem.h
@@ -430,7 +430,7 @@ public:
 
 	FORCEINLINE UNiagaraParameterCollectionInstance* GetParameterCollectionOverride(UNiagaraParameterCollection* Collection)
 	{
-		UNiagaraParameterCollectionInstance** Found = ParameterCollectionOverrides.FindByPredicate(
+		auto* Found = ParameterCollectionOverrides.FindByPredicate(
 			[&](const UNiagaraParameterCollectionInstance* CheckInst)
 		{
 			return CheckInst && Collection == CheckInst->Collection;
diff --git a/Engine/Source/Runtime/AIModule/Classes/BehaviorTree/BlackboardComponent.h b/Engine/Source/Runtime/AIModule/Classes/BehaviorTree/BlackboardComponent.h
--- a/Engine/Source/Runtime/AIModule/Classes/BehaviorTree/BlackboardComponent.h
+++ b/Engine/Source/Runtime/AIModule/Classes/BehaviorTree/BlackboardComponent.h
@@ -349,7 +349,7 @@ bool UBlackboardComponent::SetValue(FBlackboard::FKey KeyID, typename TDataClass
 	uint8* RawData = GetKeyRawData(KeyID) + DataOffset;
 	if (RawData)
 	{
-		UBlackboardKeyType* KeyOb = EntryInfo->KeyType->HasInstance() ? KeyInstances[KeyID] : ToRawPtr(EntryInfo->KeyType);
+		UBlackboardKeyType* KeyOb = EntryInfo->KeyType->HasInstance() ? KeyInstances[KeyID] : EntryInfo->KeyType;
 		const bool bChanged = TDataClass::SetValue((TDataClass*)KeyOb, RawData, Value);
 		if (bChanged)
 		{
@@ -366,7 +366,7 @@ bool UBlackboardComponent::SetValue(FBlackboard::FKey KeyID, typename TDataClass
 						const int32 OtherKeyID = OtherBlackboardAsset ? OtherBlackboardAsset->GetKeyID(EntryInfo->EntryName) : FBlackboard::InvalidKey;
 						if (OtherKeyID != FBlackboard::InvalidKey)
 						{
-							UBlackboardKeyType* OtherKeyOb = EntryInfo->KeyType->HasInstance() ? OtherBlackboard->KeyInstances[OtherKeyID] : ToRawPtr(EntryInfo->KeyType);
+							UBlackboardKeyType* OtherKeyOb = EntryInfo->KeyType->HasInstance() ? OtherBlackboard->KeyInstances[OtherKeyID] : EntryInfo->KeyType;
 							uint8* OtherRawData = OtherBlackboard->GetKeyRawData(OtherKeyID) + DataOffset;
 
 							TDataClass::SetValue((TDataClass*)OtherKeyOb, OtherRawData, Value);
@@ -399,7 +399,7 @@ typename TDataClass::FDataType UBlackboardComponent::GetValue(FBlackboard::FKey
 		return TDataClass::InvalidValue;
 	}
 
-	UBlackboardKeyType* KeyOb = EntryInfo->KeyType->HasInstance() ? KeyInstances[KeyID] : ToRawPtr(EntryInfo->KeyType);
+	UBlackboardKeyType* KeyOb = EntryInfo->KeyType->HasInstance() ? KeyInstances[KeyID] : EntryInfo->KeyType;
 	const uint16 DataOffset = EntryInfo->KeyType->HasInstance() ? sizeof(FBlackboardInstancedKeyMemory) : 0;
 
 	const uint8* RawData = GetKeyRawData(KeyID) + DataOffset;
diff --git a/Engine/Source/Runtime/ClothingSystemRuntimeCommon/Public/ClothingAsset.h b/Engine/Source/Runtime/ClothingSystemRuntimeCommon/Public/ClothingAsset.h
--- a/Engine/Source/Runtime/ClothingSystemRuntimeCommon/Public/ClothingAsset.h
+++ b/Engine/Source/Runtime/ClothingSystemRuntimeCommon/Public/ClothingAsset.h
@@ -166,7 +166,7 @@ public:
 	template<typename ClothConfigType, typename = typename TEnableIf<TIsDerivedFrom<ClothConfigType, UClothConfigBase>::IsDerived>::Type>
 	const ClothConfigType* GetClothConfig() const
 	{
-		const UClothConfigBase* const* const ClothConfig = ClothConfigs.Find(ClothConfigType::StaticClass()->GetFName());
+		auto const* const ClothConfig = ClothConfigs.Find(ClothConfigType::StaticClass()->GetFName());
 		return ClothConfig ? ExactCast<ClothConfigType>(*ClothConfig) : nullptr;
 	}
 
@@ -174,7 +174,7 @@ public:
 	template<typename ClothConfigType, typename = typename TEnableIf<TIsDerivedFrom<ClothConfigType, UClothConfigBase>::IsDerived>::Type>
 	ClothConfigType* GetClothConfig()
 	{
-		UClothConfigBase* const* const ClothConfig = ClothConfigs.Find(ClothConfigType::StaticClass()->GetFName());
+		auto const* const ClothConfig = ClothConfigs.Find(ClothConfigType::StaticClass()->GetFName());
 		return ClothConfig ? ExactCast<ClothConfigType>(*ClothConfig) : nullptr;
 	}
 
diff --git a/Engine/Source/Runtime/Engine/Classes/Engine/LODActor.h b/Engine/Source/Runtime/Engine/Classes/Engine/LODActor.h
--- a/Engine/Source/Runtime/Engine/Classes/Engine/LODActor.h
+++ b/Engine/Source/Runtime/Engine/Classes/Engine/LODActor.h
@@ -272,7 +272,7 @@ public:
 	UStaticMeshComponent* GetStaticMeshComponent() const { return StaticMeshComponent; }
 
 	/** Returns instanced static mesh components that represents the imposters found in this LODActor. */
-	TArray<UInstancedStaticMeshComponent*> GetImpostersStaticMeshComponents() const { TArray<UInstancedStaticMeshComponent*> ISM; ImpostersStaticMeshComponents.GenerateValueArray(ISM); return ISM; }
+	TArray<UInstancedStaticMeshComponent*> GetImpostersStaticMeshComponents() const { TArray<typename TDecay<decltype(ImpostersStaticMeshComponents[nullptr])>::Type> ISM; ImpostersStaticMeshComponents.GenerateValueArray(ISM); return ISM; }
 
 	/** Returns an array of distances that are used to override individual LOD actors min draw distances. */
 	static const TArray<float>& GetHLODDistanceOverride();
diff --git a/Engine/Source/Runtime/Engine/Classes/Materials/Material.h b/Engine/Source/Runtime/Engine/Classes/Materials/Material.h
--- a/Engine/Source/Runtime/Engine/Classes/Materials/Material.h
+++ b/Engine/Source/Runtime/Engine/Classes/Materials/Material.h
@@ -1285,8 +1285,8 @@ public:
 					{
 						if (LayersParameter.ParameterInfo.Name == ParameterName)
 						{
-							Layers = &LayersParameter.Value.Layers;
-							Blends = &LayersParameter.Value.Blends;
+							Layers = &static_cast<const TArray<UMaterialFunctionInterface*>&>(LayersParameter.Value.Layers);
+							Blends = &static_cast<const TArray<UMaterialFunctionInterface*>&>(LayersParameter.Value.Blends);
 #if WITH_EDITOR
 							LayerNames = &LayersParameter.Value.LayerNames;
 #endif
diff --git a/Engine/Source/Runtime/Engine/Classes/Materials/MaterialFunction.h b/Engine/Source/Runtime/Engine/Classes/Materials/MaterialFunction.h
--- a/Engine/Source/Runtime/Engine/Classes/Materials/MaterialFunction.h
+++ b/Engine/Source/Runtime/Engine/Classes/Materials/MaterialFunction.h
@@ -147,7 +147,7 @@ public:
 	virtual UMaterialFunctionInterface* GetBaseFunction() override { return this; }
 	virtual const UMaterialFunctionInterface* GetBaseFunction() const override { return this; }
 #if WITH_EDITORONLY_DATA
-	virtual const TArray<UMaterialExpression*>* GetFunctionExpressions() const override { return &FunctionExpressions; }
+	virtual const TArray<UMaterialExpression*>* GetFunctionExpressions() const override { return &static_cast<const TArray<UMaterialExpression*>&>(FunctionExpressions); }
 #endif
 	virtual const FString* GetDescription() const override { return &Description; }
 
diff --git a/Engine/Source/Runtime/GameplayTasks/Classes/GameplayTasksComponent.h b/Engine/Source/Runtime/GameplayTasks/Classes/GameplayTasksComponent.h
--- a/Engine/Source/Runtime/GameplayTasks/Classes/GameplayTasksComponent.h
+++ b/Engine/Source/Runtime/GameplayTasks/Classes/GameplayTasksComponent.h
@@ -53,8 +53,6 @@ struct FGameplayTaskEventData
 
 DECLARE_DYNAMIC_MULTICAST_DELEGATE_TwoParams(FOnClaimedResourcesChangeSignature, FGameplayResourceSet, NewlyClaimed, FGameplayResourceSet, FreshlyReleased);
 
-typedef TArray<UGameplayTask*>::TConstIterator FConstGameplayTaskIterator;
-
 /**
 *	The core ActorComponent for interfacing with the GameplayAbilities System
 */
@@ -168,9 +166,10 @@ public:
 	FString GetTasksPriorityQueueDescription() const;
 	static FString GetTaskStateName(EGameplayTaskState Value);
 #endif // !(UE_BUILD_SHIPPING || UE_BUILD_TEST)
-	FConstGameplayTaskIterator GetTickingTaskIterator() const;
-	FConstGameplayTaskIterator GetKnownTaskIterator() const;
-	FConstGameplayTaskIterator GetPriorityQueueIterator() const;
+	using GameplayTaskContainerType = decltype(TickingTasks);
+	GameplayTaskContainerType::TConstIterator GetTickingTaskIterator() const;
+	GameplayTaskContainerType::TConstIterator GetKnownTaskIterator() const;
+	GameplayTaskContainerType::TConstIterator GetPriorityQueueIterator() const;
 
 #if ENABLE_VISUAL_LOG
 	void DescribeSelfToVisLog(struct FVisualLogEntry* Snapshot) const;
@@ -201,3 +200,5 @@ private:
 
 	FORCEINLINE bool CanProcessEvents() const { return !bInEventProcessingInProgress && (EventLockCounter == 0); }
 };
+
+typedef UGameplayTasksComponent::GameplayTaskContainerType::TConstIterator FConstGameplayTaskIterator;
diff --git a/Engine/Source/Runtime/LevelSequence/Public/LevelSequence.h b/Engine/Source/Runtime/LevelSequence/Public/LevelSequence.h
--- a/Engine/Source/Runtime/LevelSequence/Public/LevelSequence.h
+++ b/Engine/Source/Runtime/LevelSequence/Public/LevelSequence.h
@@ -144,7 +144,7 @@ public:
 	UObject* FindMetaDataByClass(TSubclassOf<UObject> InClass) const
 	{
 #if WITH_EDITORONLY_DATA
-		UObject* const* Found = MetaDataObjects.FindByPredicate([InClass](UObject* In) { return In && In->GetClass() == InClass; });
+		auto const* Found = MetaDataObjects.FindByPredicate([InClass](UObject* In) { return In && In->GetClass() == InClass; });
 		return Found ? CastChecked<UObject>(*Found) : nullptr;
 #endif
 		return nullptr;
@@ -218,7 +218,7 @@ public:
 	MetaDataType* FindMetaData() const
 	{
 		UClass* PredicateClass = MetaDataType::StaticClass();
-		UObject* const* Found = MetaDataObjects.FindByPredicate([PredicateClass](UObject* In){ return In && In->GetClass() == PredicateClass; });
+		auto const* Found = MetaDataObjects.FindByPredicate([PredicateClass](UObject* In){ return In && In->GetClass() == PredicateClass; });
 		return Found ? CastChecked<MetaDataType>(*Found) : nullptr;
 	}
 
diff --git a/Engine/Source/Runtime/MovieScene/Public/MovieScene.h b/Engine/Source/Runtime/MovieScene/Public/MovieScene.h
--- a/Engine/Source/Runtime/MovieScene/Public/MovieScene.h
+++ b/Engine/Source/Runtime/MovieScene/Public/MovieScene.h
@@ -358,10 +358,10 @@ public:
 	 * DO NOT USE DIRECTLY
 	 * STL-like iterators to enable range-based for loop support.
 	 */
-	FORCEINLINE TArray<UMovieSceneNodeGroup*>::RangedForIteratorType      begin()	{ return NodeGroups.begin(); }
-	FORCEINLINE TArray<UMovieSceneNodeGroup*>::RangedForConstIteratorType begin()	const { return NodeGroups.begin(); }
-	FORCEINLINE TArray<UMovieSceneNodeGroup*>::RangedForIteratorType      end()	{ return NodeGroups.end(); }
-	FORCEINLINE TArray<UMovieSceneNodeGroup*>::RangedForConstIteratorType end()	const { return NodeGroups.end(); }
+	FORCEINLINE auto begin()	{ return NodeGroups.begin(); }
+	FORCEINLINE auto begin()	const { return NodeGroups.begin(); }
+	FORCEINLINE auto end()	{ return NodeGroups.end(); }
+	FORCEINLINE auto end()	const { return NodeGroups.end(); }
 #endif
 };
 
@@ -1267,4 +1267,4 @@ private:
 #endif
 
 		
-};
+};
\ No newline at end of file
diff --git a/Engine/Source/Runtime/MovieScene/Public/MovieSceneBinding.h b/Engine/Source/Runtime/MovieScene/Public/MovieSceneBinding.h
--- a/Engine/Source/Runtime/MovieScene/Public/MovieSceneBinding.h
+++ b/Engine/Source/Runtime/MovieScene/Public/MovieSceneBinding.h
@@ -116,7 +116,7 @@ struct FMovieSceneBinding
 	 */
 	TArray<UMovieSceneTrack*> StealTracks()
 	{
-		TArray<UMovieSceneTrack*> Empty;
+		decltype(Tracks) Empty;
 		Swap(Empty, Tracks);
 		return Empty;
 	}
