[Navigation] Prototype and first pass at pathfinding using lanes of a zone graph.
The implementation uses the existing generic templated AStar.
ZoneGraphTestingActor -> Now possible to connect another testing actor to pathfind to and display the result
ZoneGraphTestingActor -> Addition of a checkbox to draw "linked lanes"

diff --git a/Engine/Source/Runtime/AIModule/Public/GraphAStar.h b/Engine/Source/Runtime/AIModule/Public/GraphAStar.h
--- a/Engine/Source/Runtime/AIModule/Public/GraphAStar.h
+++ b/Engine/Source/Runtime/AIModule/Public/GraphAStar.h
@@ -2,6 +2,8 @@
 
 #pragma once
 
+#include "AI/NavigationSystemBase.h" // Needed for LogAStar
+
 struct FGraphAStarDefaultPolicy
 {
 	static const int32 NodePoolSize = 64;
@@ -106,12 +108,12 @@ struct CQuery##QueryFunctionName	\
 template <typename TemplateClass> static FORCEINLINE typename TEnableIf< TModels<CQuery##QueryFunctionName, TemplateClass>::Value, QueryReturnType>::Type QueryFunctionName(const TemplateClass& Obj, QueryParam) { return QueryImpl; } \
 template <typename TemplateClass> static FORCEINLINE typename TEnableIf<!TModels<CQuery##QueryFunctionName, TemplateClass>::Value, QueryReturnType>::Type QueryFunctionName(const TemplateClass& Obj, QueryParam) { return QueryDefaultImpl; }
 
-template< bool DoRangeCheck>
+template<bool DoRangeCheck>
 class FRangeChecklessAllocator : public FDefaultAllocator
 {
 public:
 
-	/** set to false if you don't want to lose performance on range checks in performance-critical path finding code. */
+	/** Set to false if you don't want to lose performance on range checks in performance-critical pathfinding code. */
 	enum { RequireRangeCheck = DoRangeCheck };
 };
 template <> struct TAllocatorTraits<FRangeChecklessAllocator<true>> : TAllocatorTraits<FDefaultAllocator> {};
@@ -122,7 +124,7 @@ template <> struct TAllocatorTraits<FRangeChecklessAllocator<false>> : TAllocato
  *
  *	TGraph holds graph representation. Needs to implement functions:
  *		bool IsValidRef(FNodeRef NodeRef) const;												- returns whether given node identyfication is correct
- *      FNodeRef GetNeighbour(const FSearchNode& NodeRef, const int32 NeighbourIndex) const;		- returns neighbour ref
+ *		FNodeRef GetNeighbour(const FSearchNode& NodeRef, const int32 NeighbourIndex) const;	- returns neighbour ref
  *
  *	it also needs to specify node type
  *		FNodeRef		- type used as identification of nodes in the graph
@@ -131,11 +133,11 @@ template <> struct TAllocatorTraits<FRangeChecklessAllocator<false>> : TAllocato
  *		float GetHeuristicScale() const;														- used as GetHeuristicCost's multiplier
  *		float GetHeuristicCost(const FSearchNode& StartNode, const FSearchNode& EndNode) const;	- estimate of cost from StartNode to EndNode from a search node
  *		float GetTraversalCost(const FSearchNode& StartNode, const FSearchNode& EndNode) const;	- real cost of traveling from StartNode directly to EndNode from a search node
- *		bool IsTraversalAllowed(const FNodeRef NodeA, const FNodeRef NodeB) const;			- whether traversing given edge is allowed from a NodeRef
+ *		bool IsTraversalAllowed(const FNodeRef NodeA, const FNodeRef NodeB) const;				- whether traversing given edge is allowed from a NodeRef
  *		bool WantsPartialSolution() const;														- whether to accept solutions that do not reach the goal
  *
  *		// Backward compatible methods, please use the FSearchNode version. If the FSearchNode version are implemented, these methods will not be called at all.
- *      FNodeRef GetNeighbour(const FNodeRef NodeRef, const int32 NeighbourIndex) const;		- returns neighbour ref
+ *		FNodeRef GetNeighbour(const FNodeRef NodeRef, const int32 NeighbourIndex) const;		- returns neighbour ref
  *		float GetHeuristicCost(const FNodeRef StartNodeRef, const FNodeRef EndNodeRef) const;	- estimate of cost from StartNode to EndNode from a NodeRef
  *		float GetTraversalCost(const FNodeRef StartNodeRef, const FNodeRef EndNodeRef) const;	- real cost of traveling from StartNode directly to EndNode from a NodeRef
  *
@@ -143,8 +145,8 @@ template <> struct TAllocatorTraits<FRangeChecklessAllocator<false>> : TAllocato
  *		int32 GetNeighbourCount(FNodeRef NodeRef) const;										- returns number of neighbours that the graph node identified with NodeRef has, it is ok if not implemented, the logic will stop calling GetNeighbour once it received an invalid noderef
  *		bool ShouldIgnoreClosedNodes() const;													- whether to revisit closed node or not
  *		bool ShouldIncludeStartNodeInPath() const;												- whether to put the start node in the resulting path
- *		int32 GetMaxSearchNodes() const;														- whether to limit the number of search node to a maximum
- *		float GetCostLimit() constk																- whether to limit the search to a maximum cost
+ *		int32 GetMaxSearchNodes() const;														- whether to limit the number of search nodes to a maximum
+ *		float GetCostLimit() const																- whether to limit the search to a maximum cost
  */
 template<typename TGraph, typename Policy = FGraphAStarDefaultPolicy, typename TSearchNode = FGraphAStarDefaultNode<TGraph>, bool DoRangeCheck = false >
 struct FGraphAStar
@@ -190,7 +192,7 @@ struct FGraphAStar
 			return NewNode;
 		}
 
-		FORCEINLINE FSearchNode& FindOrAdd(const FGraphNodeRef NodeRef)
+		FORCEINLINE_DEBUGGABLE FSearchNode& FindOrAdd(const FGraphNodeRef NodeRef)
 		{
 			// first find if node already exist in node map
 			const int32 NotInMapIndex = -1;
@@ -214,12 +216,6 @@ struct FGraphAStar
 			return IndexPtr ? &(*this)[*IndexPtr] : nullptr;
 		}
 
-		UE_DEPRECATED(4.15, "This function is now deprecated, please use FindOrAdd instead")
-			FSearchNode& Get(const FGraphNodeRef NodeRef)
-		{
-			return FindOrAdd(NodeRef);
-		}
-
 		FORCEINLINE void Reset()
 		{
 			Super::Reset(Policy::NodePoolSize);
@@ -273,13 +269,6 @@ struct FGraphAStar
 			NodePool[SearchNodeIndex].MarkNotOpened();
 			return SearchNodeIndex;
 		}
-
-		UE_DEPRECATED(4.15, "This function is now deprecated, please use PopIndex instead")
-			FSearchNode& Pop(bool bAllowShrinking = true)
-		{
-			const int32 Index = PopIndex(bAllowShrinking);
-			return NodePool[Index];
-		}
 	};
 
 	const TGraph& Graph;
@@ -316,12 +305,13 @@ struct FGraphAStar
 	 * returns true if loop should be continued
 	 */
 	template<typename TQueryFilter>
-	FORCEINLINE bool ProcessSingleNode(const FSearchNode& EndNode, const bool bIsBound, const TQueryFilter& Filter, int32& OutBestNodeIndex, float& OutBestNodeCost)
+	FORCEINLINE_DEBUGGABLE bool ProcessSingleNode(const FSearchNode& EndNode, const bool bIsBound, const TQueryFilter& Filter, int32& OutBestNodeIndex, float& OutBestNodeCost)
 	{
 		// Pop next best node and put it on closed list
 		const int32 ConsideredNodeIndex = OpenList.PopIndex();
 		FSearchNode& ConsideredNodeUnsafe = NodePool[ConsideredNodeIndex];
 		ConsideredNodeUnsafe.MarkClosed();
+		UE_GRAPH_ASTAR_LOG(Display, TEXT(" Best node %i (end node %i)"), ConsideredNodeUnsafe.NodeRef, EndNode.NodeRef);
 
 		// We're there, store and move to result composition
 		if (bIsBound && (ConsideredNodeUnsafe.NodeRef == EndNode.NodeRef))
@@ -335,6 +325,8 @@ struct FGraphAStar
 
 		// consider every neighbor of BestNode
 		const int32 NeighbourCount = GetNeighbourCount(Graph, ConsideredNodeUnsafe.NodeRef);
+		UE_GRAPH_ASTAR_LOG(Display, TEXT(" Found %i neighbor"), NeighbourCount);
+
 		for (int32 NeighbourNodeIndex = 0; NeighbourNodeIndex < NeighbourCount; ++NeighbourNodeIndex)
 		{
 			const auto& NeighbourRef = GetNeighbour(Graph, NodePool[ConsideredNodeIndex], NeighbourNodeIndex);
@@ -359,7 +351,7 @@ struct FGraphAStar
 				|| NeighbourRef == NodePool[ConsideredNodeIndex].NodeRef
 				|| Filter.IsTraversalAllowed(NodePool[ConsideredNodeIndex].NodeRef, NeighbourRef) == false)
 			{
-				UE_GRAPH_ASTAR_LOG(Warning, TEXT("Filtered %lld from %lld"), (uint64)NeighbourRef, (uint64)NodePool[ConsideredNodeIndex].NodeRef);
+				UE_GRAPH_ASTAR_LOG(Display, TEXT("    Filtered %lld from %lld"), (uint64)NeighbourRef, (uint64)NodePool[ConsideredNodeIndex].NodeRef);
 				continue;
 			}
 
@@ -371,7 +363,7 @@ struct FGraphAStar
 				ExistingNeighbourNode = NodePool.Find(NeighbourRef);
 				if (!ExistingNeighbourNode)
 				{
-					UE_GRAPH_ASTAR_LOG(Warning, TEXT("Reach Limit %lld from %lld"), (uint64)NeighbourRef, (uint64)NodePool[ConsideredNodeIndex].NodeRef);
+					UE_GRAPH_ASTAR_LOG(Display, TEXT("    Reach Limit %lld from %lld"), (uint64)NeighbourRef, (uint64)NodePool[ConsideredNodeIndex].NodeRef);
 					continue;
 				}
 			}
@@ -380,7 +372,7 @@ struct FGraphAStar
 			// check condition to avoid search of closed nodes even if they could have lower cost
 			if (ShouldIgnoreClosedNodes(Filter) && NeighbourNode.IsClosed())
 			{
-				UE_GRAPH_ASTAR_LOG(Warning, TEXT("Skipping closed %lld from %lld"), (uint64)NeighbourRef, (uint64)NodePool[ConsideredNodeIndex].NodeRef);
+				UE_GRAPH_ASTAR_LOG(Display, TEXT("    Skipping closed %lld from %lld"), (uint64)NeighbourRef, (uint64)NodePool[ConsideredNodeIndex].NodeRef);
 				continue;
 			}
 
@@ -394,7 +386,7 @@ struct FGraphAStar
 			// check against cost limit
 			if (HasExceededCostLimit(Filter, NewTotalCost))
 			{
-				UE_GRAPH_ASTAR_LOG(Warning, TEXT("Skipping reach cost limit %lld from %lld cost %f total %f prev cost %f limit %f"), (uint64)NeighbourRef, (uint64)NodePool[ConsideredNodeIndex].NodeRef, NewTraversalCost, NewTotalCost, NeighbourNode.TotalCost, GetCostLimit(Filter));
+				UE_GRAPH_ASTAR_LOG(Display, TEXT("    Skipping reach cost limit %lld from %lld cost %f total %f prev cost %f limit %f"), (uint64)NeighbourRef, (uint64)NodePool[ConsideredNodeIndex].NodeRef, NewTraversalCost, NewTotalCost, NeighbourNode.TotalCost, GetCostLimit(Filter));
 				continue;
 			}
 
@@ -402,7 +394,7 @@ struct FGraphAStar
 			if (NewTotalCost >= NeighbourNode.TotalCost)
 			{
 				// if not, skip
-				UE_GRAPH_ASTAR_LOG(Warning, TEXT("Skipping new cost higher %lld from %lld cost %f total %f prev cost %f"), (uint64)NeighbourRef, (uint64)NodePool[ConsideredNodeIndex].NodeRef, NewTraversalCost, NewTotalCost, NeighbourNode.TotalCost);
+				UE_GRAPH_ASTAR_LOG(Display, TEXT("    Skipping new cost higher %lld from %lld cost %f total %f prev cost %f"), (uint64)NeighbourRef, (uint64)NodePool[ConsideredNodeIndex].NodeRef, NewTraversalCost, NewTotalCost, NeighbourNode.TotalCost);
 				continue;
 			}
 
@@ -415,12 +407,12 @@ struct FGraphAStar
 
 			if (NeighbourNode.IsOpened() == false)
 			{
-			    UE_GRAPH_ASTAR_LOG(Warning, TEXT("Pushing %lld from %lld cost %f total %f"), (uint64)NeighbourRef, (uint64)NodePool[ConsideredNodeIndex].NodeRef, NewTraversalCost, NewTotalCost);
+			    UE_GRAPH_ASTAR_LOG(Display, TEXT("    Pushing %lld from %lld cost %f total %f"), (uint64)NeighbourRef, (uint64)NodePool[ConsideredNodeIndex].NodeRef, NewTraversalCost, NewTotalCost);
 				OpenList.Push(NeighbourNode);
 			}
 			else
 			{
-			    UE_GRAPH_ASTAR_LOG(Warning, TEXT("Modifying %lld from %lld cost %f total %f"), (uint64)NeighbourRef, (uint64)NodePool[ConsideredNodeIndex].NodeRef, NewTraversalCost, NewTotalCost);
+			    UE_GRAPH_ASTAR_LOG(Display, TEXT("    Modifying %lld from %lld cost %f total %f"), (uint64)NeighbourRef, (uint64)NodePool[ConsideredNodeIndex].NodeRef, NewTraversalCost, NewTotalCost);
 				OpenList.Modify(NeighbourNode);
 			}
 
@@ -431,7 +423,7 @@ struct FGraphAStar
 			// will be the one estimated-closest to the goal
 			if (NewHeuristicCost < OutBestNodeCost)
 			{
-				UE_GRAPH_ASTAR_LOG(Warning, TEXT("New best path %lld from %lld new best heuristic %f prev best heuristic %f"), (uint64)NeighbourRef, (uint64)NodePool[ConsideredNodeIndex].NodeRef, NewHeuristicCost, OutBestNodeCost);
+				UE_GRAPH_ASTAR_LOG(Display, TEXT("    New best path %lld from %lld new best heuristic %f prev best heuristic %f"), (uint64)NeighbourRef, (uint64)NodePool[ConsideredNodeIndex].NodeRef, NewHeuristicCost, OutBestNodeCost);
 				OutBestNodeCost = NewHeuristicCost;
 				OutBestNodeIndex = NeighbourNode.SearchNodeIndex;
 			}
@@ -448,6 +440,9 @@ struct FGraphAStar
 	template<typename TQueryFilter, typename TResultPathInfo = TArray<FGraphNodeRef> >
 	EGraphAStarResult FindPath(const FSearchNode& StartNode, const FSearchNode& EndNode, const TQueryFilter& Filter, TResultPathInfo& OutPath)
 	{
+		UE_GRAPH_ASTAR_LOG(Display, TEXT(""));
+		UE_GRAPH_ASTAR_LOG(Display, TEXT("Starting FindPath request..."));
+
 		if (!(Graph.IsValidRef(StartNode.NodeRef) && Graph.IsValidRef(EndNode.NodeRef)))
 		{
 			return SearchFail;
@@ -514,10 +509,12 @@ struct FGraphAStar
 			OutPath.AddZeroed(PathLength);
 
 			// store the path
+			UE_GRAPH_ASTAR_LOG(Display, TEXT("Storing path result (length=%i)..."), PathLength);
 			SearchNodeIndex = BestNodeIndex;
 			int32 ResultNodeIndex = PathLength - 1;
 			do
 			{
+				UE_GRAPH_ASTAR_LOG(Display, TEXT("  NodeRef %i"), NodePool[SearchNodeIndex].NodeRef);
 				SetPathInfo(OutPath, ResultNodeIndex--, NodePool[SearchNodeIndex]);
 				SearchNodeIndex = NodePool[SearchNodeIndex].ParentNodeIndex;
 			} while (ResultNodeIndex >= 0);
diff --git a/Engine/Source/Runtime/Engine/Classes/AI/NavigationSystemBase.h b/Engine/Source/Runtime/Engine/Classes/AI/NavigationSystemBase.h
--- a/Engine/Source/Runtime/Engine/Classes/AI/NavigationSystemBase.h
+++ b/Engine/Source/Runtime/Engine/Classes/AI/NavigationSystemBase.h
@@ -23,7 +23,7 @@ class UNavAreaBase;
 ENGINE_API DECLARE_LOG_CATEGORY_EXTERN(LogNavigation, Warning, All);
 ENGINE_API DECLARE_LOG_CATEGORY_EXTERN(LogNavigationDataBuild, Log, All);
 ENGINE_API DECLARE_LOG_CATEGORY_EXTERN(LogNavLink, Warning, All);
-ENGINE_API DECLARE_LOG_CATEGORY_EXTERN(LogAStar, Warning, All);
+ENGINE_API DECLARE_LOG_CATEGORY_EXTERN(LogAStar, Display, All);
 
 UENUM()
 enum class FNavigationSystemRunMode : uint8
