[DefaultManagerInstanceTracker] Moved manager unregistration to OnPostWorldCleanup instead of OnWorldCleanup to allow actors and components to unregister before the manager instance gets removed from the map container.

diff --git a/Engine/Source/Runtime/AIModule/Public/DefaultManagerInstanceTracker.h b/Engine/Source/Runtime/AIModule/Public/DefaultManagerInstanceTracker.h
--- a/Engine/Source/Runtime/AIModule/Public/DefaultManagerInstanceTracker.h
+++ b/Engine/Source/Runtime/AIModule/Public/DefaultManagerInstanceTracker.h
@@ -16,7 +16,7 @@ struct TDefaultManagerInstanceTracker
 
 	TDefaultManagerInstanceTracker()
 	{
-		FWorldDelegates::OnWorldCleanup.AddRaw(this, &TDefaultManagerInstanceTracker::OnWorldCleanup);
+		FWorldDelegates::OnPostWorldCleanup.AddRaw(this, &TDefaultManagerInstanceTracker::OnPostWorldCleanup);
 	}
 
 	FManager* GetManagerInstance(UWorld& World)
@@ -48,7 +48,7 @@ struct TDefaultManagerInstanceTracker
 		return NewInstance;
 	}
 
-	void OnWorldCleanup(UWorld* World, bool /*bSessionEnded*/, bool /*bCleanupResources*/)
+	void OnPostWorldCleanup(UWorld* World, bool /*bSessionEnded*/, bool /*bCleanupResources*/)
 	{
 		FManager** FoundInstance = WorldToInstanceMap.Find(World);
 		if (FoundInstance)
