[Navigation] Lane pathfinder prototype, second pass
Computing exact traversal distance on the start and end lanes
Allow the pathfinder to pick adjacent lanes (allow to start in an opposite direction)
Addition of debug drawing tools to draw a lane path, lanes and parts of lanes
ZoneGraph -> Addition of left and right lane links to adjacent lanes

diff --git a/Engine/Source/Runtime/AIModule/Public/GraphAStar.h b/Engine/Source/Runtime/AIModule/Public/GraphAStar.h
--- a/Engine/Source/Runtime/AIModule/Public/GraphAStar.h
+++ b/Engine/Source/Runtime/AIModule/Public/GraphAStar.h
@@ -279,6 +279,7 @@ struct FGraphAStar
 
 	// TGraph optionally implemented wrapper methods
 	DECLARE_OPTIONALLY_IMPLEMENTED_TEMPLATE_CLASS_FUNCTION_1PARAM_CONST(int32, GetNeighbourCount, const FGraphNodeRef, NO_COUNT);
+	DECLARE_OPTIONALLY_IMPLEMENTED_TEMPLATE_CLASS_FUNCTION_1PARAM_CONST(int32, GetNeighbourCountV2, const FSearchNode&, Obj.GetNeighbourCount(Param1.NodeRef));
 	DECLARE_OPTIONALLY_IMPLEMENTED_TEMPLATE_CLASS_FUNCTION_2PARAMS_CONST(FGraphNodeRef, GetNeighbour, const FSearchNode&, const int32, Obj.GetNeighbour(Param1.NodeRef,Param2));
 
 	// TQueryFilter optionally implemented wrapper methods
@@ -324,7 +325,7 @@ struct FGraphAStar
 		const float HeuristicScale = Filter.GetHeuristicScale();
 
 		// consider every neighbor of BestNode
-		const int32 NeighbourCount = GetNeighbourCount(Graph, ConsideredNodeUnsafe.NodeRef);
+		const int32 NeighbourCount = GetNeighbourCountV2(Graph, ConsideredNodeUnsafe);
 		UE_GRAPH_ASTAR_LOG(Display, TEXT(" Found %i neighbor"), NeighbourCount);
 
 		for (int32 NeighbourNodeIndex = 0; NeighbourNodeIndex < NeighbourCount; ++NeighbourNodeIndex)
