Added redirection of BlackboardComponent's to AIController's vislog when the BBComp gets created as part of AAIController::UseBlackboard call

diff --git a/Engine/Source/Runtime/AIModule/Private/AIController.cpp b/Engine/Source/Runtime/AIModule/Private/AIController.cpp
--- a/Engine/Source/Runtime/AIModule/Private/AIController.cpp
+++ b/Engine/Source/Runtime/AIModule/Private/AIController.cpp
@@ -1010,6 +1010,7 @@ bool AAIController::UseBlackboard(UBlackboardData* BlackboardAsset, UBlackboardC
 	if (Blackboard == nullptr)
 	{
 		Blackboard = NewObject<UBlackboardComponent>(this, TEXT("BlackboardComponent"));
+		REDIRECT_OBJECT_TO_VLOG(Blackboard, this);
 		if (Blackboard != nullptr)
 		{
 			InitializeBlackboard(*Blackboard, *BlackboardAsset);
