Speed up FEQSHelpers::RunNavProjection by removing elements from Points array after the workload result processing loop.

diff --git a/Engine/Source/Runtime/AIModule/Private/EnvironmentQuery/EnvQueryTraceHelpers.cpp b/Engine/Source/Runtime/AIModule/Private/EnvironmentQuery/EnvQueryTraceHelpers.cpp
--- a/Engine/Source/Runtime/AIModule/Private/EnvironmentQuery/EnvQueryTraceHelpers.cpp
+++ b/Engine/Source/Runtime/AIModule/Private/EnvironmentQuery/EnvQueryTraceHelpers.cpp
@@ -2,6 +2,7 @@
 
 #include "EnvironmentQuery/EnvQueryTraceHelpers.h"
 #include "NavigationData.h"
+#include "Algo/RemoveIf.h"
 
 template<>
 void FEQSHelpers::FBatchTrace::DoSingleSourceMultiDestinations<EEnvTraceShape::Line>(const FVector& Source, TArray<FNavLocation>& Points)
@@ -381,10 +382,17 @@ void FEQSHelpers::RunNavProjection(const ANavigationData& NavData, const UObject
 			Points[Idx] = Workload[Idx].OutLocation;
 			Points[Idx].Location.Z += TraceData.PostProjectionVerticalOffset;
 		}
-		else if (TraceMode == ETraceMode::Discard)
+	}
+
+	if (TraceMode == ETraceMode::Discard)
+	{
+		const FNavLocation* PointsBegin = Points.GetData();
+		int32 NewNum = Algo::StableRemoveIf(Points, [&Workload, PointsBegin](FNavLocation& Point)
 		{
-			Points.RemoveAt(Idx, 1, false);
-		}
+			return !Workload[&Point - PointsBegin].bResult;
+		});
+		const bool bAllowShrinking = false;
+		Points.SetNum(NewNum, bAllowShrinking);
 	}
 }
 
