Removed deprecated code from SimpleCellGrid

diff --git a/Engine/Source/Runtime/AIModule/Public/SimpleCellGrid.h b/Engine/Source/Runtime/AIModule/Public/SimpleCellGrid.h
--- a/Engine/Source/Runtime/AIModule/Public/SimpleCellGrid.h
+++ b/Engine/Source/Runtime/AIModule/Public/SimpleCellGrid.h
@@ -22,9 +22,6 @@ struct TSimpleCellGrid
 {
 	typedef CellType FCellType;
 
-	// DEPRECATED, use GridCellSize instead of this
-	uint32 CellSize;
-
 	float GridCellSize;
 	FBox WorldBounds;
 	FVector Origin;
@@ -36,8 +33,7 @@ protected:
 
 public:
 	TSimpleCellGrid()
-		: CellSize(0)
-		, GridCellSize(0.0f)
+		: GridCellSize(0.0f)
 		, WorldBounds(ForceInitToZero)
 		, Origin(FLT_MAX)
 		, BoundsSize(0)
@@ -52,7 +48,6 @@ public:
 		}
 			
 		GridCellSize = InCellSize;
-		CellSize = FMath::TruncToInt(InCellSize);
 
 		const FVector RealBoundsSize = Bounds.GetSize();
 		GridSize = FGridSize2D(FMath::CeilToInt(RealBoundsSize.X / InCellSize), FMath::CeilToInt(RealBoundsSize.Y / InCellSize));
@@ -215,7 +210,6 @@ public:
 
 	void Serialize(FArchive& Ar)
 	{	
-		// CellSize must act as version checking: MAX_uint32, float CellSize is used
 		uint32 VersionNum = MAX_uint32;
 		Ar << VersionNum;
 
@@ -224,11 +218,9 @@ public:
 			if (VersionNum == MAX_uint32)
 			{
 				Ar << GridCellSize;
-				CellSize = FMath::TruncToInt(GridCellSize);
 			}
 			else
 			{
-				CellSize = VersionNum;
 				GridCellSize = VersionNum * 1.0f;
 			}
 		}
@@ -286,74 +278,4 @@ public:
 		GridCellSize = 0;
 		Origin = FVector(FLT_MAX);
 	}
-
-	//////////////////////////////////////////////////////////////////////////
-	// DEPRECATED SUPPORT
-
-	UE_DEPRECATED(4.14, "This function is now deprecated, please use the one with float CellSize argument")
-	void Init(uint32 InCellSize, const FBox& Bounds)
-	{
-		Init(1.0f * InCellSize, Bounds);
-	}
-
-	UE_DEPRECATED(4.14, "This function is now deprecated, please use GetCellCoords instead.")
-	FIntVector WorldToCellCoords(const FVector& WorldLocation) const
-	{
-		const uint32 LocationX = FMath::TruncToInt((WorldLocation.X - Origin.X) / CellSize);
-		const uint32 LocationY = FMath::TruncToInt((WorldLocation.Y - Origin.Y) / CellSize);
-		return FIntVector(LocationX, LocationY, 0);
-	}
-
-	UE_DEPRECATED(4.14, "This function is now deprecated, please use GetCellCoords instead.")
-	void WorldToCellCoords(const FVector& WorldLocation, uint32& LocationX, uint32& LocationY) const
-	{
-		LocationX = FMath::TruncToInt((WorldLocation.X - Origin.X) / CellSize);
-		LocationY = FMath::TruncToInt((WorldLocation.Y - Origin.Y) / CellSize);
-	}
-
-	UE_DEPRECATED(4.14, "This function is now deprecated, please use GetCellIndex instead.")
-	uint32 WorldToCellIndex(const FVector& WorldLocation) const
-	{
-		const uint32 LocationX = FMath::TruncToInt((WorldLocation.X - Origin.X) / CellSize);
-		const uint32 LocationY = FMath::TruncToInt((WorldLocation.Y - Origin.Y) / CellSize);
-		return GetCellIndex(LocationX, LocationY);
-	}
-
-	UE_DEPRECATED(4.14, "This function is now deprecated, please use GetCellCoords instead.")
-	FIntVector CellIndexToCoords(uint32 CellIndex) const
-	{
-		return FIntVector(CellIndex / GridSize.Height, CellIndex % GridSize.Height, 0);
-	}
-
-	UE_DEPRECATED(4.14, "This function is now deprecated, please use GetCellCoords instead.")
-	FIntVector CellIndexToCoords(uint32 CellIndex, uint32& LocationX, uint32& LocationY) const
-	{
-		LocationX = CellIndex / GridSize.Height;
-		LocationY = CellIndex % GridSize.Height;
-		return FIntVector(LocationX, LocationY, 0);
-	}
-
-	UE_DEPRECATED(4.14, "This function is now deprecated, please use GetCellIndex instead.")
-	uint32 CellCoordsToCellIndex(const int32 LocationX, const int32 LocationY) const
-	{
-		return (LocationX * GridSize.Height) + LocationY;
-	}
-
-	UE_DEPRECATED(4.14, "This function is now deprecated, please use GetCellAtWorldLocation instead.")
-	const FCellType& GetCellAtWorldLocationSafe(const FVector& WorldLocation) const
-	{
-		return GetCellAtWorldLocation(WorldLocation);
-	}
-
-	UE_DEPRECATED(4.14, "This function is now deprecated, please use GetAllocatedSize instead.")
-	uint32 GetValuesMemorySize() const
-	{
-		return GridSize.Height * GridSize.Width * sizeof(FCellType);
-	}
-
-	UE_DEPRECATED(4.14, "This function is now deprecated, please use IsValidIndex instead.")
-	bool IsValidCellIndex(const int32 CellIndex) const
-	{
-		return IsValidIndex(CellIndex);
-	}
 };
