Made AAIController::BuildPathfindingQuery's logging handle the special case of Null nav system.

diff --git a/Engine/Source/Runtime/AIModule/Private/AIController.cpp b/Engine/Source/Runtime/AIModule/Private/AIController.cpp
--- a/Engine/Source/Runtime/AIModule/Private/AIController.cpp
+++ b/Engine/Source/Runtime/AIModule/Private/AIController.cpp
@@ -793,7 +793,14 @@ bool AAIController::BuildPathfindingQuery(const FAIMoveRequest& MoveRequest, FPa
 	}
 	else
 	{
-		UE_VLOG(this, LogAINavigation, Warning, TEXT("Unable to find NavigationData instance while calling AAIController::BuildPathfindingQuery"));
+		if (NavSys == nullptr)
+		{
+			UE_VLOG(this, LogAINavigation, Warning, TEXT("Unable AAIController::BuildPathfindingQuery due to no NavigationSystem present. Note that even pathfinding-less movement requires presence of NavigationSystem."));
+		}
+		else 
+		{
+			UE_VLOG(this, LogAINavigation, Warning, TEXT("Unable to find NavigationData instance while calling AAIController::BuildPathfindingQuery"));
+		}
 	}
 
 	return bResult;
