Minor VisLog code massaging

diff --git a/Engine/Source/Runtime/AIModule/Classes/VisualLoggerExtension.h b/Engine/Source/Runtime/AIModule/Classes/VisualLoggerExtension.h
--- a/Engine/Source/Runtime/AIModule/Classes/VisualLoggerExtension.h
+++ b/Engine/Source/Runtime/AIModule/Classes/VisualLoggerExtension.h
@@ -26,8 +26,6 @@ class UCanvas;
 class FVisualLoggerExtension : public FVisualLogExtensionInterface
 {
 public:
-	FVisualLoggerExtension();
-
 	virtual void ResetData(IVisualLoggerEditorInterface* EdInterface) override;
 	virtual void DrawData(IVisualLoggerEditorInterface* EdInterface, UCanvas* Canvas) override;
 	virtual void OnItemsSelectionChanged(IVisualLoggerEditorInterface* EdInterface) override;
@@ -38,8 +36,8 @@ private:
 	void DisableEQSRendering(AActor* HelperActor);
 
 protected:
-	uint32 SelectedEQSId;
-	float CurrentTimestamp;
+	int32 SelectedEQSId = INDEX_NONE;
+	float CurrentTimestamp = FLT_MIN;
 	TArray<TWeakObjectPtr<class UEQSRenderingComponent> >	EQSRenderingComponents;
 };
 #endif //ENABLE_VISUAL_LOG
diff --git a/Engine/Source/Runtime/AIModule/Private/VisualLoggerExtension.cpp b/Engine/Source/Runtime/AIModule/Private/VisualLoggerExtension.cpp
--- a/Engine/Source/Runtime/AIModule/Private/VisualLoggerExtension.cpp
+++ b/Engine/Source/Runtime/AIModule/Private/VisualLoggerExtension.cpp
@@ -11,13 +11,6 @@
 #include "EnvironmentQuery/EQSRenderingComponent.h"
 
 #if ENABLE_VISUAL_LOG
-FVisualLoggerExtension::FVisualLoggerExtension()
-	: SelectedEQSId(INDEX_NONE)
-	, CurrentTimestamp(FLT_MIN)
-{
-
-}
-
 void FVisualLoggerExtension::DisableEQSRendering(AActor* HelperActor)
 {
 #if USE_EQS_DEBUGGER
diff --git a/Engine/Source/Runtime/Engine/Private/VisualLogger/VisualLogger.cpp b/Engine/Source/Runtime/Engine/Private/VisualLogger/VisualLogger.cpp
--- a/Engine/Source/Runtime/Engine/Private/VisualLogger/VisualLogger.cpp
+++ b/Engine/Source/Runtime/Engine/Private/VisualLogger/VisualLogger.cpp
@@ -39,9 +39,9 @@ namespace
 		}
 
 #endif
-		if (!GIsEditor && World == nullptr)
+		if (!GIsEditor && World == nullptr && GEngine)
 		{
-			World = GEngine ? GEngine->GetWorld() : nullptr;
+			World =  GEngine->GetWorld();
 		}
 
 		return World;
