PR #7862: Fixing vlogging of EQS when a test filters out all items. (Contributed by [REDACTED_PERSON])

diff --git a/Engine/Source/Runtime/AIModule/Private/EnvironmentQuery/EnvQueryInstance.cpp b/Engine/Source/Runtime/AIModule/Private/EnvironmentQuery/EnvQueryInstance.cpp
--- a/Engine/Source/Runtime/AIModule/Private/EnvironmentQuery/EnvQueryInstance.cpp
+++ b/Engine/Source/Runtime/AIModule/Private/EnvironmentQuery/EnvQueryInstance.cpp
@@ -818,10 +818,6 @@ void FEnvQueryInstance::FinalizeQuery()
 	}
 	else
 	{
-		Items.Reset();
-		ItemDetails.Reset();
-		RawData.Reset();
-
 		MarkAsFailed();
 	}
 }
@@ -942,4 +938,4 @@ FBox FEnvQueryInstance::GetBoundingBox() const
 	return BBox;
 }
 
-#undef eqs_ensure
\ No newline at end of file
+#undef eqs_ensure
