Updated FGameplayDebuggerCategory_Perception to support Player-hosted AIPerceptionComponents

diff --git a/Engine/Source/Runtime/AIModule/Private/GameplayDebugger/GameplayDebuggerCategory_Perception.cpp b/Engine/Source/Runtime/AIModule/Private/GameplayDebugger/GameplayDebuggerCategory_Perception.cpp
--- a/Engine/Source/Runtime/AIModule/Private/GameplayDebugger/GameplayDebuggerCategory_Perception.cpp
+++ b/Engine/Source/Runtime/AIModule/Private/GameplayDebugger/GameplayDebuggerCategory_Perception.cpp
@@ -23,11 +23,20 @@ void FGameplayDebuggerCategory_Perception::CollectData(APlayerController* OwnerP
 	APawn* MyPawn = Cast<APawn>(DebugActor);
 	if (MyPawn)
 	{
-		AAIController* AIC = Cast<AAIController>(MyPawn->GetController());
-		if (AIC)
+		AController* Controller = MyPawn->GetController();
+		if (AAIController* AIC = Cast<AAIController>(Controller))
 		{
 			PerceptionComponent = AIC->GetPerceptionComponent();
 		}
+		else
+		{
+			PerceptionComponent = MyPawn->FindComponentByClass<UAIPerceptionComponent>();
+			// try the controller if the Pawn doesn't have it
+			if (PerceptionComponent == nullptr && Controller)
+			{
+				PerceptionComponent = Controller->FindComponentByClass<UAIPerceptionComponent>();
+			}
+		}
 	}
 
 	if (PerceptionComponent == nullptr && DebugActor != nullptr)
