[BehaviorTree] fixed null access when running MoveTo task without an AIController

diff --git a/Engine/Source/Runtime/AIModule/Private/BehaviorTree/Tasks/BTTask_MoveTo.cpp b/Engine/Source/Runtime/AIModule/Private/BehaviorTree/Tasks/BTTask_MoveTo.cpp
--- a/Engine/Source/Runtime/AIModule/Private/BehaviorTree/Tasks/BTTask_MoveTo.cpp
+++ b/Engine/Source/Runtime/AIModule/Private/BehaviorTree/Tasks/BTTask_MoveTo.cpp
@@ -44,14 +44,22 @@ EBTNodeResult::Type UBTTask_MoveTo::ExecuteTask(UBehaviorTreeComponent& OwnerCom
 	MyMemory->MoveRequestID = FAIRequestID::InvalidRequest;
 
 	AAIController* MyController = OwnerComp.GetAIOwner();
-	MyMemory->bWaitingForPath = bUseGameplayTasks ? false : MyController->ShouldPostponePathUpdates();
-	if (!MyMemory->bWaitingForPath)
+	if (MyController == nullptr)
 	{
-		NodeResult = PerformMoveTask(OwnerComp, NodeMemory);
+		UE_VLOG(OwnerComp.GetOwner(), LogBehaviorTree, Error, TEXT("UBTTask_MoveTo::ExecuteTask failed since AIController is missing."));
+		NodeResult = EBTNodeResult::Failed;
 	}
 	else
 	{
-		UE_VLOG(MyController, LogBehaviorTree, Log, TEXT("Pathfinding requests are freezed, waiting..."));
+		MyMemory->bWaitingForPath = bUseGameplayTasks ? false : MyController->ShouldPostponePathUpdates();
+		if (!MyMemory->bWaitingForPath)
+		{
+			NodeResult = PerformMoveTask(OwnerComp, NodeMemory);
+		}
+		else
+		{
+			UE_VLOG(MyController, LogBehaviorTree, Log, TEXT("Pathfinding requests are freezed, waiting..."));
+		}
 	}
 
 	if (NodeResult == EBTNodeResult::InProgress && bObserveBlackboardValue)
@@ -157,7 +165,7 @@ EBTNodeResult::Type UBTTask_MoveTo::PerformMoveTask(UBehaviorTreeComponent& Owne
 			}
 			else
 			{
-				FPathFollowingRequestResult RequestResult = MyController->MoveTo(MoveReq);
+				const FPathFollowingRequestResult RequestResult = MyController->MoveTo(MoveReq);
 				if (RequestResult.Code == EPathFollowingRequestResult::RequestSuccessful)
 				{
 					MyMemory->MoveRequestID = RequestResult.MoveId;
@@ -196,7 +204,7 @@ EBlackboardNotificationResult UBTTask_MoveTo::OnBlackboardValueChange(const UBla
 		return EBlackboardNotificationResult::RemoveObserver;
 	}
 
-	AAIController* MyController = BehaviorComp->GetAIOwner();
+	const AAIController* MyController = BehaviorComp->GetAIOwner();
 	uint8* RawMemory = BehaviorComp->GetNodeMemory(this, BehaviorComp->FindInstanceContainingNode(this));
 	FBTMoveToTaskMemory* MyMemory = CastInstanceNodeMemory<FBTMoveToTaskMemory>(RawMemory);
 
@@ -376,7 +384,7 @@ void UBTTask_MoveTo::DescribeRuntimeValues(const UBehaviorTreeComponent& OwnerCo
 	{
 		const FString KeyValue = BlackboardComp->DescribeKeyValue(BlackboardKey.GetSelectedKeyID(), EBlackboardDescription::OnlyValue);
 
-		FBTMoveToTaskMemory* MyMemory = (FBTMoveToTaskMemory*)NodeMemory;
+		const FBTMoveToTaskMemory* MyMemory = (FBTMoveToTaskMemory*)NodeMemory;
 		const bool bIsUsingTask = MyMemory->Task.IsValid();
 		
 		const FString ModeDesc =
