Fixed a broken check in UBlackboardComponent::PopulateSynchronizedKeys.

diff --git a/Engine/Source/Runtime/AIModule/Private/BehaviorTree/BlackboardComponent.cpp b/Engine/Source/Runtime/AIModule/Private/BehaviorTree/BlackboardComponent.cpp
--- a/Engine/Source/Runtime/AIModule/Private/BehaviorTree/BlackboardComponent.cpp
+++ b/Engine/Source/Runtime/AIModule/Private/BehaviorTree/BlackboardComponent.cpp
@@ -227,7 +227,7 @@ void UBlackboardComponent::PopulateSynchronizedKeys()
 		UBlackboardComponent* OtherBlackboard = Iter.Value();
 		if (OtherBlackboard != nullptr && ShouldSyncWithBlackboard(*OtherBlackboard))
 		{
-			for (const auto& Key : BlackboardAsset->Keys)
+			for (const FBlackboardEntry& Key : BlackboardAsset->Keys)
 			{
 				if (Key.bInstanceSynced)
 				{
@@ -237,7 +237,7 @@ void UBlackboardComponent::PopulateSynchronizedKeys()
 					{
 						const FBlackboardEntry* const OtherKey = OtherBlackboard->GetBlackboardAsset()->GetKey(OtherKeyID);
 						check(Key.EntryName == OtherKey->EntryName);
-						check(Key.KeyType == OtherKey->KeyType);
+						check(Key.KeyType && OtherKey->KeyType && Key.KeyType->GetClass() == OtherKey->KeyType->GetClass());
 
 						const bool bKeyHasInstance = Key.KeyType->HasInstance();
 						const uint16 DataOffset = bKeyHasInstance ? sizeof(FBlackboardInstancedKeyMemory) : 0;
