[BehaviorTree] added scoped debug logging context and spreaded a few more uses for node initialization from asset

diff --git a/Engine/Source/Editor/BehaviorTreeEditor/Private/BehaviorTreeGraph.cpp b/Engine/Source/Editor/BehaviorTreeEditor/Private/BehaviorTreeGraph.cpp
--- a/Engine/Source/Editor/BehaviorTreeEditor/Private/BehaviorTreeGraph.cpp
+++ b/Engine/Source/Editor/BehaviorTreeEditor/Private/BehaviorTreeGraph.cpp
@@ -61,7 +61,7 @@ void UBehaviorTreeGraph::UpdateBlackboardChange()
 			UBTNode* MyNodeInstance = Cast<UBTNode>(MyNode->NodeInstance);
 			if (MyNodeInstance)
 			{
-				UBehaviorTreeTypes::SetBTLoggingContext(MyNodeInstance);
+				FScopedBTLoggingContext LogContext(MyNodeInstance);
 				MyNodeInstance->InitializeFromAsset(*BTAsset);
 			}
 
@@ -70,7 +70,7 @@ void UBehaviorTreeGraph::UpdateBlackboardChange()
 				UBTNode* DecoratorNodeInstance = MyNode->Decorators[iDecorator] ? Cast<UBTNode>(MyNode->Decorators[iDecorator]->NodeInstance) : NULL;
 				if (DecoratorNodeInstance)
 				{
-					UBehaviorTreeTypes::SetBTLoggingContext(DecoratorNodeInstance);
+					FScopedBTLoggingContext LogContext(DecoratorNodeInstance);
 					DecoratorNodeInstance->InitializeFromAsset(*BTAsset);
 				}
 
@@ -86,12 +86,10 @@ void UBehaviorTreeGraph::UpdateBlackboardChange()
 				UBTNode* ServiceNodeInstance = MyNode->Services[iService] ? Cast<UBTNode>(MyNode->Services[iService]->NodeInstance) : NULL;
 				if (ServiceNodeInstance)
 				{
-					UBehaviorTreeTypes::SetBTLoggingContext(ServiceNodeInstance);
+					FScopedBTLoggingContext LogContext(ServiceNodeInstance);
 					ServiceNodeInstance->InitializeFromAsset(*BTAsset);
 				}
 			}
-
-			UBehaviorTreeTypes::SetBTLoggingContext(nullptr);
 		}
 	}
 }
diff --git a/Engine/Source/Runtime/AIModule/Classes/BehaviorTree/BehaviorTreeTypes.h b/Engine/Source/Runtime/AIModule/Classes/BehaviorTree/BehaviorTreeTypes.h
--- a/Engine/Source/Runtime/AIModule/Classes/BehaviorTree/BehaviorTreeTypes.h
+++ b/Engine/Source/Runtime/AIModule/Classes/BehaviorTree/BehaviorTreeTypes.h
@@ -668,3 +668,18 @@ public:
 	// 	pass nullptr to clear
 	static void SetBTLoggingContext(const UBTNode* NewBTLoggingContext);
 };
+
+/** Helper struct to push a node as the new logging context and automatically reset the context on destruction. */
+struct FScopedBTLoggingContext
+{
+	FScopedBTLoggingContext() = delete;
+	explicit FScopedBTLoggingContext(const UBTNode* Context)
+	{
+		UBehaviorTreeTypes::SetBTLoggingContext(Context);
+	}
+
+	~FScopedBTLoggingContext()
+	{
+		UBehaviorTreeTypes::SetBTLoggingContext(nullptr);
+	}
+};
\ No newline at end of file
diff --git a/Engine/Source/Runtime/AIModule/Private/BehaviorTree/BTNode.cpp b/Engine/Source/Runtime/AIModule/Private/BehaviorTree/BTNode.cpp
--- a/Engine/Source/Runtime/AIModule/Private/BehaviorTree/BTNode.cpp
+++ b/Engine/Source/Runtime/AIModule/Private/BehaviorTree/BTNode.cpp
@@ -105,6 +105,8 @@ void UBTNode::InitializeInSubtree(UBehaviorTreeComponent& OwnerComp, uint8* Node
 		NodeInstance->SetOwner(OwnerComp.GetOwner());
 		NodeInstance->InitializeMemory(OwnerComp, NodeMemory, InitType);
 		check(TreeAsset);
+
+		FScopedBTLoggingContext LogContext(NodeInstance);
 		NodeInstance->InitializeFromAsset(*TreeAsset);
 		NodeInstance->OnInstanceCreated(OwnerComp);
 		NextInstancedIndex++;
diff --git a/Engine/Source/Runtime/AIModule/Private/BehaviorTree/BehaviorTreeManager.cpp b/Engine/Source/Runtime/AIModule/Private/BehaviorTree/BehaviorTreeManager.cpp
--- a/Engine/Source/Runtime/AIModule/Private/BehaviorTree/BehaviorTreeManager.cpp
+++ b/Engine/Source/Runtime/AIModule/Private/BehaviorTree/BehaviorTreeManager.cpp
@@ -143,7 +143,10 @@ static void InitializeNodeHelper(UBTCompositeNode* ParentNode, UBTNode* NodeOb,
 	UBehaviorTree& TreeAsset, UObject* NodeOuter)
 {
 	InitList.Add(FNodeInitializationData(NodeOb, ParentNode, ExecutionIndex, TreeDepth, NodeOb->GetInstanceMemorySize(), NodeOb->GetSpecialMemorySize()));
-	NodeOb->InitializeFromAsset(TreeAsset);
+	{
+		FScopedBTLoggingContext LogContext(NodeOb);
+		NodeOb->InitializeFromAsset(TreeAsset);
+	}
 	ExecutionIndex++;
 
 	UBTCompositeNode* CompositeOb = Cast<UBTCompositeNode>(NodeOb);
@@ -167,6 +170,7 @@ static void InitializeNodeHelper(UBTCompositeNode* ParentNode, UBTNode* NodeOb,
 			InitList.Add(FNodeInitializationData(Service, CompositeOb, ExecutionIndex, TreeDepth,
 				Service->GetInstanceMemorySize(), Service->GetSpecialMemorySize()));
 
+			FScopedBTLoggingContext LogContext(Service);
 			Service->InitializeFromAsset(TreeAsset);
 			// don't initialize parent link for services on composite node
 			ExecutionIndex++;
@@ -193,6 +197,7 @@ static void InitializeNodeHelper(UBTCompositeNode* ParentNode, UBTNode* NodeOb,
 				InitList.Add(FNodeInitializationData(Decorator, CompositeOb, ExecutionIndex, TreeDepth,
 					Decorator->GetInstanceMemorySize(), Decorator->GetSpecialMemorySize()));
 
+				FScopedBTLoggingContext LogContext(Decorator);
 				Decorator->InitializeFromAsset(TreeAsset);
 				Decorator->InitializeParentLink(ChildIndex);
 				ExecutionIndex++;
@@ -223,6 +228,7 @@ static void InitializeNodeHelper(UBTCompositeNode* ParentNode, UBTNode* NodeOb,
 
 					// initialize with parent tree
 					Decorator->MarkInjectedNode();
+					FScopedBTLoggingContext LogContext(Decorator);
 					Decorator->InitializeFromAsset(TreeAsset);
 					Decorator->InitializeParentLink(ChildIndex);
 					ExecutionIndex++;
@@ -262,6 +268,7 @@ static void InitializeNodeHelper(UBTCompositeNode* ParentNode, UBTNode* NodeOb,
 					InitList.Add(FNodeInitializationData(Service, CompositeOb, ExecutionIndex, TreeDepth,
 						Service->GetInstanceMemorySize(), Service->GetSpecialMemorySize()));
 
+					FScopedBTLoggingContext LogContext(Service);
 					Service->InitializeFromAsset(TreeAsset);
 					Service->InitializeParentLink(ChildIndex);
 					ExecutionIndex++;
diff --git a/Engine/Source/Runtime/AIModule/Private/BehaviorTree/BehaviorTreeTypes.cpp b/Engine/Source/Runtime/AIModule/Private/BehaviorTree/BehaviorTreeTypes.cpp
--- a/Engine/Source/Runtime/AIModule/Private/BehaviorTree/BehaviorTreeTypes.cpp
+++ b/Engine/Source/Runtime/AIModule/Private/BehaviorTree/BehaviorTreeTypes.cpp
@@ -642,9 +642,7 @@ FString UBehaviorTreeTypes::GetShortTypeName(const UObject* Ob)
 
 void UBehaviorTreeTypes::SetBTLoggingContext(const UBTNode* NewBTLoggingContext)
 {
-	BTLoggingContext = NewBTLoggingContext 
-		? FString::Printf(TEXT("%s[%d]"), *NewBTLoggingContext->GetNodeName(), NewBTLoggingContext->GetExecutionIndex())
-		: TEXT("");
+	BTLoggingContext = DescribeNodeHelper(NewBTLoggingContext);
 }
 
 //----------------------------------------------------------------------//
