AI5System & AI5Subsystem removal:
- Added virtual method OnWorldBeginPlay to WorldSubsystem
- Converted all systems inheriting from AI5Subsystem to UWorldSubsystem/UTickableWorldSubsystem depending on their need to be ticked
- Updated .ini files to keep using legacy AIModule
- Moved some helpers functions from UE::AI5 to UE::AI namespace
- Bonus code beautification: const, typos, useless includes

diff --git a/Engine/Source/Runtime/AIModule/AICodingStandard.md b/Engine/Source/Runtime/AIModule/AICodingStandard.md
new file mode 100644
--- /dev/null
+++ b/Engine/Source/Runtime/AIModule/AICodingStandard.md
@@ -0,0 +1,44 @@
+## AI Team’s coding standard extension 
+
+_In all cases [Epic’s coding standard](https://docs.unrealengine.com/en-US/ProductionPipelines/DevelopmentSetup/CodingStandard/index.html) has higher priority. AI Team’s coding standard needs to remain compatible and act as an extension and clarification where Epic’s is unclear or leaves room for interpretation._
+
+#### Code organization and naming
+
+##### cpp files
+* If a source file has implementation of multiple structs/classes, use the following header to mark the beginning of each class/struct to make it super clear when new implementation begins.
+
+```
+//----------------------------------------------------------------//
+//  FMyTypeName
+//----------------------------------------------------------------//
+```
+* Keep member functions of a given type together, avoid interlacing functions of multiple types.
+
+##### Subsystems 
+`WorldSubsystem`-derived classes need to have a `Subsystem` postfix, while other subsystem-y classes that are not related with any of the engine's subsystem classes get a `Manager` postfix.
+
+#### Namespaces 
+
+[Unreal Engine standard here](https://docs.unrealengine.com/en-US/ProductionPipelines/DevelopmentSetup/CodingStandard/index.html#namespaces).
+
+**Nested namespaces** declaration should be done in a single line, like so:
+
+```c++
+namespace UE::Mass::Debug
+{ 
+//...	
+} // namespace UE::Mass::Debug
+```
+
+
+#### Debug-time code
+
+Debugging-related member function names should start with Debug prefix.
+
+___
+
+### FAQ
+*add things that can be best explained as a simple answer to a simple question. These things can be reflected above as well. This section needs to remain last at all times.* 
+
+**Q**: What do I do when the thing?
+**A**: You do the thing.
\ No newline at end of file
diff --git a/Engine/Source/Runtime/AIModule/Private/AIHelpers.cpp b/Engine/Source/Runtime/AIModule/Private/AIHelpers.cpp
--- a/Engine/Source/Runtime/AIModule/Private/AIHelpers.cpp
+++ b/Engine/Source/Runtime/AIModule/Private/AIHelpers.cpp
@@ -2,6 +2,9 @@
 
 #include "AIHelpers.h"
 
+#include "Engine/InheritableComponentHandler.h"
+#include "Engine/SCS_Node.h"
+
 namespace FAISystem
 {
 	FVector FindClosestLocation(const FVector& Origin, const TArray<FVector>& Locations)
@@ -76,3 +79,108 @@ namespace FAISystem
 		return false;
 	}
 }
+
+namespace UE::AI
+{
+
+TOptional<float> GetYawFromVector(const FVector& Vector)
+{
+	const FVector2D Vector2D(Vector);
+	const float SizeSquared = Vector2D.SizeSquared();
+	if (SizeSquared <= SMALL_NUMBER)
+	{
+		return TOptional<float>();
+	}
+
+	const FVector2D ForwardVector(1.0f, 0.0f);
+	const FVector2D NormVector2D = Vector2D * FMath::InvSqrt(SizeSquared);
+	const float DotVal = ForwardVector | NormVector2D;
+	const float ClampedDot = FMath::Clamp(DotVal, -1.0f, 1.0f);
+	const float Direction = ForwardVector ^ NormVector2D;
+	const float Sign = Direction < 0.0f ? -1.0f : 1.0f;
+	const float Yaw = Sign * FMath::Acos(ClampedDot);
+	return Yaw;
+}
+
+TOptional<float> GetYawFromRotator(const FRotator& Rotator)
+{
+	return GetYawFromVector(Rotator.Vector());
+}
+
+TOptional<float> GetYawFromQuaternion(const FQuat& Quaternion)
+{
+	return GetYawFromVector(Quaternion.Vector());
+}
+
+void GetActorClassDefaultComponents(const TSubclassOf<AActor>& ActorClass, TArray<UActorComponent*>& OutComponents, const TSubclassOf<UActorComponent>& InComponentClass)
+{
+	if (!ensure(ActorClass.Get()))
+	{
+		return;
+	}
+
+	UClass* ClassPtr = InComponentClass.Get();
+	TArray<UActorComponent*> ResultComponents;
+
+	// Get the components defined on the native class.
+	AActor* CDO = ActorClass->GetDefaultObject<AActor>();
+	check(CDO);
+	if (ClassPtr)
+	{
+		CDO->GetComponents(InComponentClass, ResultComponents);
+	}
+	else
+	{
+		CDO->GetComponents(ResultComponents);
+	}
+
+	// Try to get the components off the BP class.
+	UBlueprintGeneratedClass* BPClass = Cast<UBlueprintGeneratedClass>(*ActorClass);
+	if (BPClass)
+	{
+		// A BlueprintGeneratedClass has a USimpleConstructionScript member. This member has an array of RootNodes
+		// which contains the SCSNode for the root SceneComponent and non-SceneComponents. For the SceneComponent
+		// hierarchy, each SCSNode knows its children SCSNodes. Each SCSNode stores the component template that will
+		// be created when the Actor is spawned.
+		//
+		// WARNING: This may change in future engine versions!
+
+		TArray<UActorComponent*> Unfiltered;
+		// using this semantic to avoid duplicating following loops or adding a filtering check condition inside the loops
+		TArray<UActorComponent*>& TmpComponents = ClassPtr ? Unfiltered : ResultComponents;
+
+		// Check added components.
+		USimpleConstructionScript* ConstructionScript = BPClass->SimpleConstructionScript;
+		if (ConstructionScript)
+		{
+			for (const USCS_Node* Node : ConstructionScript->GetAllNodes())
+			{
+				TmpComponents.Add(Node->ComponentTemplate);
+			}
+		}
+		// Check modified inherited components.
+		UInheritableComponentHandler* InheritableComponentHandler = BPClass->InheritableComponentHandler;
+		if (InheritableComponentHandler)
+		{
+			for (TArray<FComponentOverrideRecord>::TIterator It = InheritableComponentHandler->CreateRecordIterator(); It; ++It)
+			{
+				TmpComponents.Add(It->ComponentTemplate);
+			}
+		}
+
+		// Filter to the ones matching the requested class.
+		if (ClassPtr)
+		{
+			for (UActorComponent* TemplateComponent : Unfiltered)
+			{
+				if (TemplateComponent->IsA(ClassPtr))
+				{
+					ResultComponents.Add(TemplateComponent);
+				}
+			}
+		}
+	}
+
+	OutComponents = MoveTemp(ResultComponents);
+}
+} // UE::AI
\ No newline at end of file
diff --git a/Engine/Source/Runtime/AIModule/Public/AIHelpers.h b/Engine/Source/Runtime/AIModule/Public/AIHelpers.h
--- a/Engine/Source/Runtime/AIModule/Public/AIHelpers.h
+++ b/Engine/Source/Runtime/AIModule/Public/AIHelpers.h
@@ -2,9 +2,11 @@
 
 #pragma once
 
-#include "CoreMinimal.h"
 #include "AITypes.h"
 
+class AActor;
+class UActorComponent;
+
 namespace FAISystem
 {
 	AIMODULE_API FVector FindClosestLocation(const FVector& Origin, const TArray<FVector>& Locations);
@@ -41,5 +43,34 @@ namespace FAISystem
 	// F = Far Clipping Radius (from the StartLocation adjusted by Backward offset)
 	//----------------------------------------------------------------------//
 	AIMODULE_API bool CheckIsTargetInSightCone(const FVector& StartLocation, const FVector& ConeDirectionNormal, float PeripheralVisionAngleCos,
-		float ConeDirectionBackwardOffset, float NearClippingRadiusSq, float const FarClippingRadiusSq, const FVector& TargetLocation);
+											   float ConeDirectionBackwardOffset, float NearClippingRadiusSq, float const FarClippingRadiusSq, const FVector& TargetLocation);
 }
+
+namespace UE::AI
+{
+	/**
+	 * This method will extract the yaw radian from the specified vector (The vector does not need to be normalized)
+	 * if it is not possible to compute yaw, the function will return an invalid value	 *
+	 */
+	extern AIMODULE_API TOptional<float> GetYawFromVector(const FVector& Vector);
+
+	/**
+	 * This method will extract the yaw radian from the specified rotator
+	 * if it is not possible to compute yaw, the function will not return an invalid value
+	 */
+	extern AIMODULE_API TOptional<float> GetYawFromRotator(const FRotator& Rotator);
+
+	/**
+	 * This method will extract the yaw radian from the specified quaternion
+	 * if it is not possible to compute yaw, the function will not return an invalid value
+	 */
+	extern AIMODULE_API TOptional<float> GetYawFromQuaternion(const FQuat& Quaternion);
+
+	/**
+	 * Fetches all the components of ActorClass's CDO, including the ones added via the BP editor (which AActor.GetComponents fails to do)
+	 * @param ActorClass class of AActor for which we will retrieve all components
+	 * @param OutComponents this is where the found components will end up. Note that the preexisting contents of OutComponents will get overridden.
+	 * @param InComponentClass if supplied will be used to filter the results
+	 */
+	extern AIMODULE_API void GetActorClassDefaultComponents(const TSubclassOf<AActor>& ActorClass, TArray<UActorComponent*>& OutComponents, const TSubclassOf<UActorComponent>& InComponentClass = TSubclassOf<UActorComponent>());
+} // UE::AI
diff --git a/Engine/Source/Runtime/AIModule/Public/HierarchicalHashGrid2D.h b/Engine/Source/Runtime/AIModule/Public/HierarchicalHashGrid2D.h
new file mode 100644
--- /dev/null
+++ b/Engine/Source/Runtime/AIModule/Public/HierarchicalHashGrid2D.h
@@ -0,0 +1,610 @@
+// Copyright Epic Games, Inc. All Rights Reserved.
+
+#pragma once
+
+#include "CoreMinimal.h"
+#include "Containers/StaticArray.h"
+
+/**
+ * Hierarchical Hash Grid in 2D.
+ *
+ * Items are added to the "infinite" grid at specific level based on their size. The grid size, and size ratio between levels 
+ * are defined as the template parameter. This allows some computation to be optimized based on the values.
+ *
+ * Items that cannot be fitted in the grid are added to a spill list, which contents are included in all queries.
+ * Each item is added to just one grid cell and thus can overlap some neighbour cells up to half cell size at that grid level.
+ * This is compensated during the query by expanding the query box.
+ *
+ * Potential optimizations:
+ * - (X, Y, Level) could probably be int16.
+ * - Level ratio could be a shift so LevelDown() could become just shift.
+ * - FloorToInt ends up calling floorf, which will be a function call.
+ * 	 int FloorToInt(float a) { return (int)a + ((int)a > a ? -1 : 0); }
+ * 	 auto vectorizes nicely on clang, but not on VC.
+ * - Add helper function to allow to tweak the cell size to reset the grid when spill list gets too large
+ */
+template <int32 InNumLevels = 3, int32 InLevelRatio = 4, typename InItemIDType = uint32>
+class THierarchicalHashGrid2D
+{
+public:
+
+	typedef InItemIDType ItemIDType;
+	static const int32 NumLevels = InNumLevels;		/** Number of grid levels in the grid. */
+	static const int32 LevelRatio = InLevelRatio;	/** Ratio in cells between consecutive levels. */
+
+	/** Cells are located based on their hash. Each cell stores linked list of items at that location and how many children
+	  * are under the same location in the finer grid levels.
+	  */
+	struct FCell
+	{
+		FCell() : X(0), Y(0), Level(0) {}
+		FCell(const int32 InX, const int32 InY, const int32 InLevel) : X(InX), Y(InY), Level(InLevel) {}
+
+		int32 X, Y, Level;			/** Location of the cell */
+		int32 First = INDEX_NONE;	/** First item in the item linked list. (index to Items array, or INDEX_NONE) */
+		int32 ChildItemCount = 0;	/** Num of children in the child cells. */
+
+		/** Returns hash based on cell location. The hash function is called a lot, favor speed over quality. */
+		friend uint32 GetTypeHash(const FCell& Cell)
+		{
+			constexpr uint32 H1 = 0x8da6b343;	// Arbitrary big primes.
+			constexpr uint32 H2 = 0xd8163841;
+			constexpr uint32 H3 = 0xcb1ab31f;
+			return (H1 * uint32(Cell.X) + H2 * uint32(Cell.Y) + H3 * uint32(Cell.Level));
+		}
+
+		// Need for TSet<>
+		bool operator==(const FCell& RHS) const
+		{
+			return X == RHS.X && Y == RHS.Y && Level == RHS.Level;
+		}
+	};
+
+	/** Specifies cell location within the grid at specific level. */
+	struct FCellLocation
+	{
+		FCellLocation() = default;
+		FCellLocation(const int32 InX, const int32 InY, const int32 InLevel) : X(InX), Y(InY), Level(InLevel) {}
+
+		bool operator==(const FCellLocation& RHS) const
+		{
+			return X == RHS.X && Y == RHS.Y && Level == RHS.Level;
+		}
+
+		bool operator!=(const FCellLocation& RHS) const
+		{
+			return !(*this == RHS);
+		}
+
+		int32 X = 0;
+		int32 Y = 0;
+		int32 Level = INDEX_NONE;
+	};
+
+	/** Item stored in a grid cell. */
+	struct FItem
+	{
+		ItemIDType ID = 0;			/** External item ID */
+		int32 Next = INDEX_NONE;	/** Next item in the item linked list (index to Items array, or INDEX_NONE). */
+	};
+
+	/** Rectangle bounds, coordinates inclusive. */
+	struct FCellRect
+	{
+		FCellRect() = default;
+		FCellRect(const int32 InMinX, const int32 InMinY, const int32 InMaxX, const int32 InMaxY) : MinX(InMinX), MinY(InMinY), MaxX(InMaxX), MaxY(InMaxY) {}
+
+		int32 MinX = 0;
+		int32 MinY = 0;
+		int32 MaxX = 0;
+		int32 MaxY = 0;
+	};
+
+	/** Iterator state over a rectangle */
+	struct FCellRectIterator
+	{
+		int32 X = 0;		/** Current X */
+		int32 Y = 0;		/** Current Y */
+		int32 Level = 0;	/** Current level */
+		FCellRect Rect;		/** Area to iterate over */
+	};
+
+
+	/** Constructor, initializes the grid for specific cell size.
+	  * @param InCellSize - new finest level cell size of the grid
+	  */
+	THierarchicalHashGrid2D(const float InCellSize = 500.f)
+		: SpillList(INDEX_NONE)
+	{
+		check(InCellSize > 0.0f);
+
+		float CurrCellSize = InCellSize;
+		for (int32 i = 0; i < NumLevels; i++)
+		{
+			CellSize[i] = CurrCellSize;
+			InvCellSize[i] = 1.0f / CurrCellSize;
+			CurrCellSize *= (float)(LevelRatio);
+		}
+
+		for (int32 i = 0; i < NumLevels; i++)
+		{
+			LevelItemCount[i] = 0;
+		}
+	}
+
+	/** Resets and initializes the grid for specific cell size.
+	  * @param InCellSize - new finest level cell size of the grid
+	  */
+	void Initialize(const float InCellSize)
+	{
+		check(InCellSize > 0.0f);
+
+		Reset();
+
+		float CurrCellSize = InCellSize;
+		for (int32 i = 0; i < NumLevels; i++)
+		{
+			CellSize[i] = CurrCellSize;
+			InvCellSize[i] = 1.0f / CurrCellSize;
+			CurrCellSize *= (float)(LevelRatio);
+		}
+	}
+
+	/** Reset the grid to empty. */
+	void Reset()
+	{
+		SpillList = INDEX_NONE;
+		Cells.Reset();
+		Items.Reset();
+
+		for (int32 i = 0; i < NumLevels; i++)
+		{
+			LevelItemCount[i] = 0;
+		}
+	}
+
+	/** Adds item to the grid. 
+	  * @param ID - External ID used to identify the item.
+	  * @param Bounds - Bounding box of the item. 
+	  * @return Cell location of the item, can be used later to remove the item.
+	  */
+	FCellLocation Add(const ItemIDType ID, const FBox& Bounds)
+	{
+		const FCellLocation Location = CalcCellLocation(Bounds);
+		Add(ID, Location);
+		return Location;
+	}
+
+	/** Adds item to the grid.
+	  * @param ID - External ID used to identify the item.
+	  * @param Location - Cell location where the item should be added.
+	  * @return Cell location of the item, can be used later to remove the item.
+	  */
+	void Add(const ItemIDType ID, const FCellLocation& Location)
+	{
+		const int32 Idx = Items.AddUninitialized().Index;
+		FItem& Item = Items[Idx];
+		Item.ID = ID;
+
+		if (Location.Level == INDEX_NONE)
+		{
+			// Could not fit into any of the grids, add to spill list.
+			Item.Next = SpillList;
+			SpillList = Idx;
+		}
+		else
+		{
+			// Add to cell at specific level
+			FCell& Cell = FindOrAddCell(Location.X, Location.Y, Location.Level);
+			Item.Next = Cell.First;
+			Cell.First = Idx;
+
+			// Update per level counts
+			LevelItemCount[Location.Level]++;
+
+			// Update child counts
+			FCellLocation ParentLocation = Location;
+			while (ParentLocation.Level < NumLevels - 1)
+			{
+				ParentLocation.X = LevelDown(ParentLocation.X);
+				ParentLocation.Y = LevelDown(ParentLocation.Y);
+				ParentLocation.Level++;
+				FCell& ParentCell = FindOrAddCell(ParentLocation.X, ParentLocation.Y, ParentLocation.Level);
+				ParentCell.ChildItemCount++;
+			}
+		}
+	}
+
+	/** Removes item based on the bounding box it was added with.
+	  * @param ID - External ID used to identify the item.
+	  * @param OldBounds - The same bounding box the item was previously added or moved with.
+	  */
+	void Remove(const ItemIDType ID, const FBox& OldBounds)
+	{
+		const FCellLocation OldLocation = CalcCellLocation(OldBounds);
+		Remove(OldLocation, ID);
+	}
+
+	/** Removes item based on the cell location it was added with.
+	  * @param ID - External ID used to identify the item.
+	  * @param OldLocation - The same cell location the item was previously added or moved with.
+	  */
+	void Remove(const ItemIDType ID, const FCellLocation& OldLocation)
+	{
+		if (OldLocation.Level == INDEX_NONE)
+		{
+			// Remove from spill list.
+			int32 PrevIdx = INDEX_NONE;
+			for (int32 Idx = SpillList; Idx != INDEX_NONE; PrevIdx = Idx, Idx = Items[Idx].Next)
+			{
+				if (Items[Idx].ID == ID)
+				{
+					if (PrevIdx == INDEX_NONE)
+					{
+						SpillList = Items[Idx].Next;
+					}
+					else
+					{
+						Items[PrevIdx].Next = Items[Idx].Next;
+					}
+					Items.RemoveAtUninitialized(Idx);
+					break;
+				}
+			}
+		}
+		else
+		{
+			// Remove from cell
+			if (FCell* Cell = FindCellMutable(OldLocation.X, OldLocation.Y, OldLocation.Level))
+			{
+				int32 PrevIdx = INDEX_NONE;
+				for (int32 Idx = Cell->First; Idx != INDEX_NONE; PrevIdx = Idx, Idx = Items[Idx].Next)
+				{
+					if (Items[Idx].ID == ID)
+					{
+						if (PrevIdx == INDEX_NONE)
+						{
+							Cell->First = Items[Idx].Next;
+						}
+						else
+						{
+							Items[PrevIdx].Next = Items[Idx].Next;
+						}
+						Items.RemoveAtUninitialized(Idx);
+						break;
+					}
+				}
+
+				// Update per level counts
+				LevelItemCount[OldLocation.Level]--;
+
+				// Update child counts
+				FCellLocation ParentLocation = OldLocation;
+				while (ParentLocation.Level < NumLevels - 1)
+				{
+					ParentLocation.X = LevelDown(ParentLocation.X);
+					ParentLocation.Y = LevelDown(ParentLocation.Y);
+					ParentLocation.Level++;
+					FCell* ParentCell = FindCellMutable(ParentLocation.X, ParentLocation.Y, ParentLocation.Level);
+					if (ParentCell)
+					{
+						ParentCell->ChildItemCount--;
+					}
+				}
+			}
+		}
+	}
+
+	/** Moves item based on previous bounding box and new bounding box.
+	  * @param ID - External ID used to identify the item.
+	  * @param OldBounds - The same bounding box the item was previously added or moved with.
+	  * @param NewBounds - New bounds of the item
+	  * @return The cell location where the item was added to.
+	  */
+	FCellLocation Move(const ItemIDType ID, const FBox& OldBounds, const FBox& NewBounds)
+	{
+		const FCellLocation OldLocation = CalcCellLocation(OldBounds);
+		return Move(ID, OldLocation, NewBounds);
+	}
+
+	/** Moves item based on previous bounding box and new bounding box.
+	  * @param ID - External ID used to identify the item.
+	  * @param OldLocation - The same cell location the item was previously added or moved with.
+	  * @param NewBounds - New  bounds of the item
+	  * @return The cell location where the item was added to.
+	  */
+	FCellLocation Move(const ItemIDType ID, const FCellLocation& OldLocation, const FBox& NewBounds)
+	{
+		const FCellLocation NewLocation = CalcCellLocation(NewBounds);
+		if (NewLocation != OldLocation)
+		{
+			Remove(ID, OldLocation);
+			Add(ID, NewLocation);
+		}
+		return NewLocation;
+	}
+
+	/** Returns items that potentially touch the bounds. Operates on grid level, can have false positives.
+	  * This can be faster than Query() for small query box sizes (i.e. max dimension CellSize*4) due to simpler logic.
+	  * @param Bounds - Query bounding box.
+	  * @param OutResults - Result of the query, IDs of potentially overlapping items.
+	  */
+	void QuerySmall(const FBox& Bounds, TArray<ItemIDType>& OutResults) const
+	{
+		// Return items from buckets on all levels
+		for (int32 Level = 0; Level < NumLevels; Level++)
+		{
+			// Update per level counts
+			if (LevelItemCount[Level] == 0)
+			{
+				continue;
+			}
+
+			// Finest level query rect
+			FCellRect Rect = CalcQueryBounds(Bounds, Level);
+
+			for (int32 Y = Rect.MinY; Y <= Rect.MaxY; Y++)
+			{
+				for (int32 X = Rect.MinX; X <= Rect.MaxX; X++)
+				{
+					if (const FCell* Cell = FindCell(X, Y, Level))
+					{
+						for (int32 Idx = Cell->First; Idx != INDEX_NONE; Idx = Items[Idx].Next)
+						{
+							OutResults.Add(Items[Idx].ID);
+						}
+					}
+				}
+			}
+		}
+
+		// Everything from spill list
+		for (int32 Idx = SpillList; Idx != INDEX_NONE; Idx = Items[Idx].Next)
+		{
+			OutResults.Add(Items[Idx].ID);
+		}
+	}
+
+	/** Returns items that potentially touch the bounds. Operates on grid level, can have false positives.
+	  * @param Bounds - Query bounding box.
+	  * @param OutResults - Result of the query, IDs of potentially overlapping items.
+	  */
+	void Query(const FBox& Bounds, TArray<ItemIDType>& OutResults) const
+	{
+		FCellRect Rects[NumLevels];
+		FCellRectIterator Iters[NumLevels];
+		int32 IterIdx = 0;
+
+		// Calculate cell bounds for each level, keep track of the coarsest level that has any items, we'll start from that.
+		for (int32 Level = 0; Level < NumLevels; Level++)
+		{
+			Rects[Level] = CalcQueryBounds(Bounds, Level);
+		}
+
+		// The idea of the iterator below is that it iterates over rectangle cells recursively towards finer levels, depth first.
+		// The previous level's iterator is kept in the Iters stack, and we can pop and continue that once the finer level is completed.
+		// Finer iterator rectangles is clamped against that levels tight bounds so that unnecessary cells are not visited.
+		// Coarser levels of the grid also store data if the finer levels under them has any items. This is used to skip iterating
+		// lower levels at certain locations completely. This can be big advantage in larger query boxes, compared to iterating all cells as in QuerySmall().
+
+		// Init coarsest iterator
+		const int32 StartLevel = NumLevels - 1;
+		Iters[IterIdx].Level = StartLevel;
+		Iters[IterIdx].Rect = Rects[StartLevel];
+		Iters[IterIdx].X = Iters[IterIdx].Rect.MinX;
+		Iters[IterIdx].Y = Iters[IterIdx].Rect.MinY;
+		IterIdx++;
+
+		while (IterIdx > 0)
+		{
+			FCellRectIterator& Iter = Iters[IterIdx - 1];
+			// Check if the iterator has finished
+			if (Iter.X > Iter.Rect.MaxX)
+			{
+				Iter.X = Iter.Rect.MinX;
+				Iter.Y++;
+				if (Iter.Y > Iter.Rect.MaxY)
+				{
+					IterIdx--;
+					continue;
+				}
+			}
+
+			if (const FCell* Cell = FindCell(Iter.X, Iter.Y, Iter.Level))
+			{
+				// Collect items from this cell.
+				for (int32 Idx = Cell->First; Idx != INDEX_NONE; Idx = Items[Idx].Next)
+				{
+					OutResults.Add(Items[Idx].ID);
+				}
+
+				// Advance to region at finer level if it has any items.
+				if (Cell->ChildItemCount > 0)
+				{
+					check(Iter.Level > 0);
+					const int FinerLevel = Iter.Level - 1;
+					// The iteration rect is intersection between current coarse cell and finer level bounds (which is more accurate).
+					const FCellRect& FinerLevelRect = Rects[FinerLevel];
+					FCellRect CurrentRect(Iter.X * LevelRatio, Iter.Y * LevelRatio, Iter.X * LevelRatio + LevelRatio - 1, Iter.Y * LevelRatio + LevelRatio - 1);
+					FCellRect NewIterRect = IntersectRect(CurrentRect, FinerLevelRect);
+					// Advance if rect is not empty.
+					if (NewIterRect.MaxX >= NewIterRect.MinX && NewIterRect.MaxY >= NewIterRect.MinY)
+					{
+						FCellRectIterator& FinerIter = Iters[IterIdx];
+						FinerIter.Rect = NewIterRect;
+						FinerIter.X = NewIterRect.MinX;
+						FinerIter.Y = NewIterRect.MinY;
+						FinerIter.Level = FinerLevel;
+						IterIdx++;
+					}
+				}
+			}
+
+			// Advance iteration
+			Iter.X++;
+		}
+
+		// Everything from spill list
+		for (int32 Idx = SpillList; Idx != INDEX_NONE; Idx = Items[Idx].Next)
+		{
+			OutResults.Add(Items[Idx].ID);
+		}
+	}
+
+	/** Calculates cell based query rectangle. The bounds are expanded by half grid cell size because the items are stored for only one cell
+	  * based on their center and side. For that reason the items can overlap the neighbor cells by half the cell size.
+	  * @param Bounds - Query bounding box to quantize.
+	  * @param Level - Which level of the tree the to calculate the bounds for
+	  * @return Quantized rectangle representing the cell bounds at specific level of the tree, coordinates inclusive.
+	  */
+	FCellRect CalcQueryBounds(const FBox& Bounds, const int32 Level) const
+	{
+		FCellRect Result;
+		Result.MinX = FMath::FloorToInt(Bounds.Min.X * InvCellSize[Level] - 0.5f);
+		Result.MinY = FMath::FloorToInt(Bounds.Min.Y * InvCellSize[Level] - 0.5f);
+		Result.MaxX = FMath::FloorToInt(Bounds.Max.X * InvCellSize[Level] + 0.5f);
+		Result.MaxY = FMath::FloorToInt(Bounds.Max.Y * InvCellSize[Level] + 0.5f);
+		return Result;
+	}
+
+	/** Returns intersection of the two cell bounding rectangles.
+	  * @param Left - left hand side rectangle
+	  * @param Right - right hand side rectangle
+	  * @return Intersecting are between left and right.
+	  */
+	FCellRect IntersectRect(const FCellRect& Left, const FCellRect& Right) const
+	{
+		FCellRect Result;
+		Result.MinX = FMath::Max(Left.MinX, Right.MinX);
+		Result.MinY = FMath::Max(Left.MinY, Right.MinY);
+		Result.MaxX = FMath::Min(Left.MaxX, Right.MaxX);
+		Result.MaxY = FMath::Min(Left.MaxY, Right.MaxY);
+		return Result;
+	}
+
+	/** Levels down a coordinate using floor rounding.
+	  * @param X - Coordinate to level down
+	  * @return Coordinate at coarser level.
+	  */
+	static int32 LevelDown(int32 X)
+	{
+		X -= X < 0 ? (LevelRatio - 1) : 0;
+		return X / LevelRatio;
+	}
+
+	/** Finds grid level where the bounds fit inside a cell.
+	  * @param Bounds - Bounding box of the item.
+	  * @return Returns cell location of the item, or location at Level == INDEX_NONE if the item cannot fit in the grid.
+	  */
+	FCellLocation CalcCellLocation(const FBox& Bounds) const
+	{
+		FCellLocation Location(0, 0, 0);
+
+		const FVector Center = Bounds.GetCenter();
+		const float Diameter = FMath::Max(Bounds.Max.X - Bounds.Min.X, Bounds.Max.Y - Bounds.Min.Y);
+
+		for (Location.Level = 0; Location.Level < NumLevels; Location.Level++)
+		{
+			const int32 DiameterCells = FMath::CeilToInt(Diameter * InvCellSize[Location.Level]);
+			if (DiameterCells <= 1)
+			{
+				Location.X = FMath::FloorToInt(Center.X * InvCellSize[Location.Level]);
+				Location.Y = FMath::FloorToInt(Center.Y * InvCellSize[Location.Level]);
+				break;
+			}
+		}
+
+		if (Location.Level == NumLevels)
+		{
+			// Could not fit into any of the levels.
+			Location.X = 0;
+			Location.Y = 0;
+			Location.Level = INDEX_NONE;
+		}
+
+		return Location;
+	}
+
+	/** Finds the world box from a cell.
+	 * @params CellLocation - Cell location.
+	 * @return Returns the bounds of the cell.
+	 */
+	FBox CalcCellBounds(const FCellLocation& CellLocation)
+	{
+		const float Size = CellSize[CellLocation.Level];
+		const float X = CellLocation.X * Size;
+		const float Y = CellLocation.Y * Size;
+		return FBox(FVector(X, Y, 0.f), FVector(X + Size, Y + Size, 0.f));
+	}
+
+	/** Returns a cell for specific location and level, creates new cell if it does not exist.
+	  * @param X - Cell X coordinate.
+	  * @param Y - Cell Y coordinate.
+	  * @param Level - Grid Level.
+	  * @return Reference to cell at specified location.
+	  */
+	FCell& FindOrAddCell(const int X, const int Y, const int Level)
+	{
+		FCell CellToFind(X, Y, Level);
+		const uint32 Hash = GetTypeHash(CellToFind);
+		FCell* Cell = Cells.FindByHash(Hash, CellToFind);
+		if (Cell != nullptr)
+		{
+			return *Cell;
+		}
+		FSetElementId Index = Cells.AddByHash(Hash, FCell(X, Y, Level));
+		return Cells[Index];
+	}
+
+	/** Returns a cell for specific location and level.
+	  * @param X - Cell X coordinate.
+	  * @param Y - Cell Y coordinate.
+	  * @param Level - Grid Level.
+	  * @return Pointer to cell at specified location, or return nullptr if the cell does not exist.
+	  */
+	FCell* FindCellMutable(const int X, const int Y, const int Level)
+	{
+		FCell CellToFind(X, Y, Level);
+		const uint32 Hash = GetTypeHash(CellToFind);
+		return Cells.FindByHash(Hash, CellToFind);
+	}
+
+	/** Returns a cell for specific location and level.
+	  * @param X - Cell X coordinate.
+	  * @param Y - Cell Y coordinate.
+	  * @param Level - Grid Level.
+	  * @return Pointer to cell at specified location, or return nullptr if the cell does not exist.
+	  */
+	const FCell* FindCell(const int X, const int Y, const int Level) const
+	{
+		FCell CellToFind(X, Y, Level);
+		const uint32 Hash = GetTypeHash(CellToFind);
+		return Cells.FindByHash(Hash, CellToFind);
+	}
+
+	/** @return Cell size at specific grid level. */
+	float GetCellSize(const int32 Level) const { return CellSize[Level]; }
+
+	/** @return 1/cellSize at specific grid level. */
+	float GetInvCellSize(const int32 Level) const { return InvCellSize[Level]; }
+
+	/** @return Array containing all the items. */
+	int32 GetLevelItemCount(const int32 Level) const { return LevelItemCount[Level]; }
+
+	/** @return Set containing all the cells. */
+	const TSet<FCell>& GetCells() const { return Cells; }
+
+	/** @return Sparse array containing all the items. */
+	const TSparseArray<FItem>& GetItems() const { return Items; }
+
+	/** @return Index in items of the first item on the spill list. */
+	int32 GetFirstSpillListItem() const { return SpillList; }
+
+protected:
+	TStaticArray<float, NumLevels> CellSize;		/** Lowest level cell size */
+	TStaticArray<float, NumLevels> InvCellSize;		/** 1 / CellSize */
+	TStaticArray<int32, NumLevels> LevelItemCount;	/** Number items per level, can be used to skip certain levels. */
+	TSet<FCell> Cells;								/** TSet uses hash buckets to locate items. */
+	TSparseArray<FItem> Items;						/** TSparseArray uses freelist to recycled used items. */
+	int32 SpillList = INDEX_NONE;					/** Linked list of items that did not fit into any of the levels */
+};
diff --git a/Engine/Source/Runtime/AIModule/Public/IndexedHandle.h b/Engine/Source/Runtime/AIModule/Public/IndexedHandle.h
new file mode 100644
--- /dev/null
+++ b/Engine/Source/Runtime/AIModule/Public/IndexedHandle.h
@@ -0,0 +1,331 @@
+// Copyright Epic Games, Inc. All Rights Reserved.
+
+#pragma once
+
+#include "CoreMinimal.h"
+#include "Containers/Array.h"
+#include <atomic>
+#include "IndexedHandle.generated.h"
+
+#if !(UE_BUILD_SHIPPING || UE_BUILD_TEST)
+#define UE_DO_INDEXED_HANDLE_MANAGER_ID 1
+#else
+#define UE_DO_INDEXED_HANDLE_MANAGER_ID 0
+#endif
+
+
+/**
+ * Index based handle that doesn't use a serial number.
+ * For fast access in to index based data structures when we don't expect handles to be able to be stale
+ */
+USTRUCT()
+struct AIMODULE_API FSimpleIndexedHandleBase
+{
+	GENERATED_BODY()
+
+	FSimpleIndexedHandleBase() = default;
+
+	/** @param InIndex - passing INDEX_NONE will make this handle Invalid */
+	FSimpleIndexedHandleBase(int32 InIndex) : Index(InIndex)
+	{
+	};
+
+	bool operator==(const FSimpleIndexedHandleBase& Other) const
+	{
+		return Index == Other.Index;
+	}
+
+	bool operator!=(const FSimpleIndexedHandleBase& Other) const
+	{
+		return !operator==(Other);
+	}
+
+	/** @return INDEX_NONE if invalid or >=0 for a potentially valid handle */
+	int32 GetIndex() const { return Index; }
+
+	/** @return true if this handle is valid, this doesn't check that it's in the correct range for the underlying resource it's providing a handle for */
+	bool IsValid() const { return Index >= 0; }
+
+	/**
+	 * boolean operator useful when embedding declaration in if statement
+	 * @return same as IsValid()
+	 * @see IsValid
+	 */
+	operator bool() const { return IsValid(); }
+
+	/** Makes the handle Invalid */
+	void Invalidate() { Index = INDEX_NONE; }
+
+	/** @param InIndex - passing INDEX_NONE will make this handle Invalid */
+	void SetIndex(int32 InIndex)
+	{
+		Index = InIndex;
+	}
+
+	friend uint32 GetTypeHash(const FSimpleIndexedHandleBase& Handle)
+	{
+		return static_cast<uint32>(Handle.GetIndex());
+	}
+
+protected:
+	int32 Index = INDEX_NONE;
+};
+
+/** Index based handle that has a serial number to verify stale handles. For fast safe access in to index based data structures */
+USTRUCT()
+struct AIMODULE_API FIndexedHandleBase : public FSimpleIndexedHandleBase
+{
+	GENERATED_BODY()
+
+	FIndexedHandleBase() = default;
+
+	/** @note passing INDEX_NONE as index will make this handle Invalid */
+	FIndexedHandleBase(int32 InIndex, uint32 InSerialNumber) : FSimpleIndexedHandleBase(InIndex), SerialNumber(InSerialNumber)
+	{
+	};
+
+	bool operator==(const FIndexedHandleBase& Other) const
+	{
+		return FSimpleIndexedHandleBase::operator==(Other) && SerialNumber == Other.SerialNumber;
+	}
+
+	bool operator!=(const FIndexedHandleBase& Other) const
+	{
+		return !operator==(Other);
+	}
+
+	uint32 GetSerialNumber() const { return SerialNumber; }
+	void SetSerialNumber(uint32 InSerialNumber) { SerialNumber = InSerialNumber; }
+
+	friend uint32 GetTypeHash(const FIndexedHandleBase& Handle)
+	{
+		return HashCombine(static_cast<uint32>(Handle.GetIndex()), Handle.GetSerialNumber());
+	}
+
+protected:
+	uint32 SerialNumber = 0;
+
+#if UE_DO_INDEXED_HANDLE_MANAGER_ID
+public:
+	uint32 ManagerID = 0;
+#endif
+};
+
+/** Compact Index based handle that has a serial number to verify stale handles. For fast safe access in to index based data structures. */
+USTRUCT()
+struct AIMODULE_API FCompactIndexedHandleBase
+{
+	GENERATED_BODY()
+
+	FCompactIndexedHandleBase() = default;
+
+	/** @note passing INDEX_NONE as index will make this handle Invalid */
+	FCompactIndexedHandleBase(int16 InIndex, uint16 InSerialNumber) : Index(InIndex), SerialNumber(InSerialNumber)
+	{
+	};
+
+	bool operator==(const FCompactIndexedHandleBase& Other) const
+	{
+		return Index == Other.Index && SerialNumber == Other.SerialNumber;
+	}
+
+	bool operator!=(const FCompactIndexedHandleBase& Other) const
+	{
+		return !operator==(Other);
+	}
+
+	/** @return INDEX_NONE if invalid or >=0 for a potentially valid handle */
+	int32 GetIndex() const { return Index; }
+
+	/** @return true if this handle is valid, this doesn't check that it's in the correct range for the underlying resource it's providing a handle for */
+	bool IsValid() const { return Index >= 0; }
+
+	/** 
+     * boolean operator useful when embedding declaration in if statement
+	 * @return same as IsValid()
+	 * @see IsValid
+     */
+	operator bool() const { return IsValid(); }
+
+	/** Makes the handle Invalid */
+	void Invalidate() { Index = INDEX_NONE; }
+
+	/** @param InIndex - passing INDEX_NONE will make this handle Invalid */
+	void SetIndex(int16 InIndex)
+	{
+		Index = InIndex;
+	}
+
+	uint16 GetSerialNumber() const { return SerialNumber; }
+	void SetSerialNumber(int16 InSerialNumber) { SerialNumber = InSerialNumber; }
+
+	friend uint32 GetTypeHash(const FCompactIndexedHandleBase& Handle)
+	{
+		return (static_cast<uint32>(Handle.GetIndex()) << 16) + static_cast<uint32>(Handle.GetSerialNumber());
+	}
+
+protected:
+	int16 Index = INDEX_NONE;
+	uint16 SerialNumber = 0;
+
+#if UE_DO_INDEXED_HANDLE_MANAGER_ID
+public:
+	uint32 ManagerID = 0;
+#endif
+};
+
+/** Handle Manager meant for FIndexedHandleBase and FCompactIndexedHandleBase derived classes, handles are given out from a freelist and are zero based and consecutive in nature,
+ *  so ideal for being used as indices in to arrays.
+ *  Using bOptimizeHandleReuse performance impact on releasing a handle. Because new handles will always use the smallest index, this will help greatly the compaction code(ShrinkHandles) and reduce the total number of handles.
+ *  This is very useful when you need to iterate through every valid handles often.
+ */
+template<typename TIndexedHandle, typename TIndexType, typename TSerialType, bool bOptimizeHandleReuse = false>
+struct FIndexedHandleManagerBase
+{
+protected:
+	typedef TArray<TIndexedHandle> FHandleArray;
+
+public:
+	FIndexedHandleManagerBase()
+	{
+#if UE_DO_INDEXED_HANDLE_MANAGER_ID
+		ManagerID = ManagerIDCounter.fetch_add(1);
+#endif
+	}
+
+	TIndexedHandle GetNextHandle()
+	{
+		TIndexedHandle* Handle = nullptr;
+
+		//if no free handle indices then create one
+		if (FreeHandleIndices.Num() == 0)
+		{
+			checkf(Handles.Num() <= static_cast<typename FHandleArray::SizeType>(TNumericLimits<TIndexType>::Max()), TEXT("Handles are overflowing the numeric limits of the handle TIndexType!"));
+			Handle = &Handles.Emplace_GetRef(Handles.Num(), SerialNumberCounter.fetch_add(1));
+		}
+		else //otherwise use an existing handle index
+		{
+			const TIndexType Idx = FreeHandleIndices.Pop();
+			Handle = &Handles[Idx];
+			checkf(!Handle->IsValid(), TEXT("Free handle must be set invalid before reuse"));
+			Handle->SetIndex(Idx);
+			Handle->SetSerialNumber(SerialNumberCounter.fetch_add(1));
+		}
+
+#if UE_DO_INDEXED_HANDLE_MANAGER_ID
+		Handle->ManagerID = ManagerID;
+#endif
+
+		return *Handle;
+	}
+
+	bool RemoveHandle(TIndexedHandle IndexedHandle)
+	{
+		const bool bIsValid = IsValidHandle(IndexedHandle);
+
+		if (ensureMsgf(bIsValid, TEXT("Trying to Remove Invalid Handle with Index %d, Serial Number %d"), IndexedHandle.GetIndex(), IndexedHandle.GetSerialNumber()))
+		{
+			if(bOptimizeHandleReuse)
+			{
+			    const TIndexType Idx = IndexedHandle.GetIndex();
+			    const int32 InsertPosition = Algo::LowerBound(FreeHandleIndices, Idx, [](const TIndexType& Lhs, const TIndexType& Rhs) {return Lhs > Rhs;});
+			    FreeHandleIndices.Insert(Idx, InsertPosition);
+			}
+			else
+			{
+				FreeHandleIndices.Add(IndexedHandle.GetIndex());
+			}
+
+			Handles[IndexedHandle.GetIndex()].Invalidate();
+
+			//increase the serial numbers when a handle is removed
+			Handles[IndexedHandle.GetIndex()].SetSerialNumber(SerialNumberCounter.fetch_add(1));
+		}
+
+		return bIsValid;
+	}
+
+	bool IsValidHandle(TIndexedHandle IndexedHandle) const
+	{
+#if UE_DO_INDEXED_HANDLE_MANAGER_ID
+		if (!ensureMsgf(ManagerID == IndexedHandle.ManagerID, TEXT("ManagerID %d does not match IndexedHandle.ManagerID %d, handles must only be used by the same manager they were created by!"), IndexedHandle.ManagerID))
+		{
+			return false;
+		}
+#endif // UE_DO_INDEXED_HANDLE_MANAGER_ID
+
+		return ((IndexedHandle.GetIndex() >= 0) && (IndexedHandle.GetIndex() < Handles.Num())
+			&& (IndexedHandle.GetSerialNumber() == Handles[IndexedHandle.GetIndex()].GetSerialNumber())
+			&& (IndexedHandle.GetIndex() == Handles[IndexedHandle.GetIndex()].GetIndex()));
+	}
+
+	const TArray<TIndexedHandle>& GetHandles() const { return Handles; }
+
+	/** returns the number of used handles, will be greater than or equal to zero */
+	int32 CalcNumUsedHandles() const { return Handles.Num() - FreeHandleIndices.Num(); };
+
+	FIndexedHandleManagerBase& operator=(const FIndexedHandleManagerBase& Other)
+	{
+		Handles = Other.Handles;
+		FreeHandleIndices = Other.FreeHandleIndices;
+		//Leave SerialNumberCounter well alone
+
+#if UE_DO_INDEXED_HANDLE_MANAGER_ID
+		//assign new ManagerID to the new Manager
+		ManagerID = ManagerIDCounter.fetch_add(1);
+#endif // UE_DO_INDEXED_HANDLE_MANAGER_ID
+		return *this;
+	}
+
+	/** Attempts to shrink the Handles array if there are contiguous free slots at the end of the Array
+	 *	@return the new Num items in the handles array
+	 */
+	int32 ShrinkHandles()
+	{
+		while (Handles.Num() > 0 && FreeHandleIndices.Num() > 0 && FreeHandleIndices.Remove(Handles.Num() - 1))
+		{
+			Handles.Pop(false);
+		}
+
+		return Handles.Num();
+	}
+
+
+	void Reset()
+	{
+		Handles.Reset();
+		FreeHandleIndices.Reset();
+		//Don't reset the SerialNumberCounter as we need that to increase from where it left off to keep handle consistency
+	}
+
+protected:
+	/** Handles stored as a free list sparse array, Handle entries that are free will have a GetIndex() == INDEX_NONE */
+	FHandleArray Handles;
+	TArray<TIndexType> FreeHandleIndices;
+
+	static std::atomic<TSerialType> SerialNumberCounter;
+
+#if UE_DO_INDEXED_HANDLE_MANAGER_ID
+	/** ID used in non shipping / test builds to ensure the same manager that is used to create handles is also used to remove them */
+	static std::atomic<uint32> ManagerIDCounter;
+	uint32 ManagerID = 0;
+#endif // UE_DO_INDEXED_HANDLE_MANAGER_ID
+};
+
+template<typename TIndexedHandle, typename TIndexType, typename TSerialType, bool bOptimizeHandleReuse>
+std::atomic<TSerialType> FIndexedHandleManagerBase<TIndexedHandle, TIndexType, TSerialType, bOptimizeHandleReuse>::SerialNumberCounter = { 0 };
+
+#if UE_DO_INDEXED_HANDLE_MANAGER_ID
+template<typename TIndexedHandle, typename TIndexType, typename TSerialType, bool bOptimizeHandleReuse>
+std::atomic<uint32> FIndexedHandleManagerBase<TIndexedHandle, TIndexType, TSerialType, bOptimizeHandleReuse>::ManagerIDCounter = { 0 };
+#endif // UE_DO_INDEXED_HANDLE_MANAGER_ID
+
+template<typename TIndexedHandle, bool bOptimizeHandleReuse = false>
+struct FIndexedHandleManager : public FIndexedHandleManagerBase<TIndexedHandle, int32, uint32, bOptimizeHandleReuse>
+{
+};
+
+template<typename TIndexedHandle, bool bOptimizeHandleReuse = false>
+struct FCompactIndexedHandleManager : public FIndexedHandleManagerBase<TIndexedHandle, int16, uint16, bOptimizeHandleReuse>
+{
+};
diff --git a/Engine/Source/Runtime/AIModule/Public/SequentialID.h b/Engine/Source/Runtime/AIModule/Public/SequentialID.h
new file mode 100644
--- /dev/null
+++ b/Engine/Source/Runtime/AIModule/Public/SequentialID.h
@@ -0,0 +1,36 @@
+// Copyright Epic Games, Inc. All Rights Reserved.
+
+#pragma once
+
+#include "SequentialID.generated.h"
+
+USTRUCT()
+struct AIMODULE_API FSequentialIDBase
+{
+	GENERATED_BODY()
+
+	FSequentialIDBase() = default;
+	FSequentialIDBase(const FSequentialIDBase& Other) = default;
+
+	explicit FSequentialIDBase(uint32 InID) : Value(InID) {}
+
+	bool operator==(const FSequentialIDBase& Other) const { return Value == Other.Value; }
+	bool operator!=(const FSequentialIDBase& Other) const { return Value != Other.Value; }
+
+	bool IsValid() const { return Value != InvalidID; }
+	bool IsInvalid() const { return Value == InvalidID; }
+
+	void Invalidate() { Value = InvalidID; }
+
+	uint32 GetValue() const { return Value; }
+
+	FString Describe() const { return FString::Printf(TEXT("Id[%d]"), Value); }
+
+	friend uint32 GetTypeHash(const FSequentialIDBase& SequentialID) { return SequentialID.GetValue(); }
+
+protected:
+	static constexpr uint32 InvalidID = 0;
+
+	UPROPERTY()
+	uint32 Value = InvalidID;
+};
\ No newline at end of file
diff --git a/Engine/Source/Runtime/Engine/Private/World.cpp b/Engine/Source/Runtime/Engine/Private/World.cpp
--- a/Engine/Source/Runtime/Engine/Private/World.cpp
+++ b/Engine/Source/Runtime/Engine/Private/World.cpp
@@ -4667,6 +4667,13 @@ void UWorld::InitializeActorsForPlay(const FURL& InURL, bool bResetTime, FRegist
 
 void UWorld::BeginPlay()
 {
+	const TArray<UWorldSubsystem*>& WorldSubsystems = SubsystemCollection.GetSubsystemArray<UWorldSubsystem>(UWorldSubsystem::StaticClass());
+
+	for (UWorldSubsystem* WorldSubsystem : WorldSubsystems)
+	{
+		WorldSubsystem->OnWorldBeginPlay(*this);
+	}
+
 	AGameModeBase* const GameMode = GetAuthGameMode();
 	if (GameMode)
 	{
diff --git a/Engine/Source/Runtime/Engine/Public/Subsystems/WorldSubsystem.h b/Engine/Source/Runtime/Engine/Public/Subsystems/WorldSubsystem.h
--- a/Engine/Source/Runtime/Engine/Public/Subsystems/WorldSubsystem.h
+++ b/Engine/Source/Runtime/Engine/Public/Subsystems/WorldSubsystem.h
@@ -10,9 +10,7 @@
 /**
  * UWorldSubsystem
  * Base class for auto instanced and initialized systems that share the lifetime of a UWorld
-
  */
-
 UCLASS(Abstract)
 class ENGINE_API UWorldSubsystem : public USubsystem
 {
@@ -26,6 +24,9 @@ public:
 
 	/** Called once all UWorldSubsystems have been initialized */
 	virtual void PostInitialize() {}
+	
+	/** Called when world is ready to start gameplay before the game mode transitions to the correct state and call BeginPlay on all actors */
+	virtual void OnWorldBeginPlay(UWorld& InWorld) {};
 
 	/** Updates sub-system required streaming levels (called by world's UpdateStreamingState function) */
 	virtual void UpdateStreamingState() {}
@@ -37,9 +38,7 @@ protected:
 /**
  * UTickableWorldSubsystem
  * Base class for auto instanced and initialized systems that share the lifetime of a UWorld and are ticking along with it
-
  */
-
 UCLASS(Abstract)
 class ENGINE_API UTickableWorldSubsystem : public UWorldSubsystem, public FTickableGameObject
 {
