Fixed bug where PainCausingVolume was not applying timer-based damage to all damageable Actors, only Pawns. Added GetPhysicsVolume query to Actor, deprecating the Pawn-only query. (Adaptation of PR 8228)

PR #8228

diff --git a/Engine/Source/Runtime/AIModule/Private/AIController.cpp b/Engine/Source/Runtime/AIModule/Private/AIController.cpp
--- a/Engine/Source/Runtime/AIModule/Private/AIController.cpp
+++ b/Engine/Source/Runtime/AIModule/Private/AIController.cpp
@@ -1052,7 +1052,7 @@ bool AAIController::ShouldSyncBlackboardWith(const UBlackboardComponent& OtherBl
 bool AAIController::SuggestTossVelocity(FVector& OutTossVelocity, FVector Start, FVector End, float TossSpeed, bool bPreferHighArc, float CollisionRadius, bool bOnlyTraceUp)
 {
 	// pawn's physics volume gets 2nd priority
-	APhysicsVolume const* const PhysicsVolume = GetPawn() ? GetPawn()->GetPawnPhysicsVolume() : NULL;
+	APhysicsVolume const* const PhysicsVolume = GetPawn() ? GetPawn()->GetPhysicsVolume() : nullptr;
 	float const GravityOverride = PhysicsVolume ? PhysicsVolume->GetGravityZ() : 0.f;
 	ESuggestProjVelocityTraceOption::Type const TraceOption = bOnlyTraceUp ? ESuggestProjVelocityTraceOption::OnlyTraceWhileAscending : ESuggestProjVelocityTraceOption::TraceFullPath;
 
diff --git a/Engine/Source/Runtime/Engine/Classes/GameFramework/Actor.h b/Engine/Source/Runtime/Engine/Classes/GameFramework/Actor.h
--- a/Engine/Source/Runtime/Engine/Classes/GameFramework/Actor.h
+++ b/Engine/Source/Runtime/Engine/Classes/GameFramework/Actor.h
@@ -2501,6 +2501,8 @@ public:
 	/** See if the root component has Mobility of EComponentMobility::Movable */
 	bool IsRootComponentMovable() const;
 
+	/** Get the physics volume that is currently applied to this Actor (there can only ever be one) */
+	virtual APhysicsVolume* GetPhysicsVolume() const;
 
 	//~=============================================================================
 	// Actor ticking
diff --git a/Engine/Source/Runtime/Engine/Classes/GameFramework/Pawn.h b/Engine/Source/Runtime/Engine/Classes/GameFramework/Pawn.h
--- a/Engine/Source/Runtime/Engine/Classes/GameFramework/Pawn.h
+++ b/Engine/Source/Runtime/Engine/Classes/GameFramework/Pawn.h
@@ -202,7 +202,11 @@ public:
 	void SetRemoteViewPitch(float NewRemoteViewPitch);
 
 	/** Return Physics Volume for this Pawn */
-	virtual APhysicsVolume* GetPawnPhysicsVolume() const;
+	UE_DEPRECATED(5.0, "GetPawnPhysicsVolume is deprecated. Please use GetPhysicsVolume instead.")
+	virtual APhysicsVolume* GetPawnPhysicsVolume() const; 
+
+	/** Return Physics Volume for this Pawn */
+	virtual APhysicsVolume* GetPhysicsVolume() const override;
 
 	/** Gets the owning actor of the Movement Base Component on which the pawn is standing. */
 	UFUNCTION(BlueprintPure, Category=Pawn)
diff --git a/Engine/Source/Runtime/Engine/Private/Actor.cpp b/Engine/Source/Runtime/Engine/Private/Actor.cpp
--- a/Engine/Source/Runtime/Engine/Private/Actor.cpp
+++ b/Engine/Source/Runtime/Engine/Private/Actor.cpp
@@ -2222,6 +2222,16 @@ bool AActor::IsRootComponentMovable() const
 	return(RootComponent != nullptr && RootComponent->Mobility == EComponentMobility::Movable);
 }
 
+APhysicsVolume* AActor::GetPhysicsVolume() const
+{
+	if (const USceneComponent* ActorRootComponent = GetRootComponent())
+	{
+		return ActorRootComponent->GetPhysicsVolume();
+	}
+
+	return GetWorld()->GetDefaultPhysicsVolume();
+}
+
 FVector AActor::GetTargetLocation(AActor* RequestedBy) const
 {
 	return GetActorLocation();
diff --git a/Engine/Source/Runtime/Engine/Private/PainCausingVolume.cpp b/Engine/Source/Runtime/Engine/Private/PainCausingVolume.cpp
--- a/Engine/Source/Runtime/Engine/Private/PainCausingVolume.cpp
+++ b/Engine/Source/Runtime/Engine/Private/PainCausingVolume.cpp
@@ -55,18 +55,13 @@ void APainCausingVolume::PainTimer()
 	if (bPainCausing)
 	{
 		TSet<AActor*> TouchingActors;
-		GetOverlappingActors(TouchingActors, APawn::StaticClass());
+		GetOverlappingActors(TouchingActors);
 
 		for (AActor* const A : TouchingActors)
 		{
-			if (IsValid(A) && A->CanBeDamaged())
+			if (IsValid(A) && A->CanBeDamaged() && A->GetPhysicsVolume() == this)
 			{
-				// @todo physicsVolume This won't work for other actor. Need to fix it properly
-				APawn* PawnA = Cast<APawn>(A);
-				if (PawnA && PawnA->GetPawnPhysicsVolume() == this)
-				{
-					CausePainTo(A);
-				}
+				CausePainTo(A);
 			}
 		}
 
diff --git a/Engine/Source/Runtime/Engine/Private/Pawn.cpp b/Engine/Source/Runtime/Engine/Private/Pawn.cpp
--- a/Engine/Source/Runtime/Engine/Private/Pawn.cpp
+++ b/Engine/Source/Runtime/Engine/Private/Pawn.cpp
@@ -823,18 +823,19 @@ void APawn::Restart()
 	RecalculateBaseEyeHeight();
 }
 
-class APhysicsVolume* APawn::GetPawnPhysicsVolume() const
+APhysicsVolume* APawn::GetPawnPhysicsVolume() const
 {
-	const UPawnMovementComponent* MovementComponent = GetMovementComponent();
-	if (MovementComponent)
+	return GetPhysicsVolume();
+}
+
+APhysicsVolume* APawn::GetPhysicsVolume() const
+{
+	if (const UPawnMovementComponent* MovementComponent = GetMovementComponent())
 	{
 		return MovementComponent->GetPhysicsVolume();
 	}
-	else if (GetRootComponent())
-	{
-		return GetRootComponent()->GetPhysicsVolume();
-	}
-	return GetWorld()->GetDefaultPhysicsVolume();
+
+	return Super::GetPhysicsVolume();
 }
 
 
