Make FEQSParamsExporter::AddNamedValuesFromObject properly work with composite generator by recursively adding values from subnode stored in array.
#rnx

diff --git a/Engine/Source/Runtime/AIModule/Private/EnvironmentQuery/EnvQuery.cpp b/Engine/Source/Runtime/AIModule/Private/EnvironmentQuery/EnvQuery.cpp
--- a/Engine/Source/Runtime/AIModule/Private/EnvironmentQuery/EnvQuery.cpp
+++ b/Engine/Source/Runtime/AIModule/Private/EnvironmentQuery/EnvQuery.cpp
@@ -46,6 +46,32 @@ namespace FEQSParamsExporter
 	{
 		for (FProperty* TestProperty = QueryNode.GetClass()->PropertyLink; TestProperty; TestProperty = TestProperty->PropertyLinkNext)
 		{
+			FArrayProperty* TestArray = CastField<FArrayProperty>(TestProperty);
+			if (TestArray)
+			{
+				FObjectPropertyBase* ArrayInnerProperty = CastField<FObjectPropertyBase>(TestArray->Inner);
+				if (ArrayInnerProperty == nullptr)
+				{
+					continue;
+				}
+
+				if (ArrayInnerProperty->PropertyClass == nullptr || ArrayInnerProperty->PropertyClass->IsChildOf<UEnvQueryNode>() == false)
+				{
+					continue;
+				}
+
+				FScriptArrayHelper ArrayHelper(TestArray, TestArray->ContainerPtrToValuePtr<void>(&QueryNode));
+
+				for (int32 SubNodeIndex = 0; SubNodeIndex < ArrayHelper.Num(); ++SubNodeIndex)
+				{
+					const UEnvQueryNode** SubNode = reinterpret_cast<const UEnvQueryNode**>(ArrayHelper.GetRawPtr(SubNodeIndex));
+					if (SubNode && *SubNode)
+					{
+						AddNamedValuesFromObject(QueryOwner, **SubNode, NamedValues, RequiredParams);	
+					}
+				}
+			}
+
 			FStructProperty* TestStruct = CastField<FStructProperty>(TestProperty);
 			if (TestStruct == NULL)
 			{
