Being consistent with services, they should be ticked on the frame becomming relevant. They are being ticked when they are on composite nodes but not they are on task nodes.

diff --git a/Engine/Source/Developer/AITestSuite/Private/Tests/BTTest.cpp b/Engine/Source/Developer/AITestSuite/Private/Tests/BTTest.cpp
--- a/Engine/Source/Developer/AITestSuite/Private/Tests/BTTest.cpp
+++ b/Engine/Source/Developer/AITestSuite/Private/Tests/BTTest.cpp
@@ -1315,14 +1315,16 @@ struct FAITest_BTServiceInstantTask : public FAITest_SimpleBT
 		{
 			FBTBuilder::AddTask(CompNode, 0, EBTNodeResult::Succeeded);
 			{
-				FBTBuilder::WithTaskServiceLog(CompNode, 1, 2);
+				FBTBuilder::WithTaskServiceLog(CompNode, 1/*ActivationIndex*/, 2/*DeactivationIndex*/, 4/*TickIndex*/);
 			}
 
 			FBTBuilder::AddTask(CompNode, 3, EBTNodeResult::Succeeded);
 		}
 
 		ExpectedResult.Add(1);
+		ExpectedResult.Add(4);
 		ExpectedResult.Add(0);
+		ExpectedResult.Add(4); // Extra service tick because aux nodes are always ticked before any pending requests are processed
 		ExpectedResult.Add(2);
 		ExpectedResult.Add(3);
 	}
@@ -1337,14 +1339,18 @@ struct FAITest_BTServiceLatentTask : public FAITest_SimpleBT
 		{
 			FBTBuilder::AddTask(CompNode, 0, EBTNodeResult::Succeeded, 2);
 			{
-				FBTBuilder::WithTaskServiceLog(CompNode, 1, 2);
+				FBTBuilder::WithTaskServiceLog(CompNode, 1/*ActivationIndex*/, 2/*DeactivationIndex*/, 4/*TickIndex*/);
 			}
 
 			FBTBuilder::AddTask(CompNode, 3, EBTNodeResult::Succeeded);
 		}
 
 		ExpectedResult.Add(1);
+		ExpectedResult.Add(4);
 		ExpectedResult.Add(0);
+		ExpectedResult.Add(4);
+		ExpectedResult.Add(4);
+		ExpectedResult.Add(4);// Extra service tick because aux nodes are always ticked before any pending requests are processed
 		ExpectedResult.Add(2);
 		ExpectedResult.Add(3);
 	}
@@ -1364,14 +1370,17 @@ struct FAITest_BTServiceAbortingTask : public FAITest_SimpleBT
 
 			FBTBuilder::AddTaskLatentFlags(CompNode, EBTNodeResult::Succeeded, 1, TEXT("Bool1"), 1, 2, 0, NAME_None, 3, 4);
 			{
-				FBTBuilder::WithTaskServiceLog(CompNode, 5, 6);
+				FBTBuilder::WithTaskServiceLog(CompNode, 5/*ActivationIndex*/, 6/*DeactivationIndex*/, 8/*TickIndex*/);
 			}
 
 			FBTBuilder::AddTask(CompNode, 7, EBTNodeResult::Succeeded);
 		}
 
 		ExpectedResult.Add(5);
+		ExpectedResult.Add(8);
 		ExpectedResult.Add(1);
+		ExpectedResult.Add(8);
+		ExpectedResult.Add(8); // Extra service tick because aux nodes are always ticked before any pending requests are processed
 		ExpectedResult.Add(3);
 		ExpectedResult.Add(4);
 		ExpectedResult.Add(6);
diff --git a/Engine/Source/Runtime/AIModule/Private/BehaviorTree/BehaviorTreeComponent.cpp b/Engine/Source/Runtime/AIModule/Private/BehaviorTree/BehaviorTreeComponent.cpp
--- a/Engine/Source/Runtime/AIModule/Private/BehaviorTree/BehaviorTreeComponent.cpp
+++ b/Engine/Source/Runtime/AIModule/Private/BehaviorTree/BehaviorTreeComponent.cpp
@@ -117,7 +117,7 @@ void UBehaviorTreeComponent::SetComponentTickEnabled(bool bEnabled)
 
 	// If enabling the component, this acts like a new component to tick in the TickTaskManager
 	// So act like the component was never ticked
-	if(!bWasEnabled && IsComponentTickEnabled())
+	if (!bWasEnabled && IsComponentTickEnabled())
 	{
 		bTickedOnce = false;
 		ScheduleNextTick(0.0f);
@@ -1271,7 +1271,7 @@ void UBehaviorTreeComponent::ApplySearchData(UBTNode* NewActiveNode)
 			FBehaviorTreeInstance& InstanceInfo = InstanceStack[UpdateInfo.InstanceIndex];
 			uint8* NodeMemory = UpdateInfo.AuxNode->GetNodeMemory<uint8>(InstanceInfo);
 
-            // We do not care about the next needed DeltaTime, it will be recalculated in the tick later.
+			// We do not care about the next needed DeltaTime, it will be recalculated in the tick later.
 			float NextNeededDeltaTime = 0.0f;
 			UpdateInfo.AuxNode->WrappedTickNode(*this, NodeMemory, CurrentFrameDeltaSeconds, NextNeededDeltaTime);
 		}
@@ -2026,9 +2026,8 @@ void UBehaviorTreeComponent::ExecuteTask(UBTTaskNode* TaskNode)
 	FBehaviorTreeInstance& ActiveInstance = InstanceStack[ActiveInstanceIdx];
 
 	// task service activation is not part of search update (although deactivation is, through DeactivateUpTo), start them before execution
-	for (int32 ServiceIndex = 0; ServiceIndex < TaskNode->Services.Num(); ServiceIndex++)
+	for (UBTService* ServiceNode : TaskNode->Services)
 	{
-		UBTService* ServiceNode = TaskNode->Services[ServiceIndex];
 		uint8* NodeMemory = (uint8*)ServiceNode->GetNodeMemory<uint8>(ActiveInstance);
 
 		ActiveInstance.AddToActiveAuxNodes(ServiceNode);
@@ -2037,6 +2036,20 @@ void UBehaviorTreeComponent::ExecuteTask(UBTTaskNode* TaskNode)
 		ServiceNode->WrappedOnBecomeRelevant(*this, NodeMemory);
 	}
 
+	// Services were already ticked for this frame, need to tick the new ones. 
+	UWorld* MyWorld = GetWorld();
+	const float CurrentFrameDeltaSeconds = MyWorld ? MyWorld->GetDeltaSeconds() : 0.0f;
+	for (UBTService* ServiceNode : TaskNode->Services)
+	{
+		uint8* NodeMemory = (uint8*)ServiceNode->GetNodeMemory<uint8>(ActiveInstance);
+
+		UE_VLOG(GetOwner(), LogBehaviorTree, Log, TEXT("Ticking task service: %s"), *UBehaviorTreeTypes::DescribeNodeHelper(ServiceNode));
+
+		// We do not care about the next needed DeltaTime, it will be recalculated in the tick later.
+		float NextNeededDeltaTime = 0.0f;
+		ServiceNode->WrappedTickNode(*this, NodeMemory, CurrentFrameDeltaSeconds, NextNeededDeltaTime);
+	}
+
 	ActiveInstance.ActiveNode = TaskNode;
 	ActiveInstance.ActiveNodeType = EBTActiveNode::ActiveTask;
 
