Very old bug fix: Lower priority node requesting execution ensure fix!

diff --git a/Engine/Source/Developer/AITestSuite/Private/Tests/BTTest.cpp b/Engine/Source/Developer/AITestSuite/Private/Tests/BTTest.cpp
--- a/Engine/Source/Developer/AITestSuite/Private/Tests/BTTest.cpp
+++ b/Engine/Source/Developer/AITestSuite/Private/Tests/BTTest.cpp
@@ -1307,6 +1307,63 @@ struct FAITest_BTSubtreeAbortOut : public FAITest_SimpleBT
 };
 IMPLEMENT_AI_LATENT_TEST(FAITest_BTSubtreeAbortOut, "System.AI.Behavior Trees.Subtree: abort out")
 
+/* This unittest come from UDN case 00320135 */
+struct FAITest_BTSubtreeAbortOutToLowerPrio : public FAITest_SimpleBT
+{
+	FAITest_BTSubtreeAbortOutToLowerPrio()
+	{
+		UBehaviorTree* ChildAsset = &FBTBuilder::CreateBehaviorTree(*BTAsset);
+		if (ChildAsset)
+		{
+			AddAutoDestroyObject(*ChildAsset);
+			UBTCompositeNode& CompNode = FBTBuilder::AddSelector(*ChildAsset);
+			{
+				UBTCompositeNode& CompNode2 = FBTBuilder::AddSequence(CompNode);
+				{
+					FBTBuilder::WithDecoratorBlackboard(CompNode, EBasicKeyOperation::Set, EBTFlowAbortMode::Both, TEXT("Bool2"), 3/*LogIndexBecomeRelevant*/, 4/*LogIndexCeaseRelevant*/);
+					FBTBuilder::AddTaskFlagChange(CompNode2, true, EBTNodeResult::Succeeded, TEXT("Bool3"));
+				}
+
+				UBTCompositeNode& CompNode3 = FBTBuilder::AddSequence(CompNode);
+				{
+					FBTBuilder::WithDecoratorBlackboard(CompNode, EBasicKeyOperation::Set, EBTFlowAbortMode::LowerPriority, TEXT("Bool3"), 5/*LogIndexBecomeRelevant*/, 6/*LogIndexCeaseRelevant*/);
+					FBTBuilder::AddTask(CompNode3, 12, EBTNodeResult::Succeeded, 10/*ExecutionTicks*/);
+					FBTBuilder::WithTaskServiceLog(CompNode3, 13/*ActivationIndex*/, 14/*DeactivationIndex*/, 15/*TickIndex*/, TEXT("Bool1"));
+				}
+			}
+		}
+
+		UBTCompositeNode& CompNode = FBTBuilder::AddSelector(*BTAsset);
+		{
+			FBTBuilder::AddTaskFlagChange(CompNode, true, EBTNodeResult::Failed, TEXT("Bool3"));
+			FBTBuilder::AddTaskSubtree(CompNode, ChildAsset);
+			{
+				FBTBuilder::WithDecoratorBlackboard(CompNode, EBasicKeyOperation::NotSet, EBTFlowAbortMode::Both, TEXT("Bool1"), 1/*LogIndexBecomeRelevant*/, 2/*LogIndexCeaseRelevant*/);
+			}
+
+			UBTCompositeNode& CompNode2 = FBTBuilder::AddSequence(CompNode);
+			{
+				FBTBuilder::AddTaskFlagChange(CompNode2, false, EBTNodeResult::Succeeded, TEXT("Bool3"));
+				FBTBuilder::AddTaskFlagChange(CompNode2, true, EBTNodeResult::Succeeded, TEXT("Bool2"));
+				FBTBuilder::AddTaskFlagChange(CompNode2, false, EBTNodeResult::Succeeded, TEXT("Bool1"));
+			}
+		}
+
+		ExpectedResult.Add(1);
+		ExpectedResult.Add(3);
+		ExpectedResult.Add(13);
+		ExpectedResult.Add(15);
+		ExpectedResult.Add(12);
+		ExpectedResult.Add(15); // Extra service tick because aux nodes are always ticked before any pending requests are processed
+		ExpectedResult.Add(14);
+		ExpectedResult.Add(4);
+		ExpectedResult.Add(3);
+		ExpectedResult.Add(4);
+		ExpectedResult.Add(2);
+	}
+};
+IMPLEMENT_AI_LATENT_TEST(FAITest_BTSubtreeAbortOutToLowerPrio, "System.AI.Behavior Trees.Subtree: abort out lower prio")
+
 struct FAITest_BTServiceInstantTask : public FAITest_SimpleBT
 {
 	FAITest_BTServiceInstantTask()
diff --git a/Engine/Source/Runtime/AIModule/Classes/BehaviorTree/BTCompositeNode.h b/Engine/Source/Runtime/AIModule/Classes/BehaviorTree/BTCompositeNode.h
--- a/Engine/Source/Runtime/AIModule/Classes/BehaviorTree/BTCompositeNode.h
+++ b/Engine/Source/Runtime/AIModule/Classes/BehaviorTree/BTCompositeNode.h
@@ -119,8 +119,8 @@ class AIMODULE_API UBTCompositeNode : public UBTNode
 	void OnChildActivation(FBehaviorTreeSearchData& SearchData, int32 ChildIndex) const;
 
 	/** called after child has finished search */
-	void OnChildDeactivation(FBehaviorTreeSearchData& SearchData, const UBTNode& ChildNode, EBTNodeResult::Type& NodeResult) const;
-	void OnChildDeactivation(FBehaviorTreeSearchData& SearchData, int32 ChildIndex, EBTNodeResult::Type& NodeResult) const;
+	void OnChildDeactivation(FBehaviorTreeSearchData& SearchData, const UBTNode& ChildNode, EBTNodeResult::Type& NodeResult, const bool bIsRequestFromActiveInstance) const;
+	void OnChildDeactivation(FBehaviorTreeSearchData& SearchData, int32 ChildIndex, EBTNodeResult::Type& NodeResult, const bool bIsRequestFromActiveInstance) const;
 
 	/** called when start enters this node */
 	void OnNodeActivation(FBehaviorTreeSearchData& SearchData) const;
@@ -175,6 +175,12 @@ class AIMODULE_API UBTCompositeNode : public UBTNode
 
 	bool IsApplyingDecoratorScope() const;
 
+	// Deprecated methods
+	UE_DEPRECATED(5.0, "This function is deprecated. Please use RequestBranchDeactivation instead.")
+	void OnChildDeactivation(FBehaviorTreeSearchData& SearchData, const UBTNode& ChildNode, EBTNodeResult::Type& NodeResult) const { OnChildDeactivation(SearchData, ChildNode, NodeResult, true/*bIsRequestFromActiveInstance*/); }
+	UE_DEPRECATED(5.0, "This function is deprecated. Please use RequestBranchDeactivation instead.")
+	void OnChildDeactivation(FBehaviorTreeSearchData& SearchData, int32 ChildIndex, EBTNodeResult::Type& NodeResult) const { OnChildDeactivation(SearchData, ChildIndex, NodeResult, true /*bIsRequestFromActiveInstance*/); }
+
 protected:
 
 	/** if set, all decorators in branch below will be removed when execution flow leaves (decorators on this node are not affected) */
@@ -224,7 +230,7 @@ protected:
 	void NotifyDecoratorsOnActivation(FBehaviorTreeSearchData& SearchData, int32 ChildIdx) const;
 
 	/** runs through decorators on given child node and notify them about deactivation */
-	void NotifyDecoratorsOnDeactivation(FBehaviorTreeSearchData& SearchData, int32 ChildIdx, EBTNodeResult::Type& NodeResult) const;
+	void NotifyDecoratorsOnDeactivation(FBehaviorTreeSearchData& SearchData, int32 ChildIdx, EBTNodeResult::Type& NodeResult, const bool bIsInSameActiveInstance) const;
 
 	/** runs through decorators on given child node and notify them about failed activation */
 	void NotifyDecoratorsOnFailedActivation(FBehaviorTreeSearchData& SearchData, int32 ChildIdx, EBTNodeResult::Type& NodeResult) const;
diff --git a/Engine/Source/Runtime/AIModule/Private/BehaviorTree/BTCompositeNode.cpp b/Engine/Source/Runtime/AIModule/Private/BehaviorTree/BTCompositeNode.cpp
--- a/Engine/Source/Runtime/AIModule/Private/BehaviorTree/BTCompositeNode.cpp
+++ b/Engine/Source/Runtime/AIModule/Private/BehaviorTree/BTCompositeNode.cpp
@@ -121,12 +121,12 @@ void UBTCompositeNode::OnChildActivation(FBehaviorTreeSearchData& SearchData, in
 	NodeMemory->CurrentChild = ChildIndex;
 }
 
-void UBTCompositeNode::OnChildDeactivation(FBehaviorTreeSearchData& SearchData, const UBTNode& ChildNode, EBTNodeResult::Type& NodeResult) const
+void UBTCompositeNode::OnChildDeactivation(FBehaviorTreeSearchData& SearchData, const UBTNode& ChildNode, EBTNodeResult::Type& NodeResult, const bool bIsRequestFromActiveInstance) const
 {
-	OnChildDeactivation(SearchData, GetChildIndex(SearchData, ChildNode), NodeResult);
+	OnChildDeactivation(SearchData, GetChildIndex(SearchData, ChildNode), NodeResult, bIsRequestFromActiveInstance);
 }
 
-void UBTCompositeNode::OnChildDeactivation(FBehaviorTreeSearchData& SearchData, int32 ChildIndex, EBTNodeResult::Type& NodeResult) const
+void UBTCompositeNode::OnChildDeactivation(FBehaviorTreeSearchData& SearchData, int32 ChildIndex, EBTNodeResult::Type& NodeResult, const bool bIsRequestFromActiveInstance) const
 {
 	const FBTCompositeChild& ChildInfo = Children[ChildIndex];
 
@@ -149,7 +149,7 @@ void UBTCompositeNode::OnChildDeactivation(FBehaviorTreeSearchData& SearchData,
 	const bool bCanNotify = !bUseDecoratorsDeactivationCheck || CanNotifyDecoratorsOnDeactivation(SearchData, ChildIndex, NodeResult);
 	if (bCanNotify)
 	{
-		NotifyDecoratorsOnDeactivation(SearchData, ChildIndex, NodeResult);
+		NotifyDecoratorsOnDeactivation(SearchData, ChildIndex, NodeResult, bIsRequestFromActiveInstance);
 	}
 }
 
@@ -247,7 +247,7 @@ void UBTCompositeNode::NotifyDecoratorsOnActivation(FBehaviorTreeSearchData& Sea
 	}
 }
 
-void UBTCompositeNode::NotifyDecoratorsOnDeactivation(FBehaviorTreeSearchData& SearchData, int32 ChildIdx, EBTNodeResult::Type& NodeResult) const
+void UBTCompositeNode::NotifyDecoratorsOnDeactivation(FBehaviorTreeSearchData& SearchData, int32 ChildIdx, EBTNodeResult::Type& NodeResult, const bool bIsRequestFromActiveInstance) const
 {
 	const FBTCompositeChild& ChildInfo = Children[ChildIdx];
 	if (NodeResult == EBTNodeResult::Aborted)
@@ -270,14 +270,21 @@ void UBTCompositeNode::NotifyDecoratorsOnDeactivation(FBehaviorTreeSearchData& S
 			DecoratorOb->WrappedOnNodeProcessed(SearchData, NodeResult);
 			DecoratorOb->WrappedOnNodeDeactivation(SearchData, NodeResult);
 
+			if (!bIsRequestFromActiveInstance)
+			{
+				UE_VLOG(SearchData.OwnerComp.GetOwner(), LogBehaviorTree, Verbose, TEXT("Removing decorator(%s) as request is not in active instance"), *UBehaviorTreeTypes::DescribeNodeHelper(this));
+				SearchData.AddUniqueUpdate(FBehaviorTreeSearchUpdate(DecoratorOb, SearchData.OwnerComp.GetActiveInstanceIdx(), EBTNodeUpdateMode::Remove));
+			}
 			// leaving child branch: 
-			if (DecoratorOb->GetFlowAbortMode() == EBTFlowAbortMode::Self)
+			else if (DecoratorOb->GetFlowAbortMode() == EBTFlowAbortMode::Self)
 			{
+				UE_VLOG(SearchData.OwnerComp.GetOwner(), LogBehaviorTree, Verbose, TEXT("Removing out of scope decorator(%s) because of self abort type"), *UBehaviorTreeTypes::DescribeNodeHelper(this));
 				// - observers with mode "Self" are now out of scope, remove them
 				SearchData.AddUniqueUpdate(FBehaviorTreeSearchUpdate(DecoratorOb, SearchData.OwnerComp.GetActiveInstanceIdx(), EBTNodeUpdateMode::Remove));
 			}
 			else if (DecoratorOb->GetFlowAbortMode() == EBTFlowAbortMode::LowerPriority)
 			{
+				UE_VLOG(SearchData.OwnerComp.GetOwner(), LogBehaviorTree, Verbose, TEXT("Reactivating decorator(%s) because of lower priority abort type"), *UBehaviorTreeTypes::DescribeNodeHelper(this));
 				// - observers with mode "Lower Priority" will try to reactivate themselves ("Both" is not removed on node activation)
 				SearchData.AddUniqueUpdate(FBehaviorTreeSearchUpdate(DecoratorOb, SearchData.OwnerComp.GetActiveInstanceIdx(), EBTNodeUpdateMode::Add));
 			}
diff --git a/Engine/Source/Runtime/AIModule/Private/BehaviorTree/BehaviorTreeComponent.cpp b/Engine/Source/Runtime/AIModule/Private/BehaviorTree/BehaviorTreeComponent.cpp
--- a/Engine/Source/Runtime/AIModule/Private/BehaviorTree/BehaviorTreeComponent.cpp
+++ b/Engine/Source/Runtime/AIModule/Private/BehaviorTree/BehaviorTreeComponent.cpp
@@ -284,6 +284,8 @@ void UBehaviorTreeComponent::StopTree(EBTStopMode::Type StopMode)
 {
 	SCOPE_CYCLE_COUNTER(STAT_AI_BehaviorTree_StopTree);
 
+	UE_VLOG(GetOwner(), LogBehaviorTree, Verbose, TEXT("StopTree %s, mode:%s"), *GetNameSafe(GetRootTree()), StopMode == EBTStopMode::Forced ? TEXT("Forced") : TEXT("Safe"));
+
 	if (StopTreeLock)
 	{
 		bDeferredStopTree = true;
@@ -1692,7 +1694,7 @@ void UBehaviorTreeComponent::ProcessExecutionRequest()
 
 				if (TestNode)
 				{
-					TestNode->OnChildDeactivation(SearchData, *ChildNode, NodeResult);
+					TestNode->OnChildDeactivation(SearchData, *ChildNode, NodeResult, ActiveInstanceIdx == ExecutionRequest.ExecuteInstanceIdx  /*bIsRequestInSameInstance*/);
 				}
 			}
 			else if (TestNode->Children.IsValidIndex(ChildBranchIdx))
@@ -1886,7 +1888,7 @@ bool UBehaviorTreeComponent::DeactivateUpTo(UBTCompositeNode* Node, uint16 NodeI
 		if (NotifyParent)
 		{
 			OutLastDeactivatedChildIndex = NotifyParent->GetChildIndex(SearchData, *DeactivatedChild);
-			NotifyParent->OnChildDeactivation(SearchData, OutLastDeactivatedChildIndex, NodeResult);
+			NotifyParent->OnChildDeactivation(SearchData, OutLastDeactivatedChildIndex, NodeResult, ActiveInstanceIdx == NodeInstanceIdx /*bIsRequestInSameInstance*/);
 
 			BT_SEARCHLOG(SearchData, Verbose, TEXT("Deactivate node: %s"), *UBehaviorTreeTypes::DescribeNodeHelper(DeactivatedChild));
 			StoreDebuggerSearchStep(DeactivatedChild, ActiveInstanceIdx, NodeResult);
diff --git a/Engine/Source/Runtime/AIModule/Private/BehaviorTree/BehaviorTreeTypes.cpp b/Engine/Source/Runtime/AIModule/Private/BehaviorTree/BehaviorTreeTypes.cpp
--- a/Engine/Source/Runtime/AIModule/Private/BehaviorTree/BehaviorTreeTypes.cpp
+++ b/Engine/Source/Runtime/AIModule/Private/BehaviorTree/BehaviorTreeTypes.cpp
@@ -390,12 +390,14 @@ void FBehaviorTreeSearchData::AddUniqueUpdate(const FBehaviorTreeSearchUpdate& U
 			// duplicate, skip
 			if (Info.Mode == UpdateInfo.Mode)
 			{
+				UE_VLOG(OwnerComp.GetOwner(), LogBehaviorTree, Verbose, TEXT(">> skipped: duplicated operation"));
 				bSkipAdding = true;
 				break;
 			}
 
 			// don't add pairs add-remove
 			bSkipAdding = (Info.Mode == EBTNodeUpdateMode::Remove) || (UpdateInfo.Mode == EBTNodeUpdateMode::Remove);
+			UE_CVLOG(bSkipAdding, OwnerComp.GetOwner(), LogBehaviorTree, Verbose, TEXT(">> skipped: paired add/remove"));
 
 			PendingUpdates.RemoveAt(UpdateIndex, 1, false);
 		}
@@ -407,6 +409,7 @@ void FBehaviorTreeSearchData::AddUniqueUpdate(const FBehaviorTreeSearchUpdate& U
 	{
 		const bool bIsActive = OwnerComp.IsAuxNodeActive(UpdateInfo.AuxNode, UpdateInfo.InstanceIndex);
 		bSkipAdding = !bIsActive;
+		UE_CVLOG(bSkipAdding, OwnerComp.GetOwner(), LogBehaviorTree, Verbose, TEXT(">> skipped: inactive aux nodes"));
 	}
 
 	if (!bSkipAdding)
@@ -414,10 +417,6 @@ void FBehaviorTreeSearchData::AddUniqueUpdate(const FBehaviorTreeSearchUpdate& U
 		const int32 Idx = PendingUpdates.Add(UpdateInfo);
 		PendingUpdates[Idx].bPostUpdate = (UpdateInfo.Mode == EBTNodeUpdateMode::Add) && (Cast<UBTService>(UpdateInfo.AuxNode) != NULL);
 	}
-	else
-	{
-		UE_VLOG(OwnerComp.GetOwner(), LogBehaviorTree, Verbose, TEXT(">> or not, update skipped"));
-	}
 }
 
 void FBehaviorTreeSearchData::AssignSearchId()
diff --git a/Engine/Source/Runtime/AIModule/Private/BehaviorTree/Composites/BTComposite_SimpleParallel.cpp b/Engine/Source/Runtime/AIModule/Private/BehaviorTree/Composites/BTComposite_SimpleParallel.cpp
--- a/Engine/Source/Runtime/AIModule/Private/BehaviorTree/Composites/BTComposite_SimpleParallel.cpp
+++ b/Engine/Source/Runtime/AIModule/Private/BehaviorTree/Composites/BTComposite_SimpleParallel.cpp
@@ -77,7 +77,7 @@ void UBTComposite_SimpleParallel::NotifyChildExecution(UBehaviorTreeComponent& O
 			
 			// notify decorators on main task, ignore observers updates in FakeSearchData - they are not allowed by parallel composite
 			FBehaviorTreeSearchData FakeSearchData(OwnerComp);
-			NotifyDecoratorsOnDeactivation(FakeSearchData, ChildIdx, NodeResult);
+			NotifyDecoratorsOnDeactivation(FakeSearchData, ChildIdx, NodeResult, true /*bIsRequestInSameInstance*/);
 
 			const int32 MyInstanceIdx = OwnerComp.FindInstanceContainingNode(this);
 
