[BehaviorTree] Fix for wrong child selection when an execution request is created by an auxiliary node the same frame a latent abort finishes.

diff --git a/Engine/Source/Developer/AITestSuite/Classes/BehaviorTree/TestBTTask_LatentWithFlags.h b/Engine/Source/Developer/AITestSuite/Classes/BehaviorTree/TestBTTask_LatentWithFlags.h
--- a/Engine/Source/Developer/AITestSuite/Classes/BehaviorTree/TestBTTask_LatentWithFlags.h
+++ b/Engine/Source/Developer/AITestSuite/Classes/BehaviorTree/TestBTTask_LatentWithFlags.h
@@ -17,40 +17,52 @@ struct FBTLatentTaskMemory
 UCLASS(meta = (HiddenNode))
 class UTestBTTask_LatentWithFlags : public UBTTaskNode
 {
-	GENERATED_UCLASS_BODY()
+	GENERATED_BODY()
+
+public:
+
+	UTestBTTask_LatentWithFlags();
 
 	UPROPERTY()
-	int32 LogIndexExecuteStart;
+	int32 LogIndexExecuteStart = 0;
 
 	UPROPERTY()
-	int32 LogIndexExecuteFinish;
+	int32 LogIndexExecuting = -1;
 
 	UPROPERTY()
-	int32 LogIndexAbortStart;
+	int32 LogIndexExecuteFinish = 0;
 
 	UPROPERTY()
-	int32 LogIndexAbortFinish;
+	int32 LogIndexAbortStart = 0;
 
 	UPROPERTY()
-	int32 ExecuteTicks;
+	int32 LogIndexAborting = -1;
 
 	UPROPERTY()
-	int32 AbortTicks;
+	int32 LogIndexAbortFinish = 0;
 
+	/** Num of ticks before 'execute start' and `set execute flag` and then the same num of ticks before `execute finish` */
 	UPROPERTY()
-	FName KeyNameExecute;
+	int32 ExecuteHalfTicks = 2;
 
+	/** Num of ticks before 'abort start' and `set abort flag` and then the same num of ticks before `abort finish` */
 	UPROPERTY()
-	FName KeyNameAbort;
+	int32 AbortHalfTicks = 2;
 
 	UPROPERTY()
-	TEnumAsByte<EBTNodeResult::Type> LogResult;
+	FName KeyNameExecute = TEXT("Bool1");
 
+	UPROPERTY()
+	FName KeyNameAbort = TEXT("Bool2");
+
+	UPROPERTY()
+	TEnumAsByte<EBTNodeResult::Type> LogResult = EBTNodeResult::Succeeded;
+
+protected:
 	virtual EBTNodeResult::Type ExecuteTask(class UBehaviorTreeComponent& OwnerComp, uint8* NodeMemory) override;
 	virtual EBTNodeResult::Type AbortTask(UBehaviorTreeComponent& OwnerComp, uint8* NodeMemory) override;
 	virtual uint16 GetInstanceMemorySize() const override;
 	void LogExecution(class UBehaviorTreeComponent& OwnerComp, int32 LogNumber);
 
-protected:
 	virtual void TickTask(class UBehaviorTreeComponent& OwnerComp, uint8* NodeMemory, float DeltaSeconds) override;
-};
+};
\ No newline at end of file
diff --git a/Engine/Source/Developer/AITestSuite/Private/BehaviorTree/TestBTTask_LatentWithFlags.cpp b/Engine/Source/Developer/AITestSuite/Private/BehaviorTree/TestBTTask_LatentWithFlags.cpp
--- a/Engine/Source/Developer/AITestSuite/Private/BehaviorTree/TestBTTask_LatentWithFlags.cpp
+++ b/Engine/Source/Developer/AITestSuite/Private/BehaviorTree/TestBTTask_LatentWithFlags.cpp
@@ -4,32 +4,21 @@
 #include "BehaviorTree/BlackboardComponent.h"
 #include "MockAI_BT.h"
 
-UTestBTTask_LatentWithFlags::UTestBTTask_LatentWithFlags(const FObjectInitializer& ObjectInitializer) : Super(ObjectInitializer)
+UTestBTTask_LatentWithFlags::UTestBTTask_LatentWithFlags()
 {
-	NodeName = "LatentTest";
-	LogIndexExecuteStart = 0;
-	LogIndexExecuteFinish = 0;
-	LogIndexAbortStart = 0;
-	LogIndexAbortFinish = 0;
-	ExecuteTicks = 2;
-	AbortTicks = 2;
-	KeyNameExecute = TEXT("Bool1");
-	KeyNameAbort = TEXT("Bool2");
-	LogResult = EBTNodeResult::Succeeded;
-
 	bNotifyTick = true;
 }
 
 EBTNodeResult::Type UTestBTTask_LatentWithFlags::ExecuteTask(UBehaviorTreeComponent& OwnerComp, uint8* NodeMemory)
 {
 	FBTLatentTaskMemory* MyMemory = (FBTLatentTaskMemory*)NodeMemory;
-	MyMemory->FlagFrameIdx = ExecuteTicks + FAITestHelpers::FramesCounter();
-	MyMemory->EndFrameIdx = MyMemory->FlagFrameIdx + ExecuteTicks;
+	MyMemory->FlagFrameIdx = ExecuteHalfTicks + FAITestHelpers::FramesCounter();
+	MyMemory->EndFrameIdx = MyMemory->FlagFrameIdx + ExecuteHalfTicks;
 	MyMemory->bFlagSet = false;
 	MyMemory->bIsAborting = false;
 
 	LogExecution(OwnerComp, LogIndexExecuteStart);
-	if (ExecuteTicks == 0)
+	if (ExecuteHalfTicks == 0)
 	{
 		OwnerComp.GetBlackboardComponent()->SetValueAsBool(KeyNameExecute, true);
 		MyMemory->bFlagSet = true;
@@ -44,13 +33,13 @@ EBTNodeResult::Type UTestBTTask_LatentWithFlags::ExecuteTask(UBehaviorTreeCompon
 EBTNodeResult::Type UTestBTTask_LatentWithFlags::AbortTask(UBehaviorTreeComponent& OwnerComp, uint8* NodeMemory)
 {
 	FBTLatentTaskMemory* MyMemory = (FBTLatentTaskMemory*)NodeMemory;
-	MyMemory->FlagFrameIdx = AbortTicks + FAITestHelpers::FramesCounter();
-	MyMemory->EndFrameIdx = MyMemory->FlagFrameIdx + AbortTicks;
+	MyMemory->FlagFrameIdx = AbortHalfTicks + FAITestHelpers::FramesCounter();
+	MyMemory->EndFrameIdx = MyMemory->FlagFrameIdx + AbortHalfTicks;
 	MyMemory->bFlagSet = false;
 	MyMemory->bIsAborting = true;
 
 	LogExecution(OwnerComp, LogIndexAbortStart);
-	if (AbortTicks == 0)
+	if (AbortHalfTicks == 0)
 	{
 		OwnerComp.GetBlackboardComponent()->SetValueAsBool(KeyNameAbort, true);
 		MyMemory->bFlagSet = true;
@@ -66,6 +55,8 @@ void UTestBTTask_LatentWithFlags::TickTask(UBehaviorTreeComponent& OwnerComp, ui
 {
 	FBTLatentTaskMemory* MyMemory = (FBTLatentTaskMemory*)NodeMemory;
 
+	LogExecution(OwnerComp, MyMemory->bIsAborting ? LogIndexAborting : LogIndexExecuting);
+
 	if (!MyMemory->bFlagSet && FAITestHelpers::FramesCounter() >= MyMemory->FlagFrameIdx)
 	{
 		MyMemory->bFlagSet = true;
@@ -86,7 +77,7 @@ void UTestBTTask_LatentWithFlags::TickTask(UBehaviorTreeComponent& OwnerComp, ui
 			LogExecution(OwnerComp, LogIndexExecuteFinish);
 			FinishLatentTask(OwnerComp, LogResult);
 		}
-	}	
+	}
 }
 
 uint16 UTestBTTask_LatentWithFlags::GetInstanceMemorySize() const
@@ -100,4 +91,4 @@ void UTestBTTask_LatentWithFlags::LogExecution(UBehaviorTreeComponent& OwnerComp
 	{
 		UMockAI_BT::ExecutionLog.Add(LogNumber);
 	}
-}
+}
\ No newline at end of file
diff --git a/Engine/Source/Developer/AITestSuite/Private/Tests/BTTest.cpp b/Engine/Source/Developer/AITestSuite/Private/Tests/BTTest.cpp
--- a/Engine/Source/Developer/AITestSuite/Private/Tests/BTTest.cpp
+++ b/Engine/Source/Developer/AITestSuite/Private/Tests/BTTest.cpp
@@ -661,6 +661,33 @@ struct FAITest_BTSwitchingHigherPrioDuringServiceTick : public FAITest_SimpleBT
 };
 IMPLEMENT_AI_LATENT_TEST(FAITest_BTSwitchingHigherPrioDuringServiceTick, "System.AI.Behavior Trees.Switch: higher priority during service tick")
 
+struct FAITest_BTRequestExecutionOnLatentAbortFinished : public FAITest_SimpleBT
+{
+	FAITest_BTRequestExecutionOnLatentAbortFinished()
+	{
+		UBTCompositeNode& CompNode = FBTBuilder::AddSequence(*BTAsset); // 0
+		{
+			FBTBuilder::AddTaskLatentFlags(CompNode, EBTNodeResult::Succeeded, // 3
+				1 /* ExecHalfNumTicks */, FName() /* ExecKeyName */, 100 /* ExecLogStart */, 199 /* ExecLogFinish */,
+				1 /* AbortHalfNumTicks */, FName() /* AbortKeyName */, 200 /* AbortLogStart */, 299 /* AbortLogFinish */);
+			{
+				FBTBuilder::WithDecorator<UBTDecorator_ForceSuccess>(CompNode); // 1
+				FBTBuilder::WithDecoratorDelayedAbort(CompNode, 1, false /* bOnlyOnce */); // 2
+			}
+
+			FBTBuilder::AddTask(CompNode, 4, EBTNodeResult::Succeeded); // 4
+			FBTBuilder::AddTask(CompNode, 5, EBTNodeResult::Succeeded); // 5
+		}
+
+		ExpectedResult.Add(100);
+		ExpectedResult.Add(200);
+		ExpectedResult.Add(299);
+		ExpectedResult.Add(4);
+		ExpectedResult.Add(5);
+	}
+};
+IMPLEMENT_AI_LATENT_TEST(FAITest_BTRequestExecutionOnLatentAbortFinished, "System.AI.Behavior Trees.Abort: request on latent task finished")
+
 struct FAITest_BTSwitchingHigherPrioDuringServiceBecomeRelevant : public FAITest_SimpleBT
 {
 	FAITest_BTSwitchingHigherPrioDuringServiceBecomeRelevant()
diff --git a/Engine/Source/Developer/AITestSuite/Public/BTBuilder.h b/Engine/Source/Developer/AITestSuite/Public/BTBuilder.h
--- a/Engine/Source/Developer/AITestSuite/Public/BTBuilder.h
+++ b/Engine/Source/Developer/AITestSuite/Public/BTBuilder.h
@@ -188,15 +188,17 @@ struct FBTBuilder
 	}
 
 	static void AddTaskLatentFlags(UBTCompositeNode& ParentNode, EBTNodeResult::Type NodeResult,
-		int32 ExecuteHalfTicks, FName ExecuteKeyName, int32 ExecuteLogStart, int32 ExecuteLogFinish,
-		int32 AbortHalfTicks = 0, FName AbortKeyName = NAME_None, int32 AbortLogStart = 0, int32 AbortLogFinish = 0)
+		int32 ExecuteHalfTicks, /** Num of ticks before 'execute start' and `set execute flag` and then the same num of ticks before `execute finish` */
+		FName ExecuteKeyName, int32 ExecuteLogStart, int32 ExecuteLogFinish,
+		int32 AbortHalfTicks = 0, /** Num of ticks before 'abort start' and `set abort flag` and then the same num of ticks before `abort finish` */
+		FName AbortKeyName = NAME_None, int32 AbortLogStart = 0, int32 AbortLogFinish = 0)
 	{
 		UTestBTTask_LatentWithFlags* TaskNode = NewObject<UTestBTTask_LatentWithFlags>(ParentNode.GetTreeAsset());
-		TaskNode->ExecuteTicks = ExecuteHalfTicks;
+		TaskNode->ExecuteHalfTicks = ExecuteHalfTicks;
 		TaskNode->KeyNameExecute = ExecuteKeyName;
 		TaskNode->LogIndexExecuteStart = ExecuteLogStart;
 		TaskNode->LogIndexExecuteFinish = ExecuteLogFinish;
-		TaskNode->AbortTicks = AbortHalfTicks;
+		TaskNode->AbortHalfTicks = AbortHalfTicks;
 		TaskNode->KeyNameAbort = AbortKeyName;
 		TaskNode->LogIndexAbortStart = AbortLogStart;
 		TaskNode->LogIndexAbortFinish = AbortLogFinish;
diff --git a/Engine/Source/Runtime/AIModule/Classes/BehaviorTree/BehaviorTreeComponent.h b/Engine/Source/Runtime/AIModule/Classes/BehaviorTree/BehaviorTreeComponent.h
--- a/Engine/Source/Runtime/AIModule/Classes/BehaviorTree/BehaviorTreeComponent.h
+++ b/Engine/Source/Runtime/AIModule/Classes/BehaviorTree/BehaviorTreeComponent.h
@@ -313,8 +313,8 @@ protected:
 	/** loops tree execution */
 	uint8 bLoopExecution : 1;
 
-	/** set when execution is waiting for tasks to abort (current or parallel's main) */
-	uint8 bWaitingForAbortingTasks : 1;
+	/** set when execution is waiting for tasks to finish their latent abort (current or parallel's main) */
+	uint8 bWaitingForLatentAborts : 1;
 
 	/** set when execution update is scheduled for next tick */
 	uint8 bRequestedFlowUpdate : 1;
@@ -371,8 +371,14 @@ protected:
 	/** deactivate all nodes up to requested one */
 	bool DeactivateUpTo(UBTCompositeNode* Node, uint16 NodeInstanceIdx, EBTNodeResult::Type& NodeResult, int32& OutLastDeactivatedChildIndex);
 
-	/** update state of aborting tasks */
-	void UpdateAbortingTasks();
+	/** returns true if execution was waiting on latent aborts and they are all finished;  */
+	bool TrackPendingLatentAborts();
+
+	/** tracks if there are new tasks using latent abort in progress */
+	void TrackNewLatentAborts();
+
+	/** return true if the current or any parallel task has a latent abort in progress */
+	bool HasActiveLatentAborts() const;
 
 	/** apply pending execution from last task search */
 	void ProcessPendingExecution();
@@ -494,5 +500,5 @@ FORCEINLINE bool UBehaviorTreeComponent::IsRestartPending() const
 
 FORCEINLINE bool UBehaviorTreeComponent::IsAbortPending() const
 {
-	return bWaitingForAbortingTasks || PendingExecution.IsSet();
+	return bWaitingForLatentAborts || PendingExecution.IsSet();
 }
diff --git a/Engine/Source/Runtime/AIModule/Private/BehaviorTree/BehaviorTreeComponent.cpp b/Engine/Source/Runtime/AIModule/Private/BehaviorTree/BehaviorTreeComponent.cpp
--- a/Engine/Source/Runtime/AIModule/Private/BehaviorTree/BehaviorTreeComponent.cpp
+++ b/Engine/Source/Runtime/AIModule/Private/BehaviorTree/BehaviorTreeComponent.cpp
@@ -65,7 +65,7 @@ UBehaviorTreeComponent::UBehaviorTreeComponent(const FObjectInitializer& ObjectI
 	StopTreeLock = 0;
 	bDeferredStopTree = false;
 	bLoopExecution = false;
-	bWaitingForAbortingTasks = false;
+	bWaitingForLatentAborts = false;
 	bRequestedFlowUpdate = false;
 	bAutoActivate = true;
 	bWantsInitializeComponent = true; 
@@ -255,7 +255,7 @@ void UBehaviorTreeComponent::StartTree(UBehaviorTree& Asset, EBTExecutionMode::T
 void UBehaviorTreeComponent::ProcessPendingInitialize()
 {
 	StopTree(EBTStopMode::Safe);
-	if (bWaitingForAbortingTasks)
+	if (bWaitingForLatentAborts)
 	{
 		return;
 	}
@@ -349,7 +349,7 @@ void UBehaviorTreeComponent::StopTree(EBTStopMode::Type StopMode)
 						if (bIsValidForStatus)
 						{
 							InstanceInfo.MarkParallelTaskAsAbortingAt(ParallelIndex);
-							bWaitingForAbortingTasks = true;
+							bWaitingForLatentAborts = true;
 						}
 						else
 						{
@@ -386,7 +386,7 @@ void UBehaviorTreeComponent::StopTree(EBTStopMode::Type StopMode)
 		}
 	}
 
-	if (bWaitingForAbortingTasks)
+	if (bWaitingForLatentAborts)
 	{
 		if (StopMode == EBTStopMode::Safe)
 		{
@@ -423,7 +423,7 @@ void UBehaviorTreeComponent::StopTree(EBTStopMode::Type StopMode)
 	bRequestedFlowUpdate = false;
 	bRequestedStop = false;
 	bIsRunning = false;
-	bWaitingForAbortingTasks = false;
+	bWaitingForLatentAborts = false;
 	bDeferredStopTree = false;
 }
 
@@ -492,9 +492,8 @@ void UBehaviorTreeComponent::OnTaskFinished(const UBTTaskNode* TaskNode, EBTNode
 
 	uint8* ParentMemory = ParentNode->GetNodeMemory<uint8>(InstanceStack[TaskInstanceIdx]);
 
-	const bool bWasWaitingForAbort = bWaitingForAbortingTasks;
 	ParentNode->ConditionalNotifyChildExecution(*this, ParentMemory, *TaskNode, TaskResult);
-	
+
 	if (TaskResult != EBTNodeResult::InProgress)
 	{
 		StoreDebuggerSearchStep(TaskNode, TaskInstanceIdx, TaskResult);
@@ -526,42 +525,10 @@ void UBehaviorTreeComponent::OnTaskFinished(const UBTTaskNode* TaskNode, EBTNode
 
 			InstanceStack[TaskInstanceIdx].ActiveNodeType = EBTActiveNode::InactiveTask;
 		}
-
-		// update state of aborting tasks after currently finished one was set to Inactive
-		UpdateAbortingTasks();
-
-		// make sure that we continue execution after all pending latent aborts finished
-		if (!bWaitingForAbortingTasks && bWasWaitingForAbort)
-		{
-			if (bRequestedStop)
-			{
-				StopTree(EBTStopMode::Safe);
-			}
-			else
-			{
-				// force new search if there were any execution requests while waiting for aborting task
-				if (ExecutionRequest.ExecuteNode)
-				{
-					UE_VLOG(GetOwner(), LogBehaviorTree, Log, TEXT("> found valid ExecutionRequest, locking PendingExecution data to force new search!"));
-					PendingExecution.Lock();
-
-					if (ExecutionRequest.SearchEnd.IsSet())
-					{
-						UE_VLOG(GetOwner(), LogBehaviorTree, Log, TEXT("> removing limit from end of search range! [abort done]"));
-						ExecutionRequest.SearchEnd = FBTNodeIndex();
-					}
-				}
-
-				ScheduleExecutionUpdate();
-			}
-		}
-	}
-	else
-	{
-		// always update state of aborting tasks
-		UpdateAbortingTasks();
 	}
 
+	TrackNewLatentAborts();
+
 	if (TreeStartInfo.HasPendingInitialize())
 	{
 		ProcessPendingInitialize();
@@ -1066,7 +1033,7 @@ void UBehaviorTreeComponent::RequestExecution(UBTCompositeNode* RequestedOn, int
 
     // Not only checking against deactivated branch upon applying search data or while aborting task, 
     // but also while waiting after a latent task to abort
-	if (SearchData.bFilterOutRequestFromDeactivatedBranch || bWaitingForAbortingTasks)
+	if (SearchData.bFilterOutRequestFromDeactivatedBranch || bWaitingForLatentAborts)
 	{
 		// request on same node or with higher priority doesn't require additional checks
 		if (SearchData.SearchRootNode != ExecutionIdx && SearchData.SearchRootNode.TakesPriorityOver(ExecutionIdx) && SearchData.DeactivatedBranchStart.IsSet())
@@ -1217,8 +1184,8 @@ void UBehaviorTreeComponent::RequestExecution(UBTCompositeNode* RequestedOn, int
 	// - don't search, just accumulate requests and run them when abort is done
 	// - rollback changes from search that caused abort to ensure proper state of tree
 	const bool bIsActiveNodeAborting = InstanceStack.Num() && InstanceStack.Last().ActiveNodeType == EBTActiveNode::AbortingTask;
-	const bool bInvalidateCurrentSearch = bWaitingForAbortingTasks || bIsActiveNodeAborting;
-	const bool bScheduleNewSearch = !bWaitingForAbortingTasks;
+	const bool bInvalidateCurrentSearch = bWaitingForLatentAborts || bIsActiveNodeAborting;
+	const bool bScheduleNewSearch = !bWaitingForLatentAborts;
 
 	if (bInvalidateCurrentSearch)
 	{
@@ -1329,7 +1296,7 @@ void UBehaviorTreeComponent::ApplySearchUpdates(const TArray<FBehaviorTreeSearch
 					if (NodeResult == EBTNodeResult::InProgress)
 					{
 						UpdateInstance.MarkParallelTaskAsAbortingAt(ParallelTaskIdx);
-						bWaitingForAbortingTasks = true;
+						bWaitingForLatentAborts = true;
 					}
 
 					OnTaskFinished(UpdateInfo.TaskNode, NodeResult);
@@ -1476,6 +1443,33 @@ void UBehaviorTreeComponent::TickComponent(float DeltaTime, enum ELevelTick Tick
 		}
 	}
 
+	// make sure that we continue execution after all pending latent aborts finished
+	const bool bJustFinishedLatentAborts = TrackPendingLatentAborts();
+	if (bJustFinishedLatentAborts)
+	{
+		if (bRequestedStop)
+		{
+			StopTree(EBTStopMode::Safe);
+		}
+		else
+		{
+			// force new search if there were any execution requests while waiting for aborting task
+			if (ExecutionRequest.ExecuteNode)
+			{
+				UE_VLOG(GetOwner(), LogBehaviorTree, Log, TEXT("> found valid ExecutionRequest, locking PendingExecution data to force new search!"));
+				PendingExecution.Lock();
+
+				if (ExecutionRequest.SearchEnd.IsSet())
+				{
+					UE_VLOG(GetOwner(), LogBehaviorTree, Log, TEXT("> removing limit from end of search range! [abort done]"));
+					ExecutionRequest.SearchEnd = FBTNodeIndex();
+				}
+			}
+
+			ScheduleExecutionUpdate();
+		}
+	}
+
 	bool bActiveAuxiliaryNodeDTDirty = false;
 	if (bRequestedFlowUpdate)
 	{
@@ -1649,7 +1643,7 @@ void UBehaviorTreeComponent::ProcessExecutionRequest()
 		return;
 	}
 
-	if (bWaitingForAbortingTasks)
+	if (bWaitingForLatentAborts)
 	{
 		UE_VLOG(GetOwner(), LogBehaviorTree, Verbose, TEXT("Ignoring ProcessExecutionRequest call, aborting task must finish first"));
 		return;
@@ -1903,7 +1897,7 @@ void UBehaviorTreeComponent::ProcessExecutionRequest()
 void UBehaviorTreeComponent::ProcessPendingExecution()
 {
 	// can't continue if current task is still aborting
-	if (bWaitingForAbortingTasks || !PendingExecution.IsSet())
+	if (bWaitingForLatentAborts || !PendingExecution.IsSet())
 	{
 		return;
 	}
@@ -2309,29 +2303,52 @@ void UBehaviorTreeComponent::UnregisterParallelTask(const UBTTaskNode* TaskNode,
 			}
 		}
 	}
+}
+
+bool UBehaviorTreeComponent::TrackPendingLatentAborts()
+{
+	// nothing to track if we are not currently waiting for latent aborts
+	if (!bWaitingForLatentAborts)
+	{
+		return false;
+	}
+
+	// update our internal flag	
+	bWaitingForLatentAborts = HasActiveLatentAborts();
+
+	// return true if we are no longer waiting (at this point we know that we were previously waiting on latent abortes)
+	return !bWaitingForLatentAborts;
+}
 
-	if (bShouldUpdate)
+void UBehaviorTreeComponent::TrackNewLatentAborts()
+{
+	// already waiting for latent aborts, no need to look for new ones 
+	if (bWaitingForLatentAborts)
 	{
-		UpdateAbortingTasks();
+		return;
 	}
+
+	bWaitingForLatentAborts = HasActiveLatentAborts();
 }
 
-void UBehaviorTreeComponent::UpdateAbortingTasks()
+bool UBehaviorTreeComponent::HasActiveLatentAborts() const
 {
-	bWaitingForAbortingTasks = InstanceStack.Num() ? (InstanceStack.Last().ActiveNodeType == EBTActiveNode::AbortingTask) : false;
+	bool bHasActiveLatentAborts = InstanceStack.Num() ? (InstanceStack.Last().ActiveNodeType == EBTActiveNode::AbortingTask) : false;
 
-	for (int32 InstanceIndex = 0; InstanceIndex < InstanceStack.Num() && !bWaitingForAbortingTasks; InstanceIndex++)
+	for (int32 InstanceIndex = 0; InstanceIndex < InstanceStack.Num() && !bHasActiveLatentAborts; InstanceIndex++)
 	{
-		FBehaviorTreeInstance& InstanceInfo = InstanceStack[InstanceIndex];
+		const FBehaviorTreeInstance& InstanceInfo = InstanceStack[InstanceIndex];
 		for (const FBehaviorTreeParallelTask& ParallelInfo : InstanceInfo.GetParallelTasks())
 		{
 			if (ParallelInfo.Status == EBTTaskStatus::Aborting)
 			{
-				bWaitingForAbortingTasks = true;
+				bHasActiveLatentAborts = true;
 				break;
 			}
 		}
 	}
+
+	return bHasActiveLatentAborts;
 }
 
 bool UBehaviorTreeComponent::PushInstance(UBehaviorTree& TreeAsset)
