Made EQSTestingPawn no longer abstract and just marked it as `hidedropdown` so that it doesn't show up in pawn-selection drop-downs.

diff --git a/Engine/Source/Runtime/AIModule/Classes/EnvironmentQuery/EQSTestingPawn.h b/Engine/Source/Runtime/AIModule/Classes/EnvironmentQuery/EQSTestingPawn.h
--- a/Engine/Source/Runtime/AIModule/Classes/EnvironmentQuery/EQSTestingPawn.h
+++ b/Engine/Source/Runtime/AIModule/Classes/EnvironmentQuery/EQSTestingPawn.h
@@ -20,10 +20,7 @@ enum class EEnvQueryHightlightMode : uint8
 	Best25Pct UMETA(DisplayName = "Best 25%"),
 };
 
-/** this class is abstract even though it's perfectly functional on its own.
- *	The reason is to stop it from showing as valid player pawn type when configuring 
- *	project's game mode. */
-UCLASS(abstract, hidecategories=(Advanced, Attachment, Mesh, Animation, Clothing, Physics, Rendering, Lighting, Activation, CharacterMovement, AgentPhysics, Avoidance, MovementComponent, Velocity, Shape, Camera, Input, Layers, SkeletalMesh, Optimization, Pawn, Replication, Actor))
+UCLASS(hidedropdown, hidecategories=(Advanced, Attachment, Mesh, Animation, Clothing, Physics, Rendering, Lighting, Activation, CharacterMovement, AgentPhysics, Avoidance, MovementComponent, Velocity, Shape, Camera, Input, Layers, SkeletalMesh, Optimization, Pawn, Replication, Actor))
 class AIMODULE_API AEQSTestingPawn : public ACharacter, public IEQSQueryResultSourceInterface
 {
 	GENERATED_UCLASS_BODY()
