Type conversion warning fix in BlackboardComponent

diff --git a/Engine/Source/Runtime/AIModule/Classes/BehaviorTree/BlackboardComponent.h b/Engine/Source/Runtime/AIModule/Classes/BehaviorTree/BlackboardComponent.h
--- a/Engine/Source/Runtime/AIModule/Classes/BehaviorTree/BlackboardComponent.h
+++ b/Engine/Source/Runtime/AIModule/Classes/BehaviorTree/BlackboardComponent.h
@@ -363,7 +363,7 @@ bool UBlackboardComponent::SetValue(FBlackboard::FKey KeyID, typename TDataClass
 					if (OtherBlackboard != nullptr && ShouldSyncWithBlackboard(*OtherBlackboard))
 					{
 						UBlackboardData* const OtherBlackboardAsset = OtherBlackboard->GetBlackboardAsset();
-						const int32 OtherKeyID = OtherBlackboardAsset ? OtherBlackboardAsset->GetKeyID(EntryInfo->EntryName) : FBlackboard::InvalidKey;
+						const FBlackboard::FKey OtherKeyID = OtherBlackboardAsset ? OtherBlackboardAsset->GetKeyID(EntryInfo->EntryName) : FBlackboard::InvalidKey;
 						if (OtherKeyID != FBlackboard::InvalidKey)
 						{
 							UBlackboardKeyType* OtherKeyOb = EntryInfo->KeyType->HasInstance() ? OtherBlackboard->KeyInstances[OtherKeyID] : EntryInfo->KeyType;
