Add a InitNotifyFlags function to UBTAuxiliaryNode, UBTTaskNode, UBTCompositeNode  and UBTDecorator to setup the various flags that must be set when some functions are overriden.
Add macros to call those functions, this is necessary as those functions are protected.
Use those macros in various nodes (not all of them, this is not possible for some blueprint tasks)
#tests PIE
#rnx

diff --git a/Engine/Source/Developer/AITestSuite/Private/BehaviorTree/TestBTDecorator_DelayedAbort.cpp b/Engine/Source/Developer/AITestSuite/Private/BehaviorTree/TestBTDecorator_DelayedAbort.cpp
--- a/Engine/Source/Developer/AITestSuite/Private/BehaviorTree/TestBTDecorator_DelayedAbort.cpp
+++ b/Engine/Source/Developer/AITestSuite/Private/BehaviorTree/TestBTDecorator_DelayedAbort.cpp
@@ -8,9 +8,8 @@ UTestBTDecorator_DelayedAbort::UTestBTDecorator_DelayedAbort(const FObjectInitia
 	NodeName = "Delayed Abort";
 	DelayTicks = 5;
 	bOnlyOnce = true;
-	
-	bNotifyTick = true;
-	bNotifyBecomeRelevant = true;
+
+	INIT_DECORATOR_NODE_NOTIFY_FLAGS();
 
 	bAllowAbortNone = false;
 	bAllowAbortLowerPri = false;
diff --git a/Engine/Source/Developer/AITestSuite/Private/BehaviorTree/TestBTService_Log.cpp b/Engine/Source/Developer/AITestSuite/Private/BehaviorTree/TestBTService_Log.cpp
--- a/Engine/Source/Developer/AITestSuite/Private/BehaviorTree/TestBTService_Log.cpp
+++ b/Engine/Source/Developer/AITestSuite/Private/BehaviorTree/TestBTService_Log.cpp
@@ -8,8 +8,7 @@ UTestBTService_Log::UTestBTService_Log(const FObjectInitializer& ObjectInitializ
 {
 	NodeName = "LogService";
 
-	bNotifyBecomeRelevant = true;
-	bNotifyCeaseRelevant = true;
+	INIT_AUXILIARY_NODE_NOTIFY_FLAGS();
 
 	LogActivation = INDEX_NONE;
 	LogDeactivation = INDEX_NONE;
diff --git a/Engine/Source/Developer/AITestSuite/Private/BehaviorTree/TestBTTask_LatentWithFlags.cpp b/Engine/Source/Developer/AITestSuite/Private/BehaviorTree/TestBTTask_LatentWithFlags.cpp
--- a/Engine/Source/Developer/AITestSuite/Private/BehaviorTree/TestBTTask_LatentWithFlags.cpp
+++ b/Engine/Source/Developer/AITestSuite/Private/BehaviorTree/TestBTTask_LatentWithFlags.cpp
@@ -6,12 +6,12 @@
 
 UTestBTTask_LatentWithFlags::UTestBTTask_LatentWithFlags()
 {
-	bNotifyTick = true;
+	INIT_TASK_NODE_NOTIFY_FLAGS();
 }
 
 EBTNodeResult::Type UTestBTTask_LatentWithFlags::ExecuteTask(UBehaviorTreeComponent& OwnerComp, uint8* NodeMemory)
 {
-	FBTLatentTaskMemory* MyMemory = (FBTLatentTaskMemory*)NodeMemory;
+	FBTLatentTaskMemory* MyMemory = CastInstanceNodeMemory<FBTLatentTaskMemory>(NodeMemory);
 	MyMemory->FlagFrameIdx = ExecuteHalfTicks + FAITestHelpers::FramesCounter();
 	MyMemory->EndFrameIdx = MyMemory->FlagFrameIdx + ExecuteHalfTicks;
 	MyMemory->bFlagSet = false;
@@ -32,7 +32,7 @@ EBTNodeResult::Type UTestBTTask_LatentWithFlags::ExecuteTask(UBehaviorTreeCompon
 
 EBTNodeResult::Type UTestBTTask_LatentWithFlags::AbortTask(UBehaviorTreeComponent& OwnerComp, uint8* NodeMemory)
 {
-	FBTLatentTaskMemory* MyMemory = (FBTLatentTaskMemory*)NodeMemory;
+	FBTLatentTaskMemory* MyMemory = CastInstanceNodeMemory<FBTLatentTaskMemory>(NodeMemory);
 	MyMemory->FlagFrameIdx = AbortHalfTicks + FAITestHelpers::FramesCounter();
 	MyMemory->EndFrameIdx = MyMemory->FlagFrameIdx + AbortHalfTicks;
 	MyMemory->bFlagSet = false;
@@ -53,7 +53,7 @@ EBTNodeResult::Type UTestBTTask_LatentWithFlags::AbortTask(UBehaviorTreeComponen
 
 void UTestBTTask_LatentWithFlags::TickTask(UBehaviorTreeComponent& OwnerComp, uint8* NodeMemory, float DeltaSeconds)
 {
-	FBTLatentTaskMemory* MyMemory = (FBTLatentTaskMemory*)NodeMemory;
+	FBTLatentTaskMemory* MyMemory = CastInstanceNodeMemory<FBTLatentTaskMemory>(NodeMemory);
 
 	LogExecution(OwnerComp, MyMemory->bIsAborting ? LogIndexAborting : LogIndexExecuting);
 
diff --git a/Engine/Source/Developer/AITestSuite/Private/BehaviorTree/TestBTTask_Log.cpp b/Engine/Source/Developer/AITestSuite/Private/BehaviorTree/TestBTTask_Log.cpp
--- a/Engine/Source/Developer/AITestSuite/Private/BehaviorTree/TestBTTask_Log.cpp
+++ b/Engine/Source/Developer/AITestSuite/Private/BehaviorTree/TestBTTask_Log.cpp
@@ -12,12 +12,12 @@ UTestBTTask_Log::UTestBTTask_Log(const FObjectInitializer& ObjectInitializer) :
 	LogResult = EBTNodeResult::Succeeded;
 	LogTickIndex = -1;
 
-	bNotifyTick = true;
+	INIT_TASK_NODE_NOTIFY_FLAGS();
 }
 
 EBTNodeResult::Type UTestBTTask_Log::ExecuteTask(UBehaviorTreeComponent& OwnerComp, uint8* NodeMemory)
 {
-	FBTLogTaskMemory* MyMemory = (FBTLogTaskMemory*)NodeMemory;
+	FBTLogTaskMemory* MyMemory = CastInstanceNodeMemory<FBTLogTaskMemory>(NodeMemory);
 	MyMemory->EndFrameIdx = ExecutionTicks + FAITestHelpers::FramesCounter();
 
 	LogExecution(OwnerComp, LogIndex);
@@ -31,7 +31,7 @@ EBTNodeResult::Type UTestBTTask_Log::ExecuteTask(UBehaviorTreeComponent& OwnerCo
 
 void UTestBTTask_Log::TickTask(UBehaviorTreeComponent& OwnerComp, uint8* NodeMemory, float DeltaSeconds)
 {
-	FBTLogTaskMemory* MyMemory = (FBTLogTaskMemory*)NodeMemory;
+	FBTLogTaskMemory* MyMemory = CastInstanceNodeMemory<FBTLogTaskMemory>(NodeMemory);
 
 	if (LogTickIndex != -1)
 	{
diff --git a/Engine/Source/Runtime/AIModule/Classes/BehaviorTree/BTAuxiliaryNode.h b/Engine/Source/Runtime/AIModule/Classes/BehaviorTree/BTAuxiliaryNode.h
--- a/Engine/Source/Runtime/AIModule/Classes/BehaviorTree/BTAuxiliaryNode.h
+++ b/Engine/Source/Runtime/AIModule/Classes/BehaviorTree/BTAuxiliaryNode.h
@@ -3,7 +3,6 @@
 #pragma once
 
 #include "CoreMinimal.h"
-#include "UObject/ObjectMacros.h"
 #include "BehaviorTree/BTNode.h"
 #include "BTAuxiliaryNode.generated.h"
 
@@ -75,22 +74,25 @@ protected:
 	/** if set, OnTick will be used */
 	uint8 bNotifyTick : 1;
 
-	/** if set, conditional tick will use remaining time form node's memory */
+	/** if set, conditional tick will use remaining time from node's memory */
 	uint8 bTickIntervals : 1;
 
 	/** child index in parent node */
 	uint8 ChildIndex;
 
 	/** called when auxiliary node becomes active
-	 * this function should be considered as const (don't modify state of object) if node is not instanced! */
+	* this function should be considered as const (don't modify state of object) if node is not instanced!  
+	* bNotifyBecomeRelevant must be set to true when this function must be called */
 	virtual void OnBecomeRelevant(UBehaviorTreeComponent& OwnerComp, uint8* NodeMemory);
 
 	/** called when auxiliary node becomes inactive
-	 * this function should be considered as const (don't modify state of object) if node is not instanced! */
+	* this function should be considered as const (don't modify state of object) if node is not instanced!  
+	* bNotifyCeaseRelevant must be set to true when this function must be called */
 	virtual void OnCeaseRelevant(UBehaviorTreeComponent& OwnerComp, uint8* NodeMemory);
 
 	/** tick function
-	 * this function should be considered as const (don't modify state of object) if node is not instanced! */
+	* this function should be considered as const (don't modify state of object) if node is not instanced!   
+	* bNotifyTick must be set to true when this function must be called */
 	virtual void TickNode(UBehaviorTreeComponent& OwnerComp, uint8* NodeMemory, float DeltaSeconds);
 
 	/** sets next tick time */
@@ -98,8 +100,25 @@ protected:
 
 	/** gets remaining time for next tick */
 	float GetNextTickRemainingTime(uint8* NodeMemory) const;
+	
+	template<typename T,
+			typename TickNodeFunc = decltype(&T::TickNode),
+			typename OnBecomeRelevantFunc = decltype(&T::OnBecomeRelevant),
+			typename OnCeaseRelevantFunc = decltype(&T::OnCeaseRelevant)>
+	void InitNotifyFlags(const T* Node, TickNodeFunc TickFunction,
+										OnBecomeRelevantFunc OnBecomeRelevantFunction,
+										OnCeaseRelevantFunc OnCeaseRelevantFunction)
+	{
+		bNotifyTick |= &UBTAuxiliaryNode::TickNode != TickFunction;
+		bNotifyBecomeRelevant |= &UBTAuxiliaryNode::OnBecomeRelevant != OnBecomeRelevantFunction;
+		bNotifyCeaseRelevant |= &UBTAuxiliaryNode::OnCeaseRelevant != OnCeaseRelevantFunction;
+	}
 };
 
+#define INIT_AUXILIARY_NODE_NOTIFY_FLAGS() \
+	using NodeType = TRemovePointer<decltype(this)>::Type; \
+	InitNotifyFlags(this, &NodeType::TickNode, &NodeType::OnBecomeRelevant, &NodeType::OnCeaseRelevant)
+
 FORCEINLINE uint8 UBTAuxiliaryNode::GetChildIndex() const
 {
 	return ChildIndex;
diff --git a/Engine/Source/Runtime/AIModule/Classes/BehaviorTree/BTCompositeNode.h b/Engine/Source/Runtime/AIModule/Classes/BehaviorTree/BTCompositeNode.h
--- a/Engine/Source/Runtime/AIModule/Classes/BehaviorTree/BTCompositeNode.h
+++ b/Engine/Source/Runtime/AIModule/Classes/BehaviorTree/BTCompositeNode.h
@@ -2,7 +2,6 @@
 #pragma once
 
 #include "CoreMinimal.h"
-#include "UObject/ObjectMacros.h"
 #include "BehaviorTree/BTNode.h"
 #include "BTCompositeNode.generated.h"
 
@@ -208,22 +207,28 @@ protected:
 	/** execution index of last node in child branches */
 	uint16 LastExecutionIndex;
 
-	/** called just after child execution, allows to modify result */
+	/** called just after child execution, allows to modify result  
+	* bUseChildExecutionNotify must be set to true when this function must be called */
 	virtual void NotifyChildExecution(UBehaviorTreeComponent& OwnerComp, uint8* NodeMemory, int32 ChildIdx, EBTNodeResult::Type& NodeResult) const;
 
-	/** called when start enters this node */
+	/** called when start enters this node  
+	* bUseNodeActivationNotify must be set to true when this function must be called */
 	virtual void NotifyNodeActivation(FBehaviorTreeSearchData& SearchData) const;
 
-	/** called when start leaves this node */
+	/** called when start leaves this node 
+	* bUseNodeDeactivationNotify must be set to true when this function must be called */
 	virtual void NotifyNodeDeactivation(FBehaviorTreeSearchData& SearchData, EBTNodeResult::Type& NodeResult) const;
 
-	/** check if NotifyDecoratorsOnActivation is allowed, requires bUseDecoratorsActivationCheck flag  */
+	/** check if NotifyDecoratorsOnActivation is allowed, requires bUseDecoratorsActivationCheck flag 
+	* bUseDecoratorsActivationCheck must be set to true when this function must be called */
 	virtual bool CanNotifyDecoratorsOnActivation(FBehaviorTreeSearchData& SearchData, int32 ChildIdx) const;
 
-	/** check if NotifyDecoratorsOnDeactivation is allowed, requires bUseDecoratorsDeactivationCheck flag  */
+	/** check if NotifyDecoratorsOnDeactivation is allowed, requires bUseDecoratorsDeactivationCheck flag   
+	* bUseDecoratorsDeactivationCheck must be set to true when this function must be called */
 	virtual bool CanNotifyDecoratorsOnDeactivation(FBehaviorTreeSearchData& SearchData, int32 ChildIdx, EBTNodeResult::Type& NodeResult) const;
 
-	/** check if NotifyDecoratorsOnFailedActivation is allowed, requires bUseDecoratorsActivationCheck flag  */
+	/** check if NotifyDecoratorsOnFailedActivation is allowed, requires bUseDecoratorsActivationCheck flag  
+	* bUseDecoratorsFailedActivationCheck must be set to true when this function must be called */
 	virtual bool CanNotifyDecoratorsOnFailedActivation(FBehaviorTreeSearchData& SearchData, int32 ChildIdx, EBTNodeResult::Type& NodeResult) const;
 
 	/** runs through decorators on given child node and notify them about activation */
@@ -243,8 +248,39 @@ protected:
 
 protected:
 	virtual int32 GetNextChildHandler(struct FBehaviorTreeSearchData& SearchData, int32 PrevChild, EBTNodeResult::Type LastResult) const { return BTSpecialChild::ReturnToParent; }
+
+	template<typename T,
+			typename NotifyChildExecutionFunc = decltype(&T::NotifyChildExecution),
+			typename NotifyNodeActivationFunc = decltype(&T::NotifyNodeActivation),
+			typename NotifyNodeDeactivationFunc = decltype(&T::NotifyNodeDeactivation),
+			typename CanNotifyDecoratorsOnActivationFunc = decltype(&T::CanNotifyDecoratorsOnActivation),
+			typename CanNotifyDecoratorsOnDeactivationFunc = decltype(&T::CanNotifyDecoratorsOnDeactivation),
+			typename CanNotifyDecoratorsOnFailedActivationFunc = decltype(&T::CanNotifyDecoratorsOnFailedActivation)>
+	void InitNotifyFlags(const T* Node, NotifyChildExecutionFunc ChildExecutionFunction,
+									NotifyNodeActivationFunc NodeActivationFunction,
+									NotifyNodeDeactivationFunc NodeDeactivationFunction,
+									CanNotifyDecoratorsOnActivationFunc CanNotifyDecoratorsOnActivationFunction,
+									CanNotifyDecoratorsOnDeactivationFunc CanNotifyDecoratorsOnDeactivationFunction,
+									CanNotifyDecoratorsOnFailedActivationFunc CanNotifyDecoratorsOnFailedActivationFunction)
+	{
+		bUseChildExecutionNotify |= &UBTCompositeNode::NotifyChildExecution != ChildExecutionFunction;
+		bUseNodeActivationNotify |= &UBTCompositeNode::NotifyNodeActivation != NodeActivationFunction;
+		bUseNodeDeactivationNotify |= &UBTCompositeNode::NotifyNodeDeactivation != NodeDeactivationFunction;
+		bUseDecoratorsActivationCheck |= &UBTCompositeNode::CanNotifyDecoratorsOnActivation != CanNotifyDecoratorsOnActivationFunction;
+		bUseDecoratorsDeactivationCheck |= &UBTCompositeNode::CanNotifyDecoratorsOnDeactivation != CanNotifyDecoratorsOnDeactivationFunction;
+		bUseDecoratorsFailedActivationCheck |= &UBTCompositeNode::CanNotifyDecoratorsOnFailedActivation != CanNotifyDecoratorsOnFailedActivationFunction;
+	}
 };
 
+#define INIT_COMPOSITE_NODE_NOTIFY_FLAGS() \
+	using NodeType = TRemovePointer<decltype(this)>::Type; \
+	InitNotifyFlags(this,\
+					&NodeType::NotifyChildExecution,\
+					&NodeType::NotifyNodeActivation,\
+					&NodeType::NotifyNodeDeactivation, \
+					&NodeType::CanNotifyDecoratorsOnActivation,\
+					&NodeType::CanNotifyDecoratorsOnDeactivation,\
+					&NodeType::CanNotifyDecoratorsOnFailedActivation)
 
 //////////////////////////////////////////////////////////////////////////
 // Inlines
@@ -253,8 +289,8 @@ FORCEINLINE UBTNode* UBTCompositeNode::GetChildNode(int32 Index) const
 {
 	return Children.IsValidIndex(Index) ?
 		(Children[Index].ChildComposite ?
-			(UBTNode*)Children[Index].ChildComposite :
-			(UBTNode*)Children[Index].ChildTask) :
+			static_cast<UBTNode*>(Children[Index].ChildComposite) :
+			static_cast<UBTNode*>(Children[Index].ChildTask)) :
 		nullptr;
 }
 
diff --git a/Engine/Source/Runtime/AIModule/Classes/BehaviorTree/BTDecorator.h b/Engine/Source/Runtime/AIModule/Classes/BehaviorTree/BTDecorator.h
--- a/Engine/Source/Runtime/AIModule/Classes/BehaviorTree/BTDecorator.h
+++ b/Engine/Source/Runtime/AIModule/Classes/BehaviorTree/BTDecorator.h
@@ -101,15 +101,18 @@ protected:
 	void SetIsInversed(bool bShouldBeInversed);
 
 	/** called when underlying node is activated
-	  * this function should be considered as const (don't modify state of object) if node is not instanced! */
+	* this function should be considered as const (don't modify state of object) if node is not instanced! 
+	* bNotifyActivation must be set to true when this function must be called */
 	virtual void OnNodeActivation(FBehaviorTreeSearchData& SearchData);
 
 	/** called when underlying node has finished
-	 * this function should be considered as const (don't modify state of object) if node is not instanced! */
+	* this function should be considered as const (don't modify state of object) if node is not instanced! 
+	* bNotifyDeactivation must be set to true when this function must be called */
 	virtual void OnNodeDeactivation(FBehaviorTreeSearchData& SearchData, EBTNodeResult::Type NodeResult);
 
 	/** called when underlying node was processed (deactivated or failed to activate)
-	 * this function should be considered as const (don't modify state of object) if node is not instanced! */
+	* this function should be considered as const (don't modify state of object) if node is not instanced! 
+	* bNotifyProcessed must be set to true when this function must be called */
 	virtual void OnNodeProcessed(FBehaviorTreeSearchData& SearchData, EBTNodeResult::Type& NodeResult);
 
 	/** calculates raw, core value of decorator's condition. Should not include calling IsInversed */
@@ -122,6 +125,28 @@ protected:
 
 	friend FBehaviorDecoratorDetails;
 
+	template<typename T,
+			typename TickNodeFunc = decltype(&T::TickNode),
+			typename OnBecomeRelevantFunc = decltype(&T::OnBecomeRelevant),
+			typename OnCeaseRelevantFunc = decltype(&T::OnCeaseRelevant),
+			typename OnNodeActivationFunc = decltype(&T::OnNodeActivation),
+			typename OnNodeDeactivationFunc = decltype(&T::OnNodeDeactivation),
+			typename OnNodeProcessedFunc = decltype(&T::OnNodeProcessed)>
+	void InitNotifyFlags(const T* Node, TickNodeFunc TickFunction,
+										OnBecomeRelevantFunc OnBecomeRelevantFunction,
+										OnCeaseRelevantFunc OnCeaseRelevantFunction,
+										OnNodeActivationFunc OnNodeActivationFunction,
+										OnNodeDeactivationFunc OnNodeDeactivationFunction,
+										OnNodeProcessedFunc OnNodeProcessedFunction)
+	{
+		bNotifyTick |= &UBTDecorator::TickNode != TickFunction;
+		bNotifyBecomeRelevant |= &UBTDecorator::OnBecomeRelevant != OnBecomeRelevantFunction;
+		bNotifyCeaseRelevant |= &UBTDecorator::OnCeaseRelevant != OnCeaseRelevantFunction;
+		bNotifyActivation |= &UBTDecorator::OnNodeActivation != OnNodeActivationFunction;
+		bNotifyDeactivation |= &UBTDecorator::OnNodeDeactivation != OnNodeDeactivationFunction;
+		bNotifyProcessed |= &UBTDecorator::OnNodeProcessed != OnNodeProcessedFunction;
+	}
+
 	//----------------------------------------------------------------------//
 	// DEPRECATED
 	//----------------------------------------------------------------------//
@@ -129,6 +154,14 @@ protected:
 	void InitializeDecorator(uint8 InChildIndex);
 };
 
+#define INIT_DECORATOR_NODE_NOTIFY_FLAGS() \
+	using NodeType = TRemovePointer<decltype(this)>::Type; \
+	InitNotifyFlags(this, &NodeType::TickNode, \
+					&NodeType::OnBecomeRelevant, \
+					&NodeType::OnCeaseRelevant, \
+					&NodeType::OnNodeActivation, \
+					&NodeType::OnNodeDeactivation, \
+					&NodeType::OnNodeProcessed) \
 
 //////////////////////////////////////////////////////////////////////////
 // Inlines
diff --git a/Engine/Source/Runtime/AIModule/Classes/BehaviorTree/BTTaskNode.h b/Engine/Source/Runtime/AIModule/Classes/BehaviorTree/BTTaskNode.h
--- a/Engine/Source/Runtime/AIModule/Classes/BehaviorTree/BTTaskNode.h
+++ b/Engine/Source/Runtime/AIModule/Classes/BehaviorTree/BTTaskNode.h
@@ -3,7 +3,6 @@
 #pragma once
 
 #include "CoreMinimal.h"
-#include "UObject/ObjectMacros.h"
 #include "BehaviorTree/BTNode.h"
 #include "BTTaskNode.generated.h"
 
@@ -90,9 +89,10 @@ protected:
 
 	/** if set, OnTaskFinished will be called */
 	uint32 bNotifyTaskFinished : 1;
-	
+
 	/** ticks this task 
-	 * this function should be considered as const (don't modify state of object) if node is not instanced! */
+	 * this function should be considered as const (don't modify state of object) if node is not instanced! 
+	* bNotifyTick must be set to true when this function must be called */
 	virtual void TickTask(UBehaviorTreeComponent& OwnerComp, uint8* NodeMemory, float DeltaSeconds);
 
 	/** message handler, default implementation will finish latent execution/abortion
@@ -100,7 +100,8 @@ protected:
 	virtual void OnMessage(UBehaviorTreeComponent& OwnerComp, uint8* NodeMemory, FName Message, int32 RequestID, bool bSuccess);
 
 	/** called when task execution is finished
-	 * this function should be considered as const (don't modify state of object) if node is not instanced! */
+	* this function should be considered as const (don't modify state of object) if node is not instanced! 
+	* bNotifyTaskFinished must be set to true when this function must be called */
 	virtual void OnTaskFinished(UBehaviorTreeComponent& OwnerComp, uint8* NodeMemory, EBTNodeResult::Type TaskResult);
 
 	/** register message observer */
@@ -109,8 +110,21 @@ protected:
 	
 	/** unregister message observers */
 	void StopWaitingForMessages(UBehaviorTreeComponent& OwnerComp) const;
+	
+	template<typename T,
+		typename TickTaskFunc = decltype(&T::TickTask),
+		typename OnTaskFinishedFunc = decltype(&T::OnTaskFinished)>
+	void InitNotifyFlags(const T* Node, TickTaskFunc TickFunction, OnTaskFinishedFunc OnTaskFinishedFunction)
+	{
+		bNotifyTick |= &UBTTaskNode::TickTask != TickFunction;
+		bNotifyTaskFinished |= &UBTTaskNode::OnTaskFinished != OnTaskFinishedFunction;
+	}
 };
 
+#define INIT_TASK_NODE_NOTIFY_FLAGS() \
+	using NodeType = TRemovePointer<decltype(this)>::Type; \
+	InitNotifyFlags(this, &NodeType::TickTask, &NodeType::OnTaskFinished)
+
 FORCEINLINE bool UBTTaskNode::ShouldIgnoreRestartSelf() const
 {
 	return bIgnoreRestartSelf;
diff --git a/Engine/Source/Runtime/AIModule/Private/BehaviorTree/Composites/BTComposite_SimpleParallel.cpp b/Engine/Source/Runtime/AIModule/Private/BehaviorTree/Composites/BTComposite_SimpleParallel.cpp
--- a/Engine/Source/Runtime/AIModule/Private/BehaviorTree/Composites/BTComposite_SimpleParallel.cpp
+++ b/Engine/Source/Runtime/AIModule/Private/BehaviorTree/Composites/BTComposite_SimpleParallel.cpp
@@ -9,9 +9,7 @@
 UBTComposite_SimpleParallel::UBTComposite_SimpleParallel(const FObjectInitializer& ObjectInitializer) : Super(ObjectInitializer)
 {
 	NodeName = "Simple Parallel";
-	bUseChildExecutionNotify = true;
-	bUseNodeDeactivationNotify = true;
-	bUseDecoratorsDeactivationCheck = true;
+	INIT_COMPOSITE_NODE_NOTIFY_FLAGS();
 	bApplyDecoratorScope = true;
 }
 
diff --git a/Engine/Source/Runtime/AIModule/Private/BehaviorTree/Decorators/BTDecorator_BlackboardBase.cpp b/Engine/Source/Runtime/AIModule/Private/BehaviorTree/Decorators/BTDecorator_BlackboardBase.cpp
--- a/Engine/Source/Runtime/AIModule/Private/BehaviorTree/Decorators/BTDecorator_BlackboardBase.cpp
+++ b/Engine/Source/Runtime/AIModule/Private/BehaviorTree/Decorators/BTDecorator_BlackboardBase.cpp
@@ -7,8 +7,7 @@ UBTDecorator_BlackboardBase::UBTDecorator_BlackboardBase(const FObjectInitialize
 {
 	NodeName = "BlackboardBase";
 
-	bNotifyBecomeRelevant = true;
-	bNotifyCeaseRelevant = true;
+	INIT_DECORATOR_NODE_NOTIFY_FLAGS();
 
 	// empty KeySelector = allow everything
 }
diff --git a/Engine/Source/Runtime/AIModule/Private/BehaviorTree/Decorators/BTDecorator_BlueprintBase.cpp b/Engine/Source/Runtime/AIModule/Private/BehaviorTree/Decorators/BTDecorator_BlueprintBase.cpp
--- a/Engine/Source/Runtime/AIModule/Private/BehaviorTree/Decorators/BTDecorator_BlueprintBase.cpp
+++ b/Engine/Source/Runtime/AIModule/Private/BehaviorTree/Decorators/BTDecorator_BlueprintBase.cpp
@@ -18,6 +18,7 @@ UBTDecorator_BlueprintBase::UBTDecorator_BlueprintBase(const FObjectInitializer&
 	ReceiveObserverDeactivatedImplementations = FBTNodeBPImplementationHelper::CheckEventImplementationVersion(TEXT("ReceiveObserverDeactivated"), TEXT("ReceiveObserverDeactivatedAI"), *this, *StopAtClass);
 	PerformConditionCheckImplementations = FBTNodeBPImplementationHelper::CheckEventImplementationVersion(TEXT("PerformConditionCheck"), TEXT("PerformConditionCheckAI"), *this, *StopAtClass);
 
+	INIT_DECORATOR_NODE_NOTIFY_FLAGS();
 	bNotifyBecomeRelevant = ReceiveObserverActivatedImplementations != 0;
 	bNotifyCeaseRelevant = ReceiveObserverDeactivatedImplementations != 0;
 	bNotifyTick = ReceiveTickImplementations != 0;
diff --git a/Engine/Source/Runtime/AIModule/Private/BehaviorTree/Decorators/BTDecorator_CompareBBEntries.cpp b/Engine/Source/Runtime/AIModule/Private/BehaviorTree/Decorators/BTDecorator_CompareBBEntries.cpp
--- a/Engine/Source/Runtime/AIModule/Private/BehaviorTree/Decorators/BTDecorator_CompareBBEntries.cpp
+++ b/Engine/Source/Runtime/AIModule/Private/BehaviorTree/Decorators/BTDecorator_CompareBBEntries.cpp
@@ -8,8 +8,7 @@ UBTDecorator_CompareBBEntries::UBTDecorator_CompareBBEntries(const FObjectInitia
 	NodeName = "Compare Blackboard entries";
 
 	Operator = EBlackBoardEntryComparison::Equal;
-	bNotifyBecomeRelevant = true;
-	bNotifyCeaseRelevant = true;
+	INIT_DECORATOR_NODE_NOTIFY_FLAGS();
 }
 
 void UBTDecorator_CompareBBEntries::InitializeFromAsset(UBehaviorTree& Asset)
diff --git a/Engine/Source/Runtime/AIModule/Private/BehaviorTree/Decorators/BTDecorator_ConditionalLoop.cpp b/Engine/Source/Runtime/AIModule/Private/BehaviorTree/Decorators/BTDecorator_ConditionalLoop.cpp
--- a/Engine/Source/Runtime/AIModule/Private/BehaviorTree/Decorators/BTDecorator_ConditionalLoop.cpp
+++ b/Engine/Source/Runtime/AIModule/Private/BehaviorTree/Decorators/BTDecorator_ConditionalLoop.cpp
@@ -8,7 +8,7 @@
 UBTDecorator_ConditionalLoop::UBTDecorator_ConditionalLoop(const FObjectInitializer& ObjectInitializer) : Super(ObjectInitializer)
 {
 	NodeName = "Conditional Loop";
-	bNotifyDeactivation = true;
+	INIT_DECORATOR_NODE_NOTIFY_FLAGS();
 
 	bAllowAbortNone = false;
 	bAllowAbortLowerPri = false;
diff --git a/Engine/Source/Runtime/AIModule/Private/BehaviorTree/Decorators/BTDecorator_ConeCheck.cpp b/Engine/Source/Runtime/AIModule/Private/BehaviorTree/Decorators/BTDecorator_ConeCheck.cpp
--- a/Engine/Source/Runtime/AIModule/Private/BehaviorTree/Decorators/BTDecorator_ConeCheck.cpp
+++ b/Engine/Source/Runtime/AIModule/Private/BehaviorTree/Decorators/BTDecorator_ConeCheck.cpp
@@ -17,8 +17,7 @@ UBTDecorator_ConeCheck::UBTDecorator_ConeCheck(const FObjectInitializer& ObjectI
 	Observed.AddObjectFilter(this, GET_MEMBER_NAME_CHECKED(UBTDecorator_ConeCheck, Observed), AActor::StaticClass());
 	Observed.AddVectorFilter(this, GET_MEMBER_NAME_CHECKED(UBTDecorator_ConeCheck, Observed));
 
-	bNotifyBecomeRelevant = true;
-	bNotifyTick = true;
+	INIT_DECORATOR_NODE_NOTIFY_FLAGS();
 
 	// KeepInCone always abort current branch
 	FlowAbortMode = EBTFlowAbortMode::None;
@@ -86,7 +85,7 @@ bool UBTDecorator_ConeCheck::CalculateRawConditionValue(UBehaviorTreeComponent&
 
 void UBTDecorator_ConeCheck::OnBecomeRelevant(UBehaviorTreeComponent& OwnerComp, uint8* NodeMemory)
 {
-	TNodeInstanceMemory* DecoratorMemory = (TNodeInstanceMemory*)NodeMemory;
+	TNodeInstanceMemory* DecoratorMemory = CastInstanceNodeMemory<TNodeInstanceMemory>(NodeMemory);
 	DecoratorMemory->bLastRawResult = CalcConditionImpl(OwnerComp, NodeMemory);
 }
 
diff --git a/Engine/Source/Runtime/AIModule/Private/BehaviorTree/Decorators/BTDecorator_Cooldown.cpp b/Engine/Source/Runtime/AIModule/Private/BehaviorTree/Decorators/BTDecorator_Cooldown.cpp
--- a/Engine/Source/Runtime/AIModule/Private/BehaviorTree/Decorators/BTDecorator_Cooldown.cpp
+++ b/Engine/Source/Runtime/AIModule/Private/BehaviorTree/Decorators/BTDecorator_Cooldown.cpp
@@ -6,13 +6,11 @@
 UBTDecorator_Cooldown::UBTDecorator_Cooldown(const FObjectInitializer& ObjectInitializer) : Super(ObjectInitializer)
 {
 	NodeName = "Cooldown";
+	INIT_DECORATOR_NODE_NOTIFY_FLAGS();
 	CoolDownTime = 5.0f;
 	
 	// aborting child nodes doesn't makes sense, cooldown starts after leaving this branch
 	bAllowAbortChildNodes = false;
-
-	bNotifyTick = false;
-	bNotifyDeactivation = true;
 }
 
 void UBTDecorator_Cooldown::PostLoad()
diff --git a/Engine/Source/Runtime/AIModule/Private/BehaviorTree/Decorators/BTDecorator_ForceSuccess.cpp b/Engine/Source/Runtime/AIModule/Private/BehaviorTree/Decorators/BTDecorator_ForceSuccess.cpp
--- a/Engine/Source/Runtime/AIModule/Private/BehaviorTree/Decorators/BTDecorator_ForceSuccess.cpp
+++ b/Engine/Source/Runtime/AIModule/Private/BehaviorTree/Decorators/BTDecorator_ForceSuccess.cpp
@@ -7,7 +7,7 @@
 UBTDecorator_ForceSuccess::UBTDecorator_ForceSuccess(const FObjectInitializer& ObjectInitializer) : Super(ObjectInitializer)
 {
 	NodeName = TEXT("Force Success");
-	bNotifyProcessed = true;
+	INIT_DECORATOR_NODE_NOTIFY_FLAGS();
 
 	bAllowAbortNone = false;
 	bAllowAbortLowerPri = false;
diff --git a/Engine/Source/Runtime/AIModule/Private/BehaviorTree/Decorators/BTDecorator_KeepInCone.cpp b/Engine/Source/Runtime/AIModule/Private/BehaviorTree/Decorators/BTDecorator_KeepInCone.cpp
--- a/Engine/Source/Runtime/AIModule/Private/BehaviorTree/Decorators/BTDecorator_KeepInCone.cpp
+++ b/Engine/Source/Runtime/AIModule/Private/BehaviorTree/Decorators/BTDecorator_KeepInCone.cpp
@@ -14,8 +14,7 @@ UBTDecorator_KeepInCone::UBTDecorator_KeepInCone(const FObjectInitializer& Objec
 	Observed.AddObjectFilter(this, GET_MEMBER_NAME_CHECKED(UBTDecorator_KeepInCone, Observed), AActor::StaticClass());
 	Observed.AddVectorFilter(this, GET_MEMBER_NAME_CHECKED(UBTDecorator_KeepInCone, Observed));
 
-	bNotifyBecomeRelevant = true;
-	bNotifyTick = true;
+	INIT_DECORATOR_NODE_NOTIFY_FLAGS();
 
 	// KeepInCone always abort current branch
 	bAllowAbortLowerPri = false;
@@ -55,7 +54,7 @@ void UBTDecorator_KeepInCone::InitializeFromAsset(UBehaviorTree& Asset)
 bool UBTDecorator_KeepInCone::CalculateCurrentDirection(const UBehaviorTreeComponent& OwnerComp, FVector& Direction) const
 {
 	const UBlackboardComponent* BlackboardComp = OwnerComp.GetBlackboardComponent();
-	if (BlackboardComp == NULL)
+	if (BlackboardComp == nullptr)
 	{
 		return false;
 	}
@@ -85,7 +84,7 @@ void UBTDecorator_KeepInCone::OnBecomeRelevant(UBehaviorTreeComponent& OwnerComp
 
 void UBTDecorator_KeepInCone::TickNode(UBehaviorTreeComponent& OwnerComp, uint8* NodeMemory, float DeltaSeconds)
 {
-	TNodeInstanceMemory* DecoratorMemory = CastInstanceNodeMemory<TNodeInstanceMemory>(NodeMemory);
+	const TNodeInstanceMemory* DecoratorMemory = CastInstanceNodeMemory<TNodeInstanceMemory>(NodeMemory);
 	FVector CurrentDir(1.0f, 0, 0);
 	
 	if (CalculateCurrentDirection(OwnerComp, CurrentDir))
@@ -110,7 +109,7 @@ FString UBTDecorator_KeepInCone::GetStaticDescription() const
 
 void UBTDecorator_KeepInCone::DescribeRuntimeValues(const UBehaviorTreeComponent& OwnerComp, uint8* NodeMemory, EBTDescriptionVerbosity::Type Verbosity, TArray<FString>& Values) const
 {
-	TNodeInstanceMemory* DecoratorMemory = CastInstanceNodeMemory<TNodeInstanceMemory>(NodeMemory);
+	const TNodeInstanceMemory* DecoratorMemory = CastInstanceNodeMemory<TNodeInstanceMemory>(NodeMemory);
 	FVector CurrentDir(1.0f, 0, 0);
 	
 	if (CalculateCurrentDirection(OwnerComp, CurrentDir))
diff --git a/Engine/Source/Runtime/AIModule/Private/BehaviorTree/Decorators/BTDecorator_Loop.cpp b/Engine/Source/Runtime/AIModule/Private/BehaviorTree/Decorators/BTDecorator_Loop.cpp
--- a/Engine/Source/Runtime/AIModule/Private/BehaviorTree/Decorators/BTDecorator_Loop.cpp
+++ b/Engine/Source/Runtime/AIModule/Private/BehaviorTree/Decorators/BTDecorator_Loop.cpp
@@ -10,7 +10,7 @@ UBTDecorator_Loop::UBTDecorator_Loop(const FObjectInitializer& ObjectInitializer
 	NodeName = "Loop";
 	NumLoops = 3;
 	InfiniteLoopTimeoutTime = -1.f;
-	bNotifyActivation = true;
+	INIT_DECORATOR_NODE_NOTIFY_FLAGS();
 	
 	bAllowAbortNone = false;
 	bAllowAbortLowerPri = false;
diff --git a/Engine/Source/Runtime/AIModule/Private/BehaviorTree/Decorators/BTDecorator_SetTagCooldown.cpp b/Engine/Source/Runtime/AIModule/Private/BehaviorTree/Decorators/BTDecorator_SetTagCooldown.cpp
--- a/Engine/Source/Runtime/AIModule/Private/BehaviorTree/Decorators/BTDecorator_SetTagCooldown.cpp
+++ b/Engine/Source/Runtime/AIModule/Private/BehaviorTree/Decorators/BTDecorator_SetTagCooldown.cpp
@@ -5,13 +5,13 @@
 UBTDecorator_SetTagCooldown::UBTDecorator_SetTagCooldown(const FObjectInitializer& ObjectInitializer) : Super(ObjectInitializer)
 {
 	NodeName = "Set Tag Cooldown";
+	INIT_DECORATOR_NODE_NOTIFY_FLAGS();
 	CooldownDuration = 5.0f;
 	bAddToExistingDuration = false;
 
 	bAllowAbortNone = false;
 	bAllowAbortLowerPri = false;
 	bAllowAbortChildNodes = false;
-	bNotifyDeactivation = true;
 	FlowAbortMode = EBTFlowAbortMode::None;
 }
 
diff --git a/Engine/Source/Runtime/AIModule/Private/BehaviorTree/Decorators/BTDecorator_TagCooldown.cpp b/Engine/Source/Runtime/AIModule/Private/BehaviorTree/Decorators/BTDecorator_TagCooldown.cpp
--- a/Engine/Source/Runtime/AIModule/Private/BehaviorTree/Decorators/BTDecorator_TagCooldown.cpp
+++ b/Engine/Source/Runtime/AIModule/Private/BehaviorTree/Decorators/BTDecorator_TagCooldown.cpp
@@ -6,15 +6,13 @@
 UBTDecorator_TagCooldown::UBTDecorator_TagCooldown(const FObjectInitializer& ObjectInitializer) : Super(ObjectInitializer)
 {
 	NodeName = "Tag Cooldown";
+	INIT_DECORATOR_NODE_NOTIFY_FLAGS();
 	CooldownDuration = 5.0f;
 	bAddToExistingDuration = false;
 	bActivatesCooldown = true;
 	
 	// aborting child nodes doesn't makes sense, cooldown starts after leaving this branch
 	bAllowAbortChildNodes = false;
-
-	bNotifyTick = false;
-	bNotifyDeactivation = true;
 }
 
 void UBTDecorator_TagCooldown::PostLoad()
@@ -43,8 +41,7 @@ bool UBTDecorator_TagCooldown::CalculateRawConditionValue(UBehaviorTreeComponent
 
 void UBTDecorator_TagCooldown::InitializeMemory(UBehaviorTreeComponent& OwnerComp, uint8* NodeMemory, EBTMemoryInit::Type InitType) const
 {
-	FBTTagCooldownDecoratorMemory* DecoratorMemory = (FBTTagCooldownDecoratorMemory*)NodeMemory;
-
+	FBTTagCooldownDecoratorMemory* DecoratorMemory = CastInstanceNodeMemory<FBTTagCooldownDecoratorMemory>(NodeMemory);
 	DecoratorMemory->bRequestedRestart = false;
 }
 
@@ -61,7 +58,7 @@ void UBTDecorator_TagCooldown::OnNodeDeactivation(FBehaviorTreeSearchData& Searc
 
 void UBTDecorator_TagCooldown::TickNode(UBehaviorTreeComponent& OwnerComp, uint8* NodeMemory, float DeltaSeconds)
 {
-	FBTTagCooldownDecoratorMemory* DecoratorMemory = (FBTTagCooldownDecoratorMemory*)NodeMemory;
+	FBTTagCooldownDecoratorMemory* DecoratorMemory = CastInstanceNodeMemory<FBTTagCooldownDecoratorMemory>(NodeMemory);
 	if (!DecoratorMemory->bRequestedRestart)
 	{
 		if (HasCooldownFinished(OwnerComp))
diff --git a/Engine/Source/Runtime/AIModule/Private/BehaviorTree/Decorators/BTDecorator_TimeLimit.cpp b/Engine/Source/Runtime/AIModule/Private/BehaviorTree/Decorators/BTDecorator_TimeLimit.cpp
--- a/Engine/Source/Runtime/AIModule/Private/BehaviorTree/Decorators/BTDecorator_TimeLimit.cpp
+++ b/Engine/Source/Runtime/AIModule/Private/BehaviorTree/Decorators/BTDecorator_TimeLimit.cpp
@@ -6,8 +6,7 @@ UBTDecorator_TimeLimit::UBTDecorator_TimeLimit(const FObjectInitializer& ObjectI
 {
 	NodeName = "TimeLimit";
 	TimeLimit = 5.0f;
-	bNotifyActivation = true;
-	bNotifyTick = true;
+	INIT_DECORATOR_NODE_NOTIFY_FLAGS();
 	bTickIntervals = true;
 
 	// time limit always abort current branch
@@ -43,7 +42,7 @@ void UBTDecorator_TimeLimit::DescribeRuntimeValues(const UBehaviorTreeComponent&
 {
 	Super::DescribeRuntimeValues(OwnerComp, NodeMemory, Verbosity, Values);
 
-	FBTAuxiliaryMemory* DecoratorMemory = GetSpecialNodeMemory<FBTAuxiliaryMemory>(NodeMemory);
+	const FBTAuxiliaryMemory* DecoratorMemory = GetSpecialNodeMemory<FBTAuxiliaryMemory>(NodeMemory);
 	if (OwnerComp.GetWorld())
 	{
 		const float TimeLeft = DecoratorMemory->NextTickRemainingTime > 0 ? DecoratorMemory->NextTickRemainingTime : 0;
diff --git a/Engine/Source/Runtime/AIModule/Private/BehaviorTree/Services/BTService_DefaultFocus.cpp b/Engine/Source/Runtime/AIModule/Private/BehaviorTree/Services/BTService_DefaultFocus.cpp
--- a/Engine/Source/Runtime/AIModule/Private/BehaviorTree/Services/BTService_DefaultFocus.cpp
+++ b/Engine/Source/Runtime/AIModule/Private/BehaviorTree/Services/BTService_DefaultFocus.cpp
@@ -12,10 +12,8 @@ UBTService_DefaultFocus::UBTService_DefaultFocus(const FObjectInitializer& Objec
 {
 	NodeName = "Set default focus";
 
-	bNotifyTick = false;
 	bTickIntervals = false;
-	bNotifyBecomeRelevant = true;
-	bNotifyCeaseRelevant = true;
+	INIT_AUXILIARY_NODE_NOTIFY_FLAGS();
 
 	FocusPriority = EAIFocusPriority::Default;
 
@@ -28,14 +26,14 @@ void UBTService_DefaultFocus::OnBecomeRelevant(UBehaviorTreeComponent& OwnerComp
 {
 	Super::OnBecomeRelevant(OwnerComp, NodeMemory);
 
-	FBTFocusMemory* MyMemory = (FBTFocusMemory*)NodeMemory;
+	FBTFocusMemory* MyMemory = CastInstanceNodeMemory<FBTFocusMemory>(NodeMemory);
 	check(MyMemory);
 	MyMemory->Reset();
 
 	AAIController* OwnerController = OwnerComp.GetAIOwner();
 	UBlackboardComponent* MyBlackboard = OwnerComp.GetBlackboardComponent();
 	
-	if (OwnerController != NULL && MyBlackboard != NULL)
+	if (OwnerController != nullptr && MyBlackboard != nullptr)
 	{
 		if (BlackboardKey.SelectedKeyType == UBlackboardKeyType_Object::StaticClass())
 		{
@@ -55,7 +53,7 @@ void UBTService_DefaultFocus::OnBecomeRelevant(UBehaviorTreeComponent& OwnerComp
 			MyMemory->FocusLocationSet = FocusLocation;
 		}
 
-		auto KeyID = BlackboardKey.GetSelectedKeyID();
+		const FBlackboard::FKey KeyID = BlackboardKey.GetSelectedKeyID();
 		MyBlackboard->RegisterObserver(KeyID, this, FOnBlackboardChangeNotification::CreateUObject(this, &UBTService_DefaultFocus::OnBlackboardKeyValueChange));
 	}
 }
@@ -64,10 +62,10 @@ void UBTService_DefaultFocus::OnCeaseRelevant(UBehaviorTreeComponent& OwnerComp,
 {
 	Super::OnCeaseRelevant(OwnerComp, NodeMemory);
 
-	FBTFocusMemory* MyMemory = (FBTFocusMemory*)NodeMemory;
+	FBTFocusMemory* MyMemory = CastInstanceNodeMemory<FBTFocusMemory>(NodeMemory);
 	check(MyMemory);
 	AAIController* OwnerController = OwnerComp.GetAIOwner();
-	if (OwnerController != NULL)
+	if (OwnerController != nullptr)
 	{
 		bool bClearFocus = false;
 		if (MyMemory->bActorSet)
@@ -114,7 +112,7 @@ EBlackboardNotificationResult UBTService_DefaultFocus::OnBlackboardKeyValueChang
 	}
 
 	const int32 NodeInstanceIdx = OwnerComp->FindInstanceContainingNode(this);
-	FBTFocusMemory* MyMemory = (FBTFocusMemory*)OwnerComp->GetNodeMemory(this, NodeInstanceIdx);
+	FBTFocusMemory* MyMemory = CastInstanceNodeMemory<FBTFocusMemory>(OwnerComp->GetNodeMemory(this, NodeInstanceIdx));
 	MyMemory->Reset();
 	OwnerController->ClearFocus(FocusPriority);
 
diff --git a/Engine/Source/Runtime/AIModule/Private/BehaviorTree/Tasks/BTTask_BlueprintBase.cpp b/Engine/Source/Runtime/AIModule/Private/BehaviorTree/Tasks/BTTask_BlueprintBase.cpp
--- a/Engine/Source/Runtime/AIModule/Private/BehaviorTree/Tasks/BTTask_BlueprintBase.cpp
+++ b/Engine/Source/Runtime/AIModule/Private/BehaviorTree/Tasks/BTTask_BlueprintBase.cpp
@@ -8,13 +8,13 @@
 
 UBTTask_BlueprintBase::UBTTask_BlueprintBase(const FObjectInitializer& ObjectInitializer) : Super(ObjectInitializer)
 {
-	UClass* StopAtClass = UBTTask_BlueprintBase::StaticClass();
+	const UClass* StopAtClass = UBTTask_BlueprintBase::StaticClass();
 	ReceiveTickImplementations = FBTNodeBPImplementationHelper::CheckEventImplementationVersion(TEXT("ReceiveTick"), TEXT("ReceiveTickAI"), *this, *StopAtClass);
 	ReceiveExecuteImplementations = FBTNodeBPImplementationHelper::CheckEventImplementationVersion(TEXT("ReceiveExecute"), TEXT("ReceiveExecuteAI"), *this, *StopAtClass);
 	ReceiveAbortImplementations = FBTNodeBPImplementationHelper::CheckEventImplementationVersion(TEXT("ReceiveAbort"), TEXT("ReceiveAbortAI"), *this, *StopAtClass);
 
+	INIT_TASK_NODE_NOTIFY_FLAGS();
 	bNotifyTick = ReceiveTickImplementations != FBTNodeBPImplementationHelper::NoImplementation;
-	bNotifyTaskFinished = true;
 	bShowPropertyDetails = true;
 	bIsAborting = false;
 
@@ -129,7 +129,7 @@ void UBTTask_BlueprintBase::OnTaskFinished(UBehaviorTreeComponent& OwnerComp, ui
 void UBTTask_BlueprintBase::FinishExecute(bool bSuccess)
 {
 	UBehaviorTreeComponent* OwnerComp = Cast<UBehaviorTreeComponent>(GetOuter());
-	EBTNodeResult::Type NodeResult(bSuccess ? EBTNodeResult::Succeeded : EBTNodeResult::Failed);
+	const EBTNodeResult::Type NodeResult(bSuccess ? EBTNodeResult::Succeeded : EBTNodeResult::Failed);
 
 	if (bStoreFinishResult)
 	{
@@ -144,10 +144,10 @@ void UBTTask_BlueprintBase::FinishExecute(bool bSuccess)
 void UBTTask_BlueprintBase::FinishAbort()
 {
 	UBehaviorTreeComponent* OwnerComp = Cast<UBehaviorTreeComponent>(GetOuter());
-	EBTNodeResult::Type NodeResult(EBTNodeResult::Aborted);
 
 	if (bStoreFinishResult)
 	{
+		constexpr EBTNodeResult::Type NodeResult(EBTNodeResult::Aborted);
 		CurrentCallResult = NodeResult;
 	}
 	else if (OwnerComp && bIsAborting)
@@ -158,8 +158,8 @@ void UBTTask_BlueprintBase::FinishAbort()
 
 bool UBTTask_BlueprintBase::IsTaskExecuting() const
 {
-	UBehaviorTreeComponent* OwnerComp = Cast<UBehaviorTreeComponent>(GetOuter());
-	EBTTaskStatus::Type TaskStatus = OwnerComp->GetTaskStatus(this);
+	const UBehaviorTreeComponent* OwnerComp = Cast<UBehaviorTreeComponent>(GetOuter());
+	const EBTTaskStatus::Type TaskStatus = OwnerComp->GetTaskStatus(this);
 
 	return (TaskStatus == EBTTaskStatus::Active);
 }
@@ -196,10 +196,10 @@ FString UBTTask_BlueprintBase::GetStaticDescription() const
 #endif // WITH_EDITORONLY_DATA
 		Super::GetStaticDescription();
 
-	UBTTask_BlueprintBase* CDO = (UBTTask_BlueprintBase*)(GetClass()->GetDefaultObject());
+	UBTTask_BlueprintBase* CDO = static_cast<UBTTask_BlueprintBase*>(GetClass()->GetDefaultObject());
 	if (bShowPropertyDetails && CDO)
 	{
-		UClass* StopAtClass = UBTTask_BlueprintBase::StaticClass();
+		const UClass* StopAtClass = UBTTask_BlueprintBase::StaticClass();
 		FString PropertyDesc = BlueprintNodeHelpers::CollectPropertyDescription(this, StopAtClass, CDO->PropertyData);
 		if (PropertyDesc.Len())
 		{
@@ -213,7 +213,7 @@ FString UBTTask_BlueprintBase::GetStaticDescription() const
 
 void UBTTask_BlueprintBase::DescribeRuntimeValues(const UBehaviorTreeComponent& OwnerComp, uint8* NodeMemory, EBTDescriptionVerbosity::Type Verbosity, TArray<FString>& Values) const
 {
-	UBTTask_BlueprintBase* CDO = (UBTTask_BlueprintBase*)(GetClass()->GetDefaultObject());
+	UBTTask_BlueprintBase* CDO = static_cast<UBTTask_BlueprintBase*>(GetClass()->GetDefaultObject());
 	if (CDO && CDO->PropertyData.Num())
 	{
 		BlueprintNodeHelpers::DescribeRuntimeValues(this, CDO->PropertyData, Values);
diff --git a/Engine/Source/Runtime/AIModule/Private/BehaviorTree/Tasks/BTTask_GameplayTaskBase.cpp b/Engine/Source/Runtime/AIModule/Private/BehaviorTree/Tasks/BTTask_GameplayTaskBase.cpp
--- a/Engine/Source/Runtime/AIModule/Private/BehaviorTree/Tasks/BTTask_GameplayTaskBase.cpp
+++ b/Engine/Source/Runtime/AIModule/Private/BehaviorTree/Tasks/BTTask_GameplayTaskBase.cpp
@@ -7,8 +7,8 @@
 UBTTask_GameplayTaskBase::UBTTask_GameplayTaskBase(const FObjectInitializer& ObjectInitializer) : Super(ObjectInitializer)
 {
 	NodeName = "GameplayTask Base";
+	INIT_TASK_NODE_NOTIFY_FLAGS();
 	bWaitForGameplayTask = true;
-	bNotifyTaskFinished = true;
 }
 
 EBTNodeResult::Type UBTTask_GameplayTaskBase::StartGameplayTask(UBehaviorTreeComponent& OwnerComp, uint8* NodeMemory, UAITask& Task)
@@ -70,7 +70,7 @@ void UBTTask_GameplayTaskBase::OnGameplayTaskDeactivated(UGameplayTask& Task)
 		if (BehaviorComp)
 		{
 			uint8* RawMemory = BehaviorComp->GetNodeMemory(this, BehaviorComp->FindInstanceContainingNode(this));
-			FBTGameplayTaskMemory* MyMemory = CastInstanceNodeMemory<FBTGameplayTaskMemory>(RawMemory);
+			const FBTGameplayTaskMemory* MyMemory = CastInstanceNodeMemory<FBTGameplayTaskMemory>(RawMemory);
 
 			if (MyMemory->bObserverCanFinishTask && (AITask == MyMemory->Task))
 			{
diff --git a/Engine/Source/Runtime/AIModule/Private/BehaviorTree/Tasks/BTTask_MoveTo.cpp b/Engine/Source/Runtime/AIModule/Private/BehaviorTree/Tasks/BTTask_MoveTo.cpp
--- a/Engine/Source/Runtime/AIModule/Private/BehaviorTree/Tasks/BTTask_MoveTo.cpp
+++ b/Engine/Source/Runtime/AIModule/Private/BehaviorTree/Tasks/BTTask_MoveTo.cpp
@@ -14,9 +14,9 @@
 UBTTask_MoveTo::UBTTask_MoveTo(const FObjectInitializer& ObjectInitializer) : Super(ObjectInitializer)
 {
 	NodeName = "Move To";
+	INIT_TASK_NODE_NOTIFY_FLAGS();
 	bUseGameplayTasks = GET_AI_CONFIG_VAR(bEnableBTAITasks);
 	bNotifyTick = !bUseGameplayTasks;
-	bNotifyTaskFinished = true;
 
 	AcceptableRadius = GET_AI_CONFIG_VAR(AcceptanceRadius);
 	bReachTestIncludesGoalRadius = bReachTestIncludesAgentRadius = bStopOnOverlap = GET_AI_CONFIG_VAR(bFinishMoveOnGoalOverlap);
@@ -316,7 +316,7 @@ void UBTTask_MoveTo::OnTaskFinished(UBehaviorTreeComponent& OwnerComp, uint8* No
 
 void UBTTask_MoveTo::TickTask(UBehaviorTreeComponent& OwnerComp, uint8* NodeMemory, float DeltaSeconds)
 {
-	FBTMoveToTaskMemory* MyMemory = (FBTMoveToTaskMemory*)NodeMemory;
+	FBTMoveToTaskMemory* MyMemory = CastInstanceNodeMemory<FBTMoveToTaskMemory>(NodeMemory);
 	if (MyMemory->bWaitingForPath && !OwnerComp.IsPaused())
 	{
 		AAIController* MyController = OwnerComp.GetAIOwner();
@@ -351,7 +351,7 @@ void UBTTask_MoveTo::OnGameplayTaskDeactivated(UGameplayTask& Task)
 		if (BehaviorComp)
 		{
 			uint8* RawMemory = BehaviorComp->GetNodeMemory(this, BehaviorComp->FindInstanceContainingNode(this));
-			FBTMoveToTaskMemory* MyMemory = CastInstanceNodeMemory<FBTMoveToTaskMemory>(RawMemory);
+			const FBTMoveToTaskMemory* MyMemory = CastInstanceNodeMemory<FBTMoveToTaskMemory>(RawMemory);
 
 			if (MyMemory && MyMemory->bObserverCanFinishTask && (MoveTask == MyMemory->Task))
 			{
@@ -384,7 +384,7 @@ void UBTTask_MoveTo::DescribeRuntimeValues(const UBehaviorTreeComponent& OwnerCo
 	{
 		const FString KeyValue = BlackboardComp->DescribeKeyValue(BlackboardKey.GetSelectedKeyID(), EBlackboardDescription::OnlyValue);
 
-		const FBTMoveToTaskMemory* MyMemory = (FBTMoveToTaskMemory*)NodeMemory;
+		const FBTMoveToTaskMemory* MyMemory = CastInstanceNodeMemory<FBTMoveToTaskMemory>(NodeMemory);
 		const bool bIsUsingTask = MyMemory->Task.IsValid();
 		
 		const FString ModeDesc =
diff --git a/Engine/Source/Runtime/AIModule/Private/BehaviorTree/Tasks/BTTask_RotateToFaceBBEntry.cpp b/Engine/Source/Runtime/AIModule/Private/BehaviorTree/Tasks/BTTask_RotateToFaceBBEntry.cpp
--- a/Engine/Source/Runtime/AIModule/Private/BehaviorTree/Tasks/BTTask_RotateToFaceBBEntry.cpp
+++ b/Engine/Source/Runtime/AIModule/Private/BehaviorTree/Tasks/BTTask_RotateToFaceBBEntry.cpp
@@ -13,7 +13,7 @@ UBTTask_RotateToFaceBBEntry::UBTTask_RotateToFaceBBEntry(const FObjectInitialize
 	, Precision(10.f)
 {
 	NodeName = "Rotate to face BB entry";
-	bNotifyTick = true;
+	INIT_TASK_NODE_NOTIFY_FLAGS();
 	
 	// accept only actors and vectors
 	BlackboardKey.AddObjectFilter(this, GET_MEMBER_NAME_CHECKED(UBTTask_RotateToFaceBBEntry, BlackboardKey), AActor::StaticClass());
diff --git a/Engine/Source/Runtime/AIModule/Private/BehaviorTree/Tasks/BTTask_Wait.cpp b/Engine/Source/Runtime/AIModule/Private/BehaviorTree/Tasks/BTTask_Wait.cpp
--- a/Engine/Source/Runtime/AIModule/Private/BehaviorTree/Tasks/BTTask_Wait.cpp
+++ b/Engine/Source/Runtime/AIModule/Private/BehaviorTree/Tasks/BTTask_Wait.cpp
@@ -6,12 +6,12 @@ UBTTask_Wait::UBTTask_Wait(const FObjectInitializer& ObjectInitializer) : Super(
 {
 	NodeName = "Wait";
 	WaitTime = 5.0f;
-	bNotifyTick = true;
+	INIT_TASK_NODE_NOTIFY_FLAGS();
 }
 
 EBTNodeResult::Type UBTTask_Wait::ExecuteTask(UBehaviorTreeComponent& OwnerComp, uint8* NodeMemory)
 {
-	FBTWaitTaskMemory* MyMemory = (FBTWaitTaskMemory*)NodeMemory;
+	FBTWaitTaskMemory* MyMemory = CastInstanceNodeMemory<FBTWaitTaskMemory>(NodeMemory);
 	MyMemory->RemainingWaitTime = FMath::FRandRange(FMath::Max(0.0f, WaitTime - RandomDeviation), (WaitTime + RandomDeviation));
 	
 	return EBTNodeResult::InProgress;
@@ -19,7 +19,7 @@ EBTNodeResult::Type UBTTask_Wait::ExecuteTask(UBehaviorTreeComponent& OwnerComp,
 
 void UBTTask_Wait::TickTask(UBehaviorTreeComponent& OwnerComp, uint8* NodeMemory, float DeltaSeconds)
 {
-	FBTWaitTaskMemory* MyMemory = (FBTWaitTaskMemory*)NodeMemory;
+	FBTWaitTaskMemory* MyMemory = CastInstanceNodeMemory<FBTWaitTaskMemory>(NodeMemory);
 	MyMemory->RemainingWaitTime -= DeltaSeconds;
 
 	if (MyMemory->RemainingWaitTime <= 0.0f)
@@ -45,7 +45,7 @@ void UBTTask_Wait::DescribeRuntimeValues(const UBehaviorTreeComponent& OwnerComp
 {
 	Super::DescribeRuntimeValues(OwnerComp, NodeMemory, Verbosity, Values);
 
-	FBTWaitTaskMemory* MyMemory = (FBTWaitTaskMemory*)NodeMemory;
+	FBTWaitTaskMemory* MyMemory = CastInstanceNodeMemory<FBTWaitTaskMemory>(NodeMemory);
 	if (MyMemory->RemainingWaitTime)
 	{
 		Values.Add(FString::Printf(TEXT("remaining: %ss"), *FString::SanitizeFloat(MyMemory->RemainingWaitTime)));
