Deprecation of bEnableBTAITasks
Move task always use gameplay tasks
#tests PIE

diff --git a/Engine/Source/Runtime/AIModule/Classes/AISystem.h b/Engine/Source/Runtime/AIModule/Classes/AISystem.h
--- a/Engine/Source/Runtime/AIModule/Classes/AISystem.h
+++ b/Engine/Source/Runtime/AIModule/Classes/AISystem.h
@@ -68,10 +68,9 @@ public:
 	bool bAllowStrafing;
 
 	/** 
-	 * Whether or not to enable Gameplay Tasks for move tasks
-     * this property is just a transition-time flag - in the end we're going to switch over to Gameplay Tasks anyway, that's the goal. 
+	 * Deprecated: Whether or not to enable Gameplay Tasks for move tasks (always enabled now)
 	 */
-	UPROPERTY(globalconfig, EditDefaultsOnly, Category = "Gameplay Tasks")
+	UPROPERTY(globalconfig, EditDefaultsOnly, Category = "Gameplay Tasks", meta=(DisplayName="DEPRECATED Enable BT AITasks"))
 	bool bEnableBTAITasks;
 
 	/** if enable will make EQS not complaint about using Controllers as queriers. Default behavior (false) will 
diff --git a/Engine/Source/Runtime/AIModule/Classes/BehaviorTree/Tasks/BTTask_MoveTo.h b/Engine/Source/Runtime/AIModule/Classes/BehaviorTree/Tasks/BTTask_MoveTo.h
--- a/Engine/Source/Runtime/AIModule/Classes/BehaviorTree/Tasks/BTTask_MoveTo.h
+++ b/Engine/Source/Runtime/AIModule/Classes/BehaviorTree/Tasks/BTTask_MoveTo.h
@@ -24,7 +24,6 @@ struct FBTMoveToTaskMemory
 
 	TWeakObjectPtr<UAITask_MoveTo> Task;
 
-	uint8 bWaitingForPath : 1;
 	uint8 bObserverCanFinishTask : 1;
 };
 
@@ -88,13 +87,9 @@ class AIMODULE_API UBTTask_MoveTo : public UBTTask_BlackboardBase
 	/** if set, move will use pathfinding. Not exposed on purpose, please use BTTask_MoveDirectlyToward */
 	uint32 bUsePathfinding : 1;
 
-	/** set automatically if move should use GameplayTasks */
-	uint32 bUseGameplayTasks : 1;
-
 	virtual EBTNodeResult::Type ExecuteTask(UBehaviorTreeComponent& OwnerComp, uint8* NodeMemory) override;
 	virtual EBTNodeResult::Type AbortTask(UBehaviorTreeComponent& OwnerComp, uint8* NodeMemory) override;
 	virtual void OnTaskFinished(UBehaviorTreeComponent& OwnerComp, uint8* NodeMemory, EBTNodeResult::Type TaskResult) override;
-	virtual void TickTask(UBehaviorTreeComponent& OwnerComp, uint8* NodeMemory, float DeltaSeconds) override;
 	virtual uint16 GetInstanceMemorySize() const override;
 	virtual void PostLoad() override;
 
diff --git a/Engine/Source/Runtime/AIModule/Private/AISystem.cpp b/Engine/Source/Runtime/AIModule/Private/AISystem.cpp
--- a/Engine/Source/Runtime/AIModule/Private/AISystem.cpp
+++ b/Engine/Source/Runtime/AIModule/Private/AISystem.cpp
@@ -27,7 +27,7 @@ UAISystem::UAISystem(const FObjectInitializer& ObjectInitializer)
 	bAllowStrafing = false;
 	DefaultSightCollisionChannel = ECC_Visibility;
 
-	bEnableBTAITasks = false;
+	bEnableBTAITasks = true;
 
 	if (HasAnyFlags(RF_ClassDefaultObject))
 	{
diff --git a/Engine/Source/Runtime/AIModule/Private/BehaviorTree/Tasks/BTTask_MoveTo.cpp b/Engine/Source/Runtime/AIModule/Private/BehaviorTree/Tasks/BTTask_MoveTo.cpp
--- a/Engine/Source/Runtime/AIModule/Private/BehaviorTree/Tasks/BTTask_MoveTo.cpp
+++ b/Engine/Source/Runtime/AIModule/Private/BehaviorTree/Tasks/BTTask_MoveTo.cpp
@@ -15,8 +15,6 @@ UBTTask_MoveTo::UBTTask_MoveTo(const FObjectInitializer& ObjectInitializer) : Su
 {
 	NodeName = "Move To";
 	INIT_TASK_NODE_NOTIFY_FLAGS();
-	bUseGameplayTasks = GET_AI_CONFIG_VAR(bEnableBTAITasks);
-	bNotifyTick = !bUseGameplayTasks;
 
 	AcceptableRadius = GET_AI_CONFIG_VAR(AcceptanceRadius);
 	bReachTestIncludesGoalRadius = bReachTestIncludesAgentRadius = bStopOnOverlap = GET_AI_CONFIG_VAR(bFinishMoveOnGoalOverlap);
@@ -51,15 +49,7 @@ EBTNodeResult::Type UBTTask_MoveTo::ExecuteTask(UBehaviorTreeComponent& OwnerCom
 	}
 	else
 	{
-		MyMemory->bWaitingForPath = bUseGameplayTasks ? false : MyController->ShouldPostponePathUpdates();
-		if (!MyMemory->bWaitingForPath)
-		{
-			NodeResult = PerformMoveTask(OwnerComp, NodeMemory);
-		}
-		else
-		{
-			UE_VLOG(MyController, LogBehaviorTree, Log, TEXT("Pathfinding requests are freezed, waiting..."));
-		}
+		NodeResult = PerformMoveTask(OwnerComp, NodeMemory);
 	}
 
 	if (NodeResult == EBTNodeResult::InProgress && bObserveBlackboardValue)
@@ -220,9 +210,8 @@ EBlackboardNotificationResult UBTTask_MoveTo::OnBlackboardValueChange(const UBla
 		return EBlackboardNotificationResult::RemoveObserver;
 	}
 	
-	// this means the move has already started. MyMemory->bWaitingForPath == true would mean we're waiting for right moment to start it anyway,
-	// so we don't need to do anything due to BB value change 
-	if (MyMemory != nullptr && MyMemory->bWaitingForPath == false && BehaviorComp->GetAIOwner() != nullptr)
+	// this means the move has already started. 
+	if (MyMemory != nullptr && BehaviorComp->GetAIOwner() != nullptr)
 	{
 		check(BehaviorComp->GetAIOwner()->GetPathFollowingComponent());
 
@@ -245,18 +234,10 @@ EBlackboardNotificationResult UBTTask_MoveTo::OnBlackboardValueChange(const UBla
 				BehaviorComp->GetAIOwner()->GetPathFollowingComponent()->AbortMove(*this, FPathFollowingResultFlags::NewRequest, MyMemory->MoveRequestID, EPathFollowingVelocityMode::Keep);
 			}
 
-			if (!bUseGameplayTasks && BehaviorComp->GetAIOwner()->ShouldPostponePathUpdates())
+			const EBTNodeResult::Type NodeResult = PerformMoveTask(*BehaviorComp, RawMemory);
+			if (NodeResult != EBTNodeResult::InProgress)
 			{
-				// NodeTick will take care of requesting move
-				MyMemory->bWaitingForPath = true;
-			}
-			else
-			{
-				const EBTNodeResult::Type NodeResult = PerformMoveTask(*BehaviorComp, RawMemory);
-				if (NodeResult != EBTNodeResult::InProgress)
-				{
-					FinishLatentTask(*BehaviorComp, NodeResult);
-				}
+				FinishLatentTask(*BehaviorComp, NodeResult);
 			}
 		}
 	}
@@ -267,28 +248,21 @@ EBlackboardNotificationResult UBTTask_MoveTo::OnBlackboardValueChange(const UBla
 EBTNodeResult::Type UBTTask_MoveTo::AbortTask(UBehaviorTreeComponent& OwnerComp, uint8* NodeMemory)
 {
 	FBTMoveToTaskMemory* MyMemory = CastInstanceNodeMemory<FBTMoveToTaskMemory>(NodeMemory);
-	if (!MyMemory->bWaitingForPath)
+	if (MyMemory->MoveRequestID.IsValid())
 	{
-		if (MyMemory->MoveRequestID.IsValid())
+		AAIController* MyController = OwnerComp.GetAIOwner();
+		if (MyController && MyController->GetPathFollowingComponent())
 		{
-			AAIController* MyController = OwnerComp.GetAIOwner();
-			if (MyController && MyController->GetPathFollowingComponent())
-			{
-				MyController->GetPathFollowingComponent()->AbortMove(*this, FPathFollowingResultFlags::OwnerFinished, MyMemory->MoveRequestID);
-			}
+			MyController->GetPathFollowingComponent()->AbortMove(*this, FPathFollowingResultFlags::OwnerFinished, MyMemory->MoveRequestID);
 		}
-		else
+	}
+	else
+	{
+		MyMemory->bObserverCanFinishTask = false;
+		UAITask_MoveTo* MoveTask = MyMemory->Task.Get();
+		if (MoveTask)
 		{
-			MyMemory->bObserverCanFinishTask = false;
-			UAITask_MoveTo* MoveTask = MyMemory->Task.Get();
-			if (MoveTask)
-			{
-				MoveTask->ExternalCancel();
-			}
-			else
-			{
-				UE_VLOG(OwnerComp.GetAIOwner(), LogBehaviorTree, Error, TEXT("Can't abort path following! bWaitingForPath:false, MoveRequestID:invalid, MoveTask:none!"));
-			}
+			MoveTask->ExternalCancel();
 		}
 	}
 
@@ -314,26 +288,6 @@ void UBTTask_MoveTo::OnTaskFinished(UBehaviorTreeComponent& OwnerComp, uint8* No
 	Super::OnTaskFinished(OwnerComp, NodeMemory, TaskResult);
 }
 
-void UBTTask_MoveTo::TickTask(UBehaviorTreeComponent& OwnerComp, uint8* NodeMemory, float DeltaSeconds)
-{
-	FBTMoveToTaskMemory* MyMemory = CastInstanceNodeMemory<FBTMoveToTaskMemory>(NodeMemory);
-	if (MyMemory->bWaitingForPath && !OwnerComp.IsPaused())
-	{
-		AAIController* MyController = OwnerComp.GetAIOwner();
-		if (MyController && !MyController->ShouldPostponePathUpdates())
-		{
-			UE_VLOG(MyController, LogBehaviorTree, Log, TEXT("Pathfinding requests are unlocked!"));
-			MyMemory->bWaitingForPath = false;
-
-			const EBTNodeResult::Type NodeResult = PerformMoveTask(OwnerComp, NodeMemory);
-			if (NodeResult != EBTNodeResult::InProgress)
-			{
-				FinishLatentTask(OwnerComp, NodeResult);
-			}
-		}
-	}
-}
-
 void UBTTask_MoveTo::OnMessage(UBehaviorTreeComponent& OwnerComp, uint8* NodeMemory, FName Message, int32 SenderID, bool bSuccess)
 {
 	// AIMessage_RepathFailed means task has failed
@@ -387,10 +341,7 @@ void UBTTask_MoveTo::DescribeRuntimeValues(const UBehaviorTreeComponent& OwnerCo
 		const FBTMoveToTaskMemory* MyMemory = CastInstanceNodeMemory<FBTMoveToTaskMemory>(NodeMemory);
 		const bool bIsUsingTask = MyMemory->Task.IsValid();
 		
-		const FString ModeDesc =
-			MyMemory->bWaitingForPath ? TEXT("(WAITING)") :
-			bIsUsingTask ? TEXT("(task)") :
-			TEXT("");
+		const FString ModeDesc = bIsUsingTask ? TEXT("(task)") : TEXT("");
 
 		Values.Add(FString::Printf(TEXT("move target: %s%s"), *KeyValue, *ModeDesc));
 	}
