Expose the subtree in RunBehaviorDynamic through GetSubtreeAsset

diff --git a/Engine/Source/Runtime/AIModule/Classes/BehaviorTree/Tasks/BTTask_RunBehaviorDynamic.h b/Engine/Source/Runtime/AIModule/Classes/BehaviorTree/Tasks/BTTask_RunBehaviorDynamic.h
--- a/Engine/Source/Runtime/AIModule/Classes/BehaviorTree/Tasks/BTTask_RunBehaviorDynamic.h
+++ b/Engine/Source/Runtime/AIModule/Classes/BehaviorTree/Tasks/BTTask_RunBehaviorDynamic.h
@@ -33,6 +33,9 @@ class AIMODULE_API UBTTask_RunBehaviorDynamic : public UBTTaskNode
 
 	bool HasMatchingTag(const FGameplayTag& Tag) const;
 	bool SetBehaviorAsset(UBehaviorTree* NewBehaviorAsset);
+	
+	/** @returns subtree asset */
+	UBehaviorTree* GetSubtreeAsset() const;
 
 protected:
 
@@ -60,3 +63,8 @@ FORCEINLINE bool UBTTask_RunBehaviorDynamic::HasMatchingTag(const FGameplayTag&
 	return InjectionTag == Tag;
 }
 
+FORCEINLINE UBehaviorTree* UBTTask_RunBehaviorDynamic::GetSubtreeAsset() const
+{
+	return BehaviorAsset;
+}
+
