Clearing a reference to PathFollowingComponent on MovementComponent when the PathFollowingComponent changes its MovementComponent to prevent UObject leaks

diff --git a/Engine/Source/Runtime/AIModule/Private/Navigation/PathFollowingComponent.cpp b/Engine/Source/Runtime/AIModule/Private/Navigation/PathFollowingComponent.cpp
--- a/Engine/Source/Runtime/AIModule/Private/Navigation/PathFollowingComponent.cpp
+++ b/Engine/Source/Runtime/AIModule/Private/Navigation/PathFollowingComponent.cpp
@@ -659,6 +659,11 @@ void UPathFollowingComponent::TickComponent(float DeltaTime, enum ELevelTick Tic
 
 void UPathFollowingComponent::SetMovementComponent(UNavMovementComponent* MoveComp)
 {
+	if (!MoveComp && MovementComp && MovementComp->GetPathFollowingAgent() == this)
+	{
+		MovementComp->SetPathFollowingAgent(nullptr);
+	}
+
 	MovementComp = MoveComp;
 	MyNavData = nullptr;
 
