Exposed RunNavProjection

diff --git a/Engine/Source/Runtime/AIModule/Private/EnvironmentQuery/EnvQueryTraceHelpers.h b/Engine/Source/Runtime/AIModule/Private/EnvironmentQuery/EnvQueryTraceHelpers.h
--- a/Engine/Source/Runtime/AIModule/Private/EnvironmentQuery/EnvQueryTraceHelpers.h
+++ b/Engine/Source/Runtime/AIModule/Private/EnvironmentQuery/EnvQueryTraceHelpers.h
@@ -107,7 +107,7 @@ namespace FEQSHelpers
 	};
 
 	void RunNavRaycasts(const ANavigationData& NavData, const UObject& Querier, const FEnvTraceData& TraceData, const FVector& SourcePt, TArray<FNavLocation>& Points, const ETraceMode TraceMode = ETraceMode::Keep);
-	void RunNavProjection(const ANavigationData& NavData, const UObject& Querier, const FEnvTraceData& TraceData, TArray<FNavLocation>& Points, const ETraceMode TraceMode = ETraceMode::Discard);
+	AIMODULE_API void RunNavProjection(const ANavigationData& NavData, const UObject& Querier, const FEnvTraceData& TraceData, TArray<FNavLocation>& Points, const ETraceMode TraceMode = ETraceMode::Discard);
 	void RunPhysRaycasts(UWorld* World, const FEnvTraceData& TraceData, const FVector& SourcePt, TArray<FNavLocation>& Points, const TArray<AActor*>& IgnoredActors, const ETraceMode TraceMode = ETraceMode::Keep);
 	void RunPhysProjection(UWorld* World, const FEnvTraceData& TraceData, TArray<FNavLocation>& Points, const ETraceMode TraceMode = ETraceMode::Discard);
 	void RunPhysProjection(UWorld* World, const FEnvTraceData& TraceData, TArray<FNavLocation>& Points, TArray<uint8>& TraceHits);
