CrowdManager - Updating the crowd manager code to not store what it thinks is the game world in postinitproperties.  It's causing a leak between worlds when it grabs the old world when we transition to a new map.  Replaced existing code with a function to get the right world when needed.  It's still not 'good' since this wont work in multi-user PIE, but it wont crash either.  Debug drawing should be re-evaluated to work correctly with multi client PIE.

diff --git a/Engine/Source/Runtime/AIModule/Classes/Navigation/CrowdManager.h b/Engine/Source/Runtime/AIModule/Classes/Navigation/CrowdManager.h
--- a/Engine/Source/Runtime/AIModule/Classes/Navigation/CrowdManager.h
+++ b/Engine/Source/Runtime/AIModule/Classes/Navigation/CrowdManager.h
@@ -173,7 +173,6 @@ class AIMODULE_API UCrowdManager : public UCrowdManagerBase
 public:
 	UCrowdManager(const FObjectInitializer& ObjectInitializer = FObjectInitializer::Get());
 
-	virtual void PostInitProperties() override;
 	virtual void Tick(float DeltaTime) override;
 	virtual void BeginDestroy() override;
 
@@ -340,14 +339,6 @@ protected:
 	FCrowdTickHelper* TickHelper;
 #endif // WITH_EDITOR
 
-#if WITH_EDITORONLY_DATA
-	/** UWorld that's going to be used for debug drawing calls. DebugDrawingWorld should be different than GetWorld()
-	 *	only on dedicated-server.
-	 *	Note that this is only relevant when ENABLE_DRAW_DEBUG is enabled*/
-	UPROPERTY(transient)
-	TObjectPtr<UWorld> DebugDrawingWorld;
-#endif // WITH_EDITORONLY_DATA
-
 	/** try to initialize nav data from already existing ones */
 	virtual void UpdateNavData();
 
@@ -386,6 +377,7 @@ protected:
 	void DestroyCrowdManager();
 
 #if ENABLE_DRAW_DEBUG
+	UWorld* GetDebugDrawingWorld() const;
 	void DrawDebugCorners(const dtCrowdAgent* CrowdAgent) const;
 	void DrawDebugCollisionSegments(const dtCrowdAgent* CrowdAgent) const;
 	void DrawDebugPath(const dtCrowdAgent* CrowdAgent) const;
diff --git a/Engine/Source/Runtime/AIModule/Private/Navigation/CrowdManager.cpp b/Engine/Source/Runtime/AIModule/Private/Navigation/CrowdManager.cpp
--- a/Engine/Source/Runtime/AIModule/Private/Navigation/CrowdManager.cpp
+++ b/Engine/Source/Runtime/AIModule/Private/Navigation/CrowdManager.cpp
@@ -205,32 +205,6 @@ UCrowdManager::UCrowdManager(const FObjectInitializer& ObjectInitializer) : Supe
 	AvoidanceConfig.Add(AvoidanceConfig66);
 }
 
-void UCrowdManager::PostInitProperties() 
-{
-	Super::PostInitProperties();
-
-#if WITH_EDITORONLY_DATA
-	if (HasAnyFlags(RF_ClassDefaultObject) == false)
-	{
-		DebugDrawingWorld = GetWorld();
-		// note that being ENetMode::NM_DedicatedServer implies DebugDrawingWorld is a game world, which is exactly what we need
-		if (DebugDrawingWorld != nullptr && DebugDrawingWorld->GetNetMode() == ENetMode::NM_DedicatedServer)
-		{
-			// no point in trying to draw on dedicated server. Let's see if there's a client world we can use for drawing!
-			const TIndirectArray<FWorldContext>& WorldContexts = GEngine->GetWorldContexts();
-			for (const FWorldContext& Context : WorldContexts)
-			{
-				if (Context.World()->IsGameWorld() && Context.World()->GetNetMode() != ENetMode::NM_DedicatedServer)
-				{
-					DebugDrawingWorld = Context.World();
-					break;
-				}
-			}
-		}
-	}
-#endif // WITH_EDITORONLY_DATA && ENABLE_DRAW_DEBUG
-}
-
 void UCrowdManager::BeginDestroy()
 {
 #if WITH_RECAST
@@ -974,11 +948,33 @@ void UCrowdManager::DestroyCrowdManager()
 }
 
 #if ENABLE_DRAW_DEBUG
-void UCrowdManager::DrawDebugCorners(const dtCrowdAgent* CrowdAgent) const
+UWorld* UCrowdManager::GetDebugDrawingWorld() const
 {
-#if !WITH_EDITORONLY_DATA
 	UWorld* DebugDrawingWorld = GetWorld();
-#endif // !WITH_EDITORONLY_DATA
+
+#if WITH_EDITORONLY_DATA
+	// note that being ENetMode::NM_DedicatedServer implies DebugDrawingWorld is a game world, which is exactly what we need
+	if (DebugDrawingWorld != nullptr && DebugDrawingWorld->GetNetMode() == ENetMode::NM_DedicatedServer)
+	{
+		// no point in trying to draw on dedicated server. Let's see if there's a client world we can use for drawing!
+		const TIndirectArray<FWorldContext>& WorldContexts = GEngine->GetWorldContexts();
+		for (const FWorldContext& Context : WorldContexts)
+		{
+			if (Context.World()->IsGameWorld() && Context.World()->GetNetMode() != ENetMode::NM_DedicatedServer)
+			{
+				DebugDrawingWorld = Context.World();
+				break;
+			}
+		}
+	}
+#endif
+
+	return DebugDrawingWorld;
+}
+
+void UCrowdManager::DrawDebugCorners(const dtCrowdAgent* CrowdAgent) const
+{
+	UWorld* DebugDrawingWorld = GetDebugDrawingWorld();
 
 	{
 		FVector P0 = Recast2UnrealPoint(CrowdAgent->npos);
@@ -999,9 +995,7 @@ void UCrowdManager::DrawDebugCorners(const dtCrowdAgent* CrowdAgent) const
 
 void UCrowdManager::DrawDebugCollisionSegments(const dtCrowdAgent* CrowdAgent) const
 {
-#if !WITH_EDITORONLY_DATA
-	UWorld* DebugDrawingWorld = GetWorld();
-#endif // !WITH_EDITORONLY_DATA
+	UWorld* DebugDrawingWorld = GetDebugDrawingWorld();
 
 	FVector Center = Recast2UnrealPoint(CrowdAgent->boundary.getCenter()) + FCrowdDebug::Offset;
 	DrawDebugCylinder(DebugDrawingWorld, Center - FCrowdDebug::Offset, Center, UE_REAL_TO_FLOAT_CLAMPED_MAX(CrowdAgent->params.collisionQueryRange), 32, FCrowdDebug::CollisionRange);
@@ -1023,9 +1017,7 @@ void UCrowdManager::DrawDebugCollisionSegments(const dtCrowdAgent* CrowdAgent) c
 
 void UCrowdManager::DrawDebugPath(const dtCrowdAgent* CrowdAgent) const
 {
-#if !WITH_EDITORONLY_DATA
-	UWorld* DebugDrawingWorld = GetWorld();
-#endif // !WITH_EDITORONLY_DATA
+	UWorld* DebugDrawingWorld = GetDebugDrawingWorld();
 
 	ARecastNavMesh* NavMesh = Cast<ARecastNavMesh>(MyNavData);
 	if (NavMesh == NULL)
@@ -1064,9 +1056,8 @@ void UCrowdManager::DrawDebugPath(const dtCrowdAgent* CrowdAgent) const
 
 void UCrowdManager::DrawDebugVelocityObstacles(const dtCrowdAgent* CrowdAgent) const
 {
-#if !WITH_EDITORONLY_DATA
-	UWorld* DebugDrawingWorld = GetWorld();
-#endif // !WITH_EDITORONLY_DATA
+	UWorld* DebugDrawingWorld = GetDebugDrawingWorld();
+
 	FVector Center = Recast2UnrealPoint(CrowdAgent->npos) + FCrowdDebug::Offset;
 	DrawDebugCylinder(DebugDrawingWorld, Center - FCrowdDebug::Offset, Center, UE_REAL_TO_FLOAT_CLAMPED_MAX(CrowdAgent->params.maxSpeed), 32, FCrowdDebug::AvoidanceRange);
 
@@ -1104,9 +1095,8 @@ void UCrowdManager::DrawDebugVelocityObstacles(const dtCrowdAgent* CrowdAgent) c
 
 void UCrowdManager::DrawDebugPathOptimization(const dtCrowdAgent* CrowdAgent) const
 {
-#if !WITH_EDITORONLY_DATA
-	UWorld* DebugDrawingWorld = GetWorld();
-#endif // !WITH_EDITORONLY_DATA
+	UWorld* DebugDrawingWorld = GetDebugDrawingWorld();
+
 	FVector Pt0 = Recast2UnrealPoint(DetourAgentDebug->optStart) + FCrowdDebug::Offset * 1.25f;
 	FVector Pt1 = Recast2UnrealPoint(DetourAgentDebug->optEnd) + FCrowdDebug::Offset * 1.25f;
 
@@ -1115,9 +1105,8 @@ void UCrowdManager::DrawDebugPathOptimization(const dtCrowdAgent* CrowdAgent) co
 
 void UCrowdManager::DrawDebugNeighbors(const dtCrowdAgent* CrowdAgent) const
 {
-#if !WITH_EDITORONLY_DATA
-	UWorld* DebugDrawingWorld = GetWorld();
-#endif // !WITH_EDITORONLY_DATA
+	UWorld* DebugDrawingWorld = GetDebugDrawingWorld();
+
 	FVector Center = Recast2UnrealPoint(CrowdAgent->npos) + FCrowdDebug::Offset;
 	DrawDebugCylinder(DebugDrawingWorld, Center - FCrowdDebug::Offset, Center, UE_REAL_TO_FLOAT_CLAMPED_MAX(CrowdAgent->params.collisionQueryRange), 32, FCrowdDebug::CollisionRange);
 
@@ -1134,9 +1123,8 @@ void UCrowdManager::DrawDebugNeighbors(const dtCrowdAgent* CrowdAgent) const
 
 void UCrowdManager::DrawDebugSharedBoundary() const
 {
-#if !WITH_EDITORONLY_DATA
-	UWorld* DebugDrawingWorld = GetWorld();
-#endif // !WITH_EDITORONLY_DATA
+	UWorld* DebugDrawingWorld = GetDebugDrawingWorld();
+
 	FColor Colors[] = { FColorList::Red, FColorList::Orange };
 
 	const dtSharedBoundary* sharedBounds = DetourCrowd->getSharedBoundary();
