Remove GetSubtreeAsset and expose a getter for DefaultBehaviorAsset in RunBehaviorDynamic

diff --git a/Engine/Source/Runtime/AIModule/Classes/BehaviorTree/Tasks/BTTask_RunBehaviorDynamic.h b/Engine/Source/Runtime/AIModule/Classes/BehaviorTree/Tasks/BTTask_RunBehaviorDynamic.h
--- a/Engine/Source/Runtime/AIModule/Classes/BehaviorTree/Tasks/BTTask_RunBehaviorDynamic.h
+++ b/Engine/Source/Runtime/AIModule/Classes/BehaviorTree/Tasks/BTTask_RunBehaviorDynamic.h
@@ -34,8 +34,8 @@ class AIMODULE_API UBTTask_RunBehaviorDynamic : public UBTTaskNode
 	bool HasMatchingTag(const FGameplayTag& Tag) const;
 	bool SetBehaviorAsset(UBehaviorTree* NewBehaviorAsset);
 	
-	/** @returns subtree asset */
-	UBehaviorTree* GetSubtreeAsset() const;
+	/** @returns default subtree asset */
+	UBehaviorTree* GetDefaultBehaviorAsset() const;
 
 protected:
 
@@ -63,8 +63,8 @@ FORCEINLINE bool UBTTask_RunBehaviorDynamic::HasMatchingTag(const FGameplayTag&
 	return InjectionTag == Tag;
 }
 
-FORCEINLINE UBehaviorTree* UBTTask_RunBehaviorDynamic::GetSubtreeAsset() const
+FORCEINLINE UBehaviorTree* UBTTask_RunBehaviorDynamic::GetDefaultBehaviorAsset() const
 {
-	return BehaviorAsset;
+	return DefaultBehaviorAsset;
 }
 
