Minor touch-ups on recently added EnvQueryGenerator_PerceivedActors

diff --git a/Engine/Source/Runtime/AIModule/Classes/EnvironmentQuery/Generators/EnvQueryGenerator_PerceivedActors.h b/Engine/Source/Runtime/AIModule/Classes/EnvironmentQuery/Generators/EnvQueryGenerator_PerceivedActors.h
--- a/Engine/Source/Runtime/AIModule/Classes/EnvironmentQuery/Generators/EnvQueryGenerator_PerceivedActors.h
+++ b/Engine/Source/Runtime/AIModule/Classes/EnvironmentQuery/Generators/EnvQueryGenerator_PerceivedActors.h
@@ -2,9 +2,6 @@
 
 #pragma once
 
-//#include "Templates/SubclassOf.h"
-//#include "GameFramework/Actor.h"
-//#include "EnvironmentQuery/EnvQueryContext.h"
 #include "DataProviders/AIDataProvider.h"
 #include "EnvironmentQuery/EnvQueryGenerator.h"
 #include "Perception/AISense.h"
@@ -19,7 +16,7 @@ class AIMODULE_API UEnvQueryGenerator_PerceivedActors : public UEnvQueryGenerato
 
 	/** If set will be used to filter results */
 	UPROPERTY(EditDefaultsOnly, Category=Generator)
-	TSubclassOf<AActor> ActorClassWhitelist;
+	TSubclassOf<AActor> AllowedActorClass;
 
 	/** Additional distance limit imposed on the items generated. Perception's range limit still applies. */
 	UPROPERTY(EditDefaultsOnly, Category=Generator)
diff --git a/Engine/Source/Runtime/AIModule/Private/EnvironmentQuery/Generators/EnvQueryGenerator_PerceivedActors.cpp b/Engine/Source/Runtime/AIModule/Private/EnvironmentQuery/Generators/EnvQueryGenerator_PerceivedActors.cpp
--- a/Engine/Source/Runtime/AIModule/Private/EnvironmentQuery/Generators/EnvQueryGenerator_PerceivedActors.cpp
+++ b/Engine/Source/Runtime/AIModule/Private/EnvironmentQuery/Generators/EnvQueryGenerator_PerceivedActors.cpp
@@ -110,7 +110,7 @@ void UEnvQueryGenerator_PerceivedActors::GenerateItems(FEnvQueryInstance& QueryI
 		{
 			if (PerceivedActor 
 				&& (RadiusValue <= 0 || FVector::DistSquared(ListenerLocation, PerceivedActor->GetActorLocation()) < RadiusSq)
-				&& (!ActorClassWhitelist || PerceivedActor->IsA(ActorClassWhitelist.Get())))
+				&& (!AllowedActorClass || PerceivedActor->IsA(AllowedActorClass.Get())))
 			{
 				AllPerceivedActors.Add(PerceivedActor);
 				AllPerceivedActors.Add(PerceivedActor);
@@ -151,9 +151,9 @@ FText UEnvQueryGenerator_PerceivedActors::GetDescriptionDetails() const
 		Args.Add(TEXT("Sense"), FText::FromString(TEXT("All senses")));
 	}
 	
-	if (ActorClassWhitelist)
+	if (AllowedActorClass)
 	{
-		Args.Add(TEXT("ActorClass"), FText::FromString(ActorClassWhitelist->GetName()));
+		Args.Add(TEXT("ActorClass"), FText::FromString(AllowedActorClass->GetName()));
 		return FText::Format(LOCTEXT("PerceivedActorsDescription", "radius: {Radius}\nsense: {Sense}\nactors of class: {ActorClass}"), Args);
 	}
 
