Remove a piece of test-time code accidentally left behind.

diff --git a/Engine/Source/Runtime/AIModule/Classes/EnvironmentQuery/Generators/EnvQueryGenerator_PerceivedActors.h b/Engine/Source/Runtime/AIModule/Classes/EnvironmentQuery/Generators/EnvQueryGenerator_PerceivedActors.h
--- a/Engine/Source/Runtime/AIModule/Classes/EnvironmentQuery/Generators/EnvQueryGenerator_PerceivedActors.h
+++ b/Engine/Source/Runtime/AIModule/Classes/EnvironmentQuery/Generators/EnvQueryGenerator_PerceivedActors.h
@@ -14,6 +14,7 @@ class AIMODULE_API UEnvQueryGenerator_PerceivedActors : public UEnvQueryGenerato
 {
 	GENERATED_UCLASS_BODY()
 
+protected:
 	/** If set will be used to filter results */
 	UPROPERTY(EditDefaultsOnly, Category=Generator)
 	TSubclassOf<AActor> AllowedActorClass;
diff --git a/Engine/Source/Runtime/AIModule/Private/EnvironmentQuery/Generators/EnvQueryGenerator_PerceivedActors.cpp b/Engine/Source/Runtime/AIModule/Private/EnvironmentQuery/Generators/EnvQueryGenerator_PerceivedActors.cpp
--- a/Engine/Source/Runtime/AIModule/Private/EnvironmentQuery/Generators/EnvQueryGenerator_PerceivedActors.cpp
+++ b/Engine/Source/Runtime/AIModule/Private/EnvironmentQuery/Generators/EnvQueryGenerator_PerceivedActors.cpp
@@ -113,7 +113,6 @@ void UEnvQueryGenerator_PerceivedActors::GenerateItems(FEnvQueryInstance& QueryI
 				&& (!AllowedActorClass || PerceivedActor->IsA(AllowedActorClass.Get())))
 			{
 				AllPerceivedActors.Add(PerceivedActor);
-				AllPerceivedActors.Add(PerceivedActor);
 			}
 		}
 	}
