Added a parameter to EQS running BT task and service that once set will clear out the indicated BB entry upon EQS query's failure.

diff --git a/Engine/Source/Runtime/AIModule/Classes/AISystem.h b/Engine/Source/Runtime/AIModule/Classes/AISystem.h
--- a/Engine/Source/Runtime/AIModule/Classes/AISystem.h
+++ b/Engine/Source/Runtime/AIModule/Classes/AISystem.h
@@ -94,6 +94,9 @@ public:
 	UPROPERTY(globalconfig, EditDefaultsOnly, Category = "Blackboard")
 	bool bAddBlackboardSelfKey = true;
 
+	UPROPERTY(globalconfig, EditDefaultsOnly, Category = "Behavior Tree")
+	bool bClearBBEntryOnBTEQSFail = true;
+
 	/** Which collision channel to use for sight checks by default */
 	UPROPERTY(globalconfig, EditDefaultsOnly, Category = "PerceptionSystem")
 	TEnumAsByte<ECollisionChannel> DefaultSightCollisionChannel;
diff --git a/Engine/Source/Runtime/AIModule/Classes/BehaviorTree/Services/BTService_RunEQS.h b/Engine/Source/Runtime/AIModule/Classes/BehaviorTree/Services/BTService_RunEQS.h
--- a/Engine/Source/Runtime/AIModule/Classes/BehaviorTree/Services/BTService_RunEQS.h
+++ b/Engine/Source/Runtime/AIModule/Classes/BehaviorTree/Services/BTService_RunEQS.h
@@ -26,6 +26,9 @@ protected:
 	UPROPERTY(Category = EQS, EditAnywhere)
 	FEQSParametrizedQueryExecutionRequest EQSRequest;
 
+	UPROPERTY(Category = EQS, EditAnywhere)
+	bool bUpdateBBOnFail = false;
+
 	FQueryFinishedSignature QueryFinishedDelegate;
 
 public:
diff --git a/Engine/Source/Runtime/AIModule/Classes/BehaviorTree/Tasks/BTTask_RunEQSQuery.h b/Engine/Source/Runtime/AIModule/Classes/BehaviorTree/Tasks/BTTask_RunEQSQuery.h
--- a/Engine/Source/Runtime/AIModule/Classes/BehaviorTree/Tasks/BTTask_RunEQSQuery.h
+++ b/Engine/Source/Runtime/AIModule/Classes/BehaviorTree/Tasks/BTTask_RunEQSQuery.h
@@ -51,6 +51,9 @@ class AIMODULE_API UBTTask_RunEQSQuery : public UBTTask_BlackboardBase
 	UPROPERTY(Category = EQS, EditAnywhere)
 	FEQSParametrizedQueryExecutionRequest EQSRequest;
 
+	UPROPERTY(Category = EQS, EditAnywhere)
+	bool bUpdateBBOnFail = false;
+
 	FQueryFinishedSignature QueryFinishedDelegate;
 
 	virtual void InitializeFromAsset(UBehaviorTree& Asset) override;
diff --git a/Engine/Source/Runtime/AIModule/Private/BehaviorTree/Services/BTService_RunEQS.cpp b/Engine/Source/Runtime/AIModule/Private/BehaviorTree/Services/BTService_RunEQS.cpp
--- a/Engine/Source/Runtime/AIModule/Private/BehaviorTree/Services/BTService_RunEQS.cpp
+++ b/Engine/Source/Runtime/AIModule/Private/BehaviorTree/Services/BTService_RunEQS.cpp
@@ -8,6 +8,9 @@
 #include "VisualLogger/VisualLogger.h"
 #include "EnvironmentQuery/EnvQuery.h"
 #include "EnvironmentQuery/EnvQueryManager.h"
+#include "BehaviorTree/BlackboardComponent.h"
+#include "AISystem.h"
+
 
 UBTService_RunEQS::UBTService_RunEQS(const FObjectInitializer& ObjectInitializer)
 	: Super(ObjectInitializer)
@@ -22,6 +25,8 @@ UBTService_RunEQS::UBTService_RunEQS(const FObjectInitializer& ObjectInitializer
 	BlackboardKey.AddVectorFilter(this, GET_MEMBER_NAME_CHECKED(UBTService_RunEQS, BlackboardKey));
 
 	QueryFinishedDelegate = FQueryFinishedSignature::CreateUObject(this, &UBTService_RunEQS::OnQueryFinished);
+
+	bUpdateBBOnFail = GET_AI_CONFIG_VAR(bClearBBEntryOnBTEQSFail);
 }
 
 void UBTService_RunEQS::InitializeFromAsset(UBehaviorTree& Asset)
@@ -99,6 +104,12 @@ void UBTService_RunEQS::OnQueryFinished(TSharedPtr<FEnvQueryResult> Result)
 				*UBehaviorTreeTypes::GetShortTypeName(BlackboardKey.SelectedKeyType));
 		}
 	}
+	else if (bUpdateBBOnFail)
+	{
+		UBlackboardComponent* MyBlackboard = BTComp->GetBlackboardComponent();
+		check(MyBlackboard);
+		MyBlackboard->ClearValue(BlackboardKey.GetSelectedKeyID());
+	}
 
 	MyMemory->RequestID = INDEX_NONE;
 }
diff --git a/Engine/Source/Runtime/AIModule/Private/BehaviorTree/Tasks/BTTask_RunEQSQuery.cpp b/Engine/Source/Runtime/AIModule/Private/BehaviorTree/Tasks/BTTask_RunEQSQuery.cpp
--- a/Engine/Source/Runtime/AIModule/Private/BehaviorTree/Tasks/BTTask_RunEQSQuery.cpp
+++ b/Engine/Source/Runtime/AIModule/Private/BehaviorTree/Tasks/BTTask_RunEQSQuery.cpp
@@ -8,6 +8,8 @@
 #include "VisualLogger/VisualLogger.h"
 #include "EnvironmentQuery/EnvQuery.h"
 #include "EnvironmentQuery/EnvQueryManager.h"
+#include "BehaviorTree/BlackboardComponent.h"
+#include "AISystem.h"
 
 
 UBTTask_RunEQSQuery::UBTTask_RunEQSQuery(const FObjectInitializer& ObjectInitializer) 
@@ -24,6 +26,8 @@ UBTTask_RunEQSQuery::UBTTask_RunEQSQuery(const FObjectInitializer& ObjectInitial
 
 	QueryFinishedDelegate = FQueryFinishedSignature::CreateUObject(this, &UBTTask_RunEQSQuery::OnQueryFinished);
 
+	bUpdateBBOnFail = GET_AI_CONFIG_VAR(bClearBBEntryOnBTEQSFail);
+
 	// deprecated
 	EQSQueryBlackboardKey.AddObjectFilter(this, GET_MEMBER_NAME_CHECKED(UBTTask_RunEQSQuery, EQSQueryBlackboardKey), UEnvQuery::StaticClass());
 }
@@ -132,6 +136,12 @@ void UBTTask_RunEQSQuery::OnQueryFinished(TSharedPtr<FEnvQueryResult> Result)
 				*UBehaviorTreeTypes::GetShortTypeName(BlackboardKey.SelectedKeyType));
 		}
 	}
+	else if (bUpdateBBOnFail)
+	{
+		UBlackboardComponent* MyBlackboard = MyComp->GetBlackboardComponent();
+		check(MyBlackboard);
+		MyBlackboard->ClearValue(BlackboardKey.GetSelectedKeyID());
+	}
 
 	FAIMessage::Send(MyComp, FAIMessage(UBrainComponent::AIMessage_QueryFinished, this, Result->QueryID, bSuccess));
 }
