Added an option to EnvQueryGenerator_PerceivedActors to control whether all actors known to the AIPerceptionComponent will be gathered or only the ones actively, currently being perceived (like "visible at this very moment").

diff --git a/Engine/Source/Runtime/AIModule/Classes/EnvironmentQuery/Generators/EnvQueryGenerator_PerceivedActors.h b/Engine/Source/Runtime/AIModule/Classes/EnvironmentQuery/Generators/EnvQueryGenerator_PerceivedActors.h
--- a/Engine/Source/Runtime/AIModule/Classes/EnvironmentQuery/Generators/EnvQueryGenerator_PerceivedActors.h
+++ b/Engine/Source/Runtime/AIModule/Classes/EnvironmentQuery/Generators/EnvQueryGenerator_PerceivedActors.h
@@ -28,9 +28,17 @@ protected:
 	TSubclassOf<UEnvQueryContext> ListenerContext;
 
 	/** If set will be used to filter gathered results so that only actors perceived with a given sense are considered */
-	UPROPERTY(EditAnywhere, Category = Generator)
+	UPROPERTY(EditAnywhere, Category=Generator)
 	TSubclassOf<UAISense> SenseToUse;
 
+	/**
+	 * Indicates whether to include all actors known via perception (TRUE) or just the ones actively being perceived 
+	 * at the moment (example "currently visible" as opposed to "seen and the perception stimulus haven't expired yet").
+	 * @see FAIStimulus.bExpired
+	 */
+	UPROPERTY(EditAnywhere, Category=Generator)
+	bool bIncludeKnownActors = true;
+
 	virtual void GenerateItems(FEnvQueryInstance& QueryInstance) const override;
 
 	virtual FText GetDescriptionTitle() const override;
diff --git a/Engine/Source/Runtime/AIModule/Private/EnvironmentQuery/Generators/EnvQueryGenerator_PerceivedActors.cpp b/Engine/Source/Runtime/AIModule/Private/EnvironmentQuery/Generators/EnvQueryGenerator_PerceivedActors.cpp
--- a/Engine/Source/Runtime/AIModule/Private/EnvironmentQuery/Generators/EnvQueryGenerator_PerceivedActors.cpp
+++ b/Engine/Source/Runtime/AIModule/Private/EnvironmentQuery/Generators/EnvQueryGenerator_PerceivedActors.cpp
@@ -104,7 +104,14 @@ void UEnvQueryGenerator_PerceivedActors::GenerateItems(FEnvQueryInstance& QueryI
 
 		const FVector ListenerLocation = ListenerActor->GetActorLocation();
 		TArray<AActor*> LocalPerceivedActors;
-		PerceptionComponent->GetKnownPerceivedActors(SenseToUse, LocalPerceivedActors);
+		if (bIncludeKnownActors)
+		{
+			PerceptionComponent->GetKnownPerceivedActors(SenseToUse, LocalPerceivedActors);
+		}
+		else
+		{
+			PerceptionComponent->GetCurrentlyPerceivedActors(SenseToUse, LocalPerceivedActors);
+		}
 		
 		for (AActor* PerceivedActor : LocalPerceivedActors)
 		{
