Fix localization issue #119824

diff --git a/Engine/Source/Runtime/AIModule/Private/EnvironmentQuery/Generators/EnvQueryGenerator_PerceivedActors.cpp b/Engine/Source/Runtime/AIModule/Private/EnvironmentQuery/Generators/EnvQueryGenerator_PerceivedActors.cpp
--- a/Engine/Source/Runtime/AIModule/Private/EnvironmentQuery/Generators/EnvQueryGenerator_PerceivedActors.cpp
+++ b/Engine/Source/Runtime/AIModule/Private/EnvironmentQuery/Generators/EnvQueryGenerator_PerceivedActors.cpp
@@ -160,7 +160,7 @@ FText UEnvQueryGenerator_PerceivedActors::GetDescriptionDetails() const
 	if (AllowedActorClass)
 	{
 		Args.Add(TEXT("ActorClass"), FText::FromString(AllowedActorClass->GetName()));
-		return FText::Format(LOCTEXT("PerceivedActorsDescription", "radius: {Radius}\nsense: {Sense}\nactors of class: {ActorClass}"), Args);
+		return FText::Format(LOCTEXT("PerceivedActorsOfClassDescription", "radius: {Radius}\nsense: {Sense}\nactors of class: {ActorClass}"), Args);
 	}
 
 	return FText::Format(LOCTEXT("PerceivedActorsDescription", "radius: {Radius}\nsense: {Sense}"), Args);
