Formally marked as deprecated EQS-configuration properties that have been annotated as deprecated for a long while now. The properties are to be removed entirely for 5.1.

diff --git a/Engine/Source/Runtime/AIModule/Classes/BehaviorTree/Tasks/BTTask_RunEQSQuery.h b/Engine/Source/Runtime/AIModule/Classes/BehaviorTree/Tasks/BTTask_RunEQSQuery.h
--- a/Engine/Source/Runtime/AIModule/Classes/BehaviorTree/Tasks/BTTask_RunEQSQuery.h
+++ b/Engine/Source/Runtime/AIModule/Classes/BehaviorTree/Tasks/BTTask_RunEQSQuery.h
@@ -27,22 +27,27 @@ class AIMODULE_API UBTTask_RunEQSQuery : public UBTTask_BlackboardBase
 	GENERATED_UCLASS_BODY()
 
 	/** query to run */
-	UPROPERTY(Category = Node, VisibleAnywhere, meta = (EditCondition = "!bUseBBKey", DisplayName = "QueryTemplate_DEPRECATED"))
+	UE_DEPRECATED_FORGAME(5.0, "UBTTask_RunEQSQuery.QueryTemplate has been deprecated for a long while now. Will be removed in the next engine version.")
+	UPROPERTY()
 	TObjectPtr<UEnvQuery> QueryTemplate;
 
 	/** optional parameters for query */
-	UPROPERTY(Category = Node, VisibleAnywhere, meta = (DisplayName = "QueryParams_DEPRECATED"))
+	UE_DEPRECATED_FORGAME(5.0, "UBTTask_RunEQSQuery.QueryParams has been deprecated for a long while now. Will be removed in the next engine version.")
+	UPROPERTY()
 	TArray<FEnvNamedValue> QueryParams;
 
-	UPROPERTY(Category = Node, VisibleAnywhere, meta = (DisplayName = "QueryConfig_DEPRECATED"))
+	UE_DEPRECATED_FORGAME(5.0, "UBTTask_RunEQSQuery.QueryConfig has been deprecated for a long while now. Will be removed in the next engine version.")
+	UPROPERTY()
 	TArray<FAIDynamicParam> QueryConfig;
 
 	/** determines which item will be stored (All = only first matching) */
-	UPROPERTY(Category = Node, VisibleAnywhere, meta = (DisplayName = "RunMode_DEPRECATED"))
+	UE_DEPRECATED_FORGAME(5.0, "UBTTask_RunEQSQuery.RunMode has been deprecated for a long while now. Will be removed in the next engine version.")
+	UPROPERTY()
 	TEnumAsByte<EEnvQueryRunMode::Type> RunMode;
 
 	/** blackboard key storing an EQS query template */
-	UPROPERTY(VisibleAnywhere, Category = Blackboard, meta = (EditCondition = "bUseBBKey", DisplayName = "EQSQueryBlackboardKey_DEPRECATED"))
+	UE_DEPRECATED_FORGAME(5.0, "UBTTask_RunEQSQuery.EQSQueryBlackboardKey been deprecated for a long while now. Will be removed in the next engine version.")
+	UPROPERTY()
 	struct FBlackboardKeySelector EQSQueryBlackboardKey;
 
 	UPROPERTY(EditAnywhere, Category=Node, meta=(InlineEditConditionToggle))
diff --git a/Engine/Source/Runtime/AIModule/Classes/EnvironmentQuery/EQSTestingPawn.h b/Engine/Source/Runtime/AIModule/Classes/EnvironmentQuery/EQSTestingPawn.h
--- a/Engine/Source/Runtime/AIModule/Classes/EnvironmentQuery/EQSTestingPawn.h
+++ b/Engine/Source/Runtime/AIModule/Classes/EnvironmentQuery/EQSTestingPawn.h
@@ -29,7 +29,8 @@ class AIMODULE_API AEQSTestingPawn : public ACharacter, public IEQSQueryResultSo
 	TObjectPtr<UEnvQuery> QueryTemplate;
 
 	/** optional parameters for query */
-	UPROPERTY(Category=EQS, VisibleAnywhere, meta=(DisplayName="QueryParams_DEPRECATED"))
+	UE_DEPRECATED_FORGAME(5.0, "QueryParams has been deprecated for a long while now. Will be removed in the next engine version.")
+	UPROPERTY()
 	TArray<FEnvNamedValue> QueryParams;
 
 	UPROPERTY(Category=EQS, EditAnywhere)
