UE-140854 - Floats are still supported in Blackboard system

The easiest fix here is to update the metadata to display "real" instead of "float" where appropriate. The alternative would require class&function renaming, which involves fixup, and is likely not desirable during lockdown.

diff --git a/Engine/Source/Runtime/AIModule/Classes/BehaviorTree/Blackboard/BlackboardKeyType_Float.h b/Engine/Source/Runtime/AIModule/Classes/BehaviorTree/Blackboard/BlackboardKeyType_Float.h
--- a/Engine/Source/Runtime/AIModule/Classes/BehaviorTree/Blackboard/BlackboardKeyType_Float.h
+++ b/Engine/Source/Runtime/AIModule/Classes/BehaviorTree/Blackboard/BlackboardKeyType_Float.h
@@ -8,7 +8,7 @@
 
 class UBlackboardComponent;
 
-UCLASS(EditInlineNew, meta=(DisplayName="Float"))
+UCLASS(EditInlineNew, meta=(DisplayName="Real", ToolTip="Blackboard Key Type Real"))
 class AIMODULE_API UBlackboardKeyType_Float : public UBlackboardKeyType
 {
 	GENERATED_UCLASS_BODY()
diff --git a/Engine/Source/Runtime/AIModule/Classes/BehaviorTree/BlackboardComponent.h b/Engine/Source/Runtime/AIModule/Classes/BehaviorTree/BlackboardComponent.h
--- a/Engine/Source/Runtime/AIModule/Classes/BehaviorTree/BlackboardComponent.h
+++ b/Engine/Source/Runtime/AIModule/Classes/BehaviorTree/BlackboardComponent.h
@@ -141,8 +141,8 @@ public:
 	UFUNCTION(BlueprintCallable, Category="AI|Components|Blackboard")
 	void SetValueAsInt(const FName& KeyName, int32 IntValue);
 
-	UFUNCTION(BlueprintCallable, Category="AI|Components|Blackboard")
-	void SetValueAsFloat(const FName& KeyName, float FloatValue);
+	UFUNCTION(BlueprintCallable, Category="AI|Components|Blackboard", DisplayName = "Set Value As Real", Meta = (ToolTip = "Set Value As Real"))
+	void SetValueAsFloat(const FName& KeyName, UPARAM(DisplayName = "Real Value") float FloatValue);
 
 	UFUNCTION(BlueprintCallable, Category="AI|Components|Blackboard")
 	void SetValueAsBool(const FName& KeyName, bool BoolValue);
