THierarchicalHashGrid2D::CalcCellBounds as a const method
#rnx

diff --git a/Engine/Source/Runtime/AIModule/Public/HierarchicalHashGrid2D.h b/Engine/Source/Runtime/AIModule/Public/HierarchicalHashGrid2D.h
--- a/Engine/Source/Runtime/AIModule/Public/HierarchicalHashGrid2D.h
+++ b/Engine/Source/Runtime/AIModule/Public/HierarchicalHashGrid2D.h
@@ -529,7 +529,7 @@ public:
 	 * @params CellLocation - Cell location.
 	 * @return Returns the bounds of the cell.
 	 */
-	FBox CalcCellBounds(const FCellLocation& CellLocation)
+	FBox CalcCellBounds(const FCellLocation& CellLocation) const
 	{
 		const float Size = CellSize[CellLocation.Level];
 		const float X = CellLocation.X * Size;
