Changing the Pathfinding and PathfindingBatch EQS tests to use a function to get the nav filter.

#rnx

diff --git a/Engine/Source/Runtime/AIModule/Classes/EnvironmentQuery/Tests/EnvQueryTest_Pathfinding.h b/Engine/Source/Runtime/AIModule/Classes/EnvironmentQuery/Tests/EnvQueryTest_Pathfinding.h
--- a/Engine/Source/Runtime/AIModule/Classes/EnvironmentQuery/Tests/EnvQueryTest_Pathfinding.h
+++ b/Engine/Source/Runtime/AIModule/Classes/EnvironmentQuery/Tests/EnvQueryTest_Pathfinding.h
@@ -76,4 +76,6 @@ protected:
 	float FindPathLengthTo(const FVector& ItemPos, const FVector& ContextPos, EPathFindingMode::Type Mode, const ANavigationData& NavData, UNavigationSystemV1& NavSys, FSharedConstNavQueryFilter NavFilter, const UObject* PathOwner) const;
 
 	ANavigationData* FindNavigationData(UNavigationSystemV1& NavSys, UObject* Owner) const;
+
+	virtual TSubclassOf<UNavigationQueryFilter> GetNavFilterClass(FEnvQueryInstance& QueryInstance) const;
 };
diff --git a/Engine/Source/Runtime/AIModule/Classes/EnvironmentQuery/Tests/EnvQueryTest_PathfindingBatch.h b/Engine/Source/Runtime/AIModule/Classes/EnvironmentQuery/Tests/EnvQueryTest_PathfindingBatch.h
--- a/Engine/Source/Runtime/AIModule/Classes/EnvironmentQuery/Tests/EnvQueryTest_PathfindingBatch.h
+++ b/Engine/Source/Runtime/AIModule/Classes/EnvironmentQuery/Tests/EnvQueryTest_PathfindingBatch.h
@@ -9,7 +9,7 @@
 #include "EnvQueryTest_PathfindingBatch.generated.h"
 
 UCLASS()
-class UEnvQueryTest_PathfindingBatch : public UEnvQueryTest_Pathfinding
+class AIMODULE_API UEnvQueryTest_PathfindingBatch : public UEnvQueryTest_Pathfinding
 {
 	GENERATED_UCLASS_BODY()
 
diff --git a/Engine/Source/Runtime/AIModule/Private/EnvironmentQuery/Tests/EnvQueryTest_Pathfinding.cpp b/Engine/Source/Runtime/AIModule/Private/EnvironmentQuery/Tests/EnvQueryTest_Pathfinding.cpp
--- a/Engine/Source/Runtime/AIModule/Private/EnvironmentQuery/Tests/EnvQueryTest_Pathfinding.cpp
+++ b/Engine/Source/Runtime/AIModule/Private/EnvironmentQuery/Tests/EnvQueryTest_Pathfinding.cpp
@@ -58,7 +58,7 @@ void UEnvQueryTest_Pathfinding::RunTest(FEnvQueryInstance& QueryInstance) const
 	}
 
 	EPathFindingMode::Type PFMode(EPathFindingMode::Regular);
-	FSharedConstNavQueryFilter NavFilter = UNavigationQueryFilter::GetQueryFilter(*NavData, QueryOwner, FilterClass);
+	FSharedConstNavQueryFilter NavFilter = UNavigationQueryFilter::GetQueryFilter(*NavData, QueryOwner, GetNavFilterClass(QueryInstance));
 
 	if (GetWorkOnFloatValues())
 	{
@@ -231,4 +231,9 @@ ANavigationData* UEnvQueryTest_Pathfinding::FindNavigationData(UNavigationSystem
 	return NavSys.GetDefaultNavDataInstance(FNavigationSystem::DontCreate);
 }
 
+TSubclassOf<UNavigationQueryFilter> UEnvQueryTest_Pathfinding::GetNavFilterClass(FEnvQueryInstance& QueryInstance) const
+{
+	return FilterClass;
+}
+
 #undef LOCTEXT_NAMESPACE
diff --git a/Engine/Source/Runtime/AIModule/Private/EnvironmentQuery/Tests/EnvQueryTest_PathfindingBatch.cpp b/Engine/Source/Runtime/AIModule/Private/EnvironmentQuery/Tests/EnvQueryTest_PathfindingBatch.cpp
--- a/Engine/Source/Runtime/AIModule/Private/EnvironmentQuery/Tests/EnvQueryTest_PathfindingBatch.cpp
+++ b/Engine/Source/Runtime/AIModule/Private/EnvironmentQuery/Tests/EnvQueryTest_PathfindingBatch.cpp
@@ -91,8 +91,10 @@ void UEnvQueryTest_PathfindingBatch::RunTest(FEnvQueryInstance& QueryInstance) c
 	TArray<float> CollectDistanceSq;
 	CollectDistanceSq.Init(0.0f, ContextLocations.Num());
 
-	FSharedNavQueryFilter NavigationFilterCopy = FilterClass ?
-		UNavigationQueryFilter::GetQueryFilter(*NavMeshData, QueryOwner, FilterClass)->GetCopy() :
+	TSubclassOf<UNavigationQueryFilter> NavFilterToUse = GetNavFilterClass(QueryInstance);
+
+	FSharedNavQueryFilter NavigationFilterCopy = NavFilterToUse ?
+		UNavigationQueryFilter::GetQueryFilter(*NavMeshData, QueryOwner, NavFilterToUse)->GetCopy() :
 		NavMeshData->GetDefaultQueryFilter()->GetCopy();
 
 	NavigationFilterCopy->SetBacktrackingEnabled(!bPathToItem);
