Fix range loop being change while iterating through stimuli in AI perception component

diff --git a/Engine/Source/Runtime/AIModule/Private/Perception/AIPerceptionComponent.cpp b/Engine/Source/Runtime/AIModule/Private/Perception/AIPerceptionComponent.cpp
--- a/Engine/Source/Runtime/AIModule/Private/Perception/AIPerceptionComponent.cpp
+++ b/Engine/Source/Runtime/AIModule/Private/Perception/AIPerceptionComponent.cpp
@@ -462,7 +462,7 @@ AActor* UAIPerceptionComponent::GetMutableBodyActor()
 
 void UAIPerceptionComponent::RegisterStimulus(AActor* Source, const FAIStimulus& Stimulus)
 {
-	FStimulusToProcess& StimulusToProcess = StimuliToProcess[StimuliToProcess.Add(FStimulusToProcess(Source, Stimulus))];
+	FStimulusToProcess& StimulusToProcess = StimuliToProcess.Add_GetRef(FStimulusToProcess(Source, Stimulus));
 	StimulusToProcess.Stimulus.SetExpirationAge(MaxActiveAge[int32(Stimulus.Type)]);
 }
 
@@ -479,13 +479,14 @@ void UAIPerceptionComponent::ProcessStimuli()
 	const bool bBroadcastEveryTargetUpdate = OnTargetPerceptionUpdated.IsBound();
 	const bool bBroadcastEveryTargetInfoUpdate = OnTargetPerceptionInfoUpdated.IsBound();
 	
+	TArray<FStimulusToProcess> ProcessingStimuli = MoveTemp(StimuliToProcess);
 	TArray<AActor*> UpdatedActors;
-	UpdatedActors.Reserve(StimuliToProcess.Num());
+	UpdatedActors.Reserve(ProcessingStimuli.Num());
 	TArray<AActor*> ActorsToForget;
-	ActorsToForget.Reserve(StimuliToProcess.Num());
+	ActorsToForget.Reserve(ProcessingStimuli.Num());
 	TArray<TObjectKey<AActor>, TInlineAllocator<8>> DataToRemove;
 
-	for (FStimulusToProcess& SourcedStimulus : StimuliToProcess)
+	for (FStimulusToProcess& SourcedStimulus : ProcessingStimuli)
 	{
 		const TObjectKey<AActor>& SourceKey = SourcedStimulus.Source;
 
@@ -583,8 +584,6 @@ void UAIPerceptionComponent::ProcessStimuli()
 		}
 	}
 
-	StimuliToProcess.Reset();
-
 	if (UpdatedActors.Num() > 0)
 	{
 		if (AIOwner != NULL)
