Show the sight/lose sight radius number in the gameplay debugger
Github #7079 tag

#rnx

diff --git a/Engine/Source/Runtime/AIModule/Private/Perception/AISense.cpp b/Engine/Source/Runtime/AIModule/Private/Perception/AISense.cpp
--- a/Engine/Source/Runtime/AIModule/Private/Perception/AISense.cpp
+++ b/Engine/Source/Runtime/AIModule/Private/Perception/AISense.cpp
@@ -211,10 +211,10 @@ void UAISenseConfig_Sight::DescribeSelfToGameplayDebugger(const UAIPerceptionCom
 
 	// don't call Super implementation on purpose, replace color description line
 	DebuggerCategory->AddTextLine(
-		FString::Printf(TEXT("%s: {%s}%s {white}rangeIN:{%s}%s {white} rangeOUT:{%s}%s"), *GetSenseName(),
+		FString::Printf(TEXT("%s: {%s}%s {white}rangeIN:{%s} %.2f (%s) {white} rangeOUT:{%s} %.2f (%s)"), *GetSenseName(),
 			*GetDebugColor().ToString(), *DescribeColorHelper(GetDebugColor()),
-			*SightRangeColor.ToString(), *DescribeColorHelper(SightRangeColor),
-			*LoseSightRangeColor.ToString(), *DescribeColorHelper(LoseSightRangeColor))
+			*SightRangeColor.ToString(), SightRadius, *DescribeColorHelper(SightRangeColor),
+			*LoseSightRangeColor.ToString(), LoseSightRadius, *DescribeColorHelper(LoseSightRangeColor))
 		);
 
 	const AActor* BodyActor = PerceptionComponent->GetBodyActor();
