AIController - Provide protected overridable StopBrainComponent function

diff --git a/Engine/Source/Runtime/AIModule/Classes/AIController.h b/Engine/Source/Runtime/AIModule/Classes/AIController.h
--- a/Engine/Source/Runtime/AIModule/Classes/AIController.h
+++ b/Engine/Source/Runtime/AIModule/Classes/AIController.h
@@ -256,6 +256,10 @@ public:
 	UFUNCTION(BlueprintCallable, Category = "AI")
 	virtual bool RunBehaviorTree(UBehaviorTree* BTAsset);
 
+protected:
+	virtual void CleanupBrainComponent();
+
+public:
 	/**
 	 * Makes AI use the specified Blackboard asset & creates a Blackboard Component if one does not already exist.
 	 * @param	BlackboardAsset			The Blackboard asset to use.
diff --git a/Engine/Source/Runtime/AIModule/Private/AIController.cpp b/Engine/Source/Runtime/AIModule/Private/AIController.cpp
--- a/Engine/Source/Runtime/AIModule/Private/AIController.cpp
+++ b/Engine/Source/Runtime/AIModule/Private/AIController.cpp
@@ -533,9 +533,9 @@ void AAIController::OnUnPossess()
 		PathFollowingComponent->Cleanup();
 	}
 
-	if (bStopAILogicOnUnposses && BrainComponent)
+	if (bStopAILogicOnUnposses)
 	{
-		BrainComponent->Cleanup();
+		CleanupBrainComponent();
 	}
 
 	if (CachedGameplayTasksComponent && (CachedGameplayTasksComponent->GetOwner() == CurrentPawn))
@@ -940,6 +940,14 @@ bool AAIController::RunBehaviorTree(UBehaviorTree* BTAsset)
 	return bSuccess;
 }
 
+void AAIController::CleanupBrainComponent()
+{
+	if (BrainComponent)
+	{
+		BrainComponent->Cleanup();
+	}
+}
+
 void AAIController::ClaimTaskResource(TSubclassOf<UGameplayTaskResource> ResourceClass)
 {
 	if (CachedGameplayTasksComponent)
