PR #9005: Fix EEnvTestScoreOperator Multiply (Contributed by [REDACTED_PERSON])

diff --git a/Engine/Source/Runtime/AIModule/Classes/EnvironmentQuery/EnvQueryTypes.h b/Engine/Source/Runtime/AIModule/Classes/EnvironmentQuery/EnvQueryTypes.h
--- a/Engine/Source/Runtime/AIModule/Classes/EnvironmentQuery/EnvQueryTypes.h
+++ b/Engine/Source/Runtime/AIModule/Classes/EnvironmentQuery/EnvQueryTypes.h
@@ -1221,7 +1221,8 @@ public:
 				}
 				break;
 			case EEnvTestScoreOperator::Multiply:
-				ItemScore *= Score;
+				// ItemScore defaults to 0, so for first test we need to initialize the score, otherwise we end up constantly multiplying by 0
+				ItemScore = (NumTestsForItem == 0) ? Score : (ItemScore * Score);
 				break;
 			}
 		}
