Fixed DefaultManagerInstanceTracker not properly subscribing to OnPostWorldCleanup delegate, causing an occasional failed assertion (!IsRooted()) at world destruction time, depending on the order of static initialization.

diff --git a/Engine/Source/Runtime/AIModule/Public/DefaultManagerInstanceTracker.h b/Engine/Source/Runtime/AIModule/Public/DefaultManagerInstanceTracker.h
--- a/Engine/Source/Runtime/AIModule/Public/DefaultManagerInstanceTracker.h
+++ b/Engine/Source/Runtime/AIModule/Public/DefaultManagerInstanceTracker.h
@@ -16,7 +16,6 @@ struct TDefaultManagerInstanceTracker
 
 	TDefaultManagerInstanceTracker()
 	{
-		FWorldDelegates::OnPostWorldCleanup.AddRaw(this, &TDefaultManagerInstanceTracker::OnPostWorldCleanup);
 	}
 
 	FManager* GetManagerInstance(UWorld& World)
@@ -37,6 +36,10 @@ struct TDefaultManagerInstanceTracker
 			NewInstance = NewObject<FManager>(&World);
 			NewInstance->AddToRoot();
 			WorldToInstanceMap.Add(&World, NewInstance);
+			if (!FWorldDelegates::OnPostWorldCleanup.IsBoundToObject(this))
+			{
+				FWorldDelegates::OnPostWorldCleanup.AddRaw(this, &TDefaultManagerInstanceTracker::OnPostWorldCleanup);
+			}
 		}
 		else
 		{
