Added Log to make it clearer when a perception component is not attached to a controller. Some of the sense to do not function under these circumstances.

diff --git a/Engine/Source/Runtime/AIModule/Private/Perception/AIPerceptionComponent.cpp b/Engine/Source/Runtime/AIModule/Private/Perception/AIPerceptionComponent.cpp
--- a/Engine/Source/Runtime/AIModule/Private/Perception/AIPerceptionComponent.cpp
+++ b/Engine/Source/Runtime/AIModule/Private/Perception/AIPerceptionComponent.cpp
@@ -185,6 +185,9 @@ void UAIPerceptionComponent::OnRegister()
 	{
 		Owner->OnEndPlay.AddUniqueDynamic(this, &UAIPerceptionComponent::OnOwnerEndPlay);
 		AIOwner = Cast<AAIController>(Owner);
+
+		// Whilst it should be possible with some code changes, to make perception components work when being added to other AActors than AIControllers, it's not something Epic support.
+		UE_CVLOG_UELOG(!AIOwner && Owner->GetWorld() && (Owner->GetWorld()->WorldType != EWorldType::Editor), Owner, LogAIPerception, Warning, TEXT("%s: Perception Component is being registered with %s, they are designed to work with AAIControllers!"), ANSI_TO_TCHAR(__FUNCTION__), *Owner->GetName());
 	}
 
 	UAIPerceptionSystem* AIPerceptionSys = UAIPerceptionSystem::GetCurrent(GetWorld());
