GitHub 8453 : Avoid nullptr dereference when the Crowd Manager is missing

#8453

diff --git a/Engine/Source/Runtime/AIModule/Private/Navigation/CrowdFollowingComponent.cpp b/Engine/Source/Runtime/AIModule/Private/Navigation/CrowdFollowingComponent.cpp
--- a/Engine/Source/Runtime/AIModule/Private/Navigation/CrowdFollowingComponent.cpp
+++ b/Engine/Source/Runtime/AIModule/Private/Navigation/CrowdFollowingComponent.cpp
@@ -413,6 +413,8 @@ void UCrowdFollowingComponent::ApplyCrowdAgentPosition(const FVector& NewPositio
 
 void UCrowdFollowingComponent::SetCrowdSimulationState(ECrowdSimulationState NewState)
 {
+	static FString CrowdSimulationDesc[] = { TEXT("Enabled"), TEXT("ObstacleOnly"), TEXT("Disabled") };
+
 	if (NewState == SimulationState)
 	{
 		return;
@@ -425,13 +427,13 @@ void UCrowdFollowingComponent::SetCrowdSimulationState(ECrowdSimulationState New
 	}
 
 	UCrowdManager* Manager = UCrowdManager::GetCurrent(GetWorld());
-	if (Manager == NULL && NewState != ECrowdSimulationState::Disabled)
+	if (Manager == NULL)
 	{
-		UE_VLOG(GetOwner(), LogCrowdFollowing, Log, TEXT("Crowd manager can't be found, disabling simulation"));
-		NewState = ECrowdSimulationState::Disabled;
+		UE_VLOG(GetOwner(), LogCrowdFollowing, Log, TEXT("SetCrowdSimulation: NewState %s: Crowd manager can't be found, disabling simulation."), *CrowdSimulationDesc[static_cast<uint8>(NewState)]);
+		SimulationState = ECrowdSimulationState::Disabled;
+		return;
 	}
 
-	static FString CrowdSimulationDesc[] = { TEXT("Enabled"), TEXT("ObstacleOnly"), TEXT("Disabled") };
 	UE_VLOG(GetOwner(), LogCrowdFollowing, Log, TEXT("SetCrowdSimulation: %s"), *CrowdSimulationDesc[static_cast<uint8>(NewState)]);
 
 	const bool bNeedRegistration = (NewState != ECrowdSimulationState::Disabled);
