GitHub 8649 : Added ReportTouchEvent function for touch sense (AI module)

#8649

diff --git a/Engine/Source/Runtime/AIModule/Classes/Perception/AISense_Touch.h b/Engine/Source/Runtime/AIModule/Classes/Perception/AISense_Touch.h
--- a/Engine/Source/Runtime/AIModule/Classes/Perception/AISense_Touch.h
+++ b/Engine/Source/Runtime/AIModule/Classes/Perception/AISense_Touch.h
@@ -7,6 +7,7 @@
 #include "Perception/AISense.h"
 #include "AISense_Touch.generated.h"
 
+class IAIPerceptionListenerInterface;
 class UAISense_Touch;
 
 USTRUCT()
@@ -30,6 +31,8 @@ struct AIMODULE_API FAITouchEvent
 		: Location(EventLocation), TouchReceiver(InTouchReceiver), OtherActor(InOtherActor)
 	{
 	}
+
+	IAIPerceptionListenerInterface* GetTouchedActorAsPerceptionListener() const;
 };
 
 UCLASS(ClassGroup=AI)
@@ -43,6 +46,9 @@ class AIMODULE_API UAISense_Touch : public UAISense
 public:		
 	void RegisterEvent(const FAITouchEvent& Event);	
 
+	UFUNCTION(BlueprintCallable, Category = "AI|Perception", meta = (WorldContext = "WorldContextObject"))
+	static void ReportTouchEvent(UObject* WorldContextObject, AActor* TouchReceiver, AActor* OtherActor, FVector Location);
+
 protected:
 	virtual float Update() override;
 };
diff --git a/Engine/Source/Runtime/AIModule/Private/Perception/AISense_Touch.cpp b/Engine/Source/Runtime/AIModule/Private/Perception/AISense_Touch.cpp
--- a/Engine/Source/Runtime/AIModule/Private/Perception/AISense_Touch.cpp
+++ b/Engine/Source/Runtime/AIModule/Private/Perception/AISense_Touch.cpp
@@ -1,9 +1,30 @@
 // Copyright Epic Games, Inc. All Rights Reserved.
 
 #include "Perception/AISense_Touch.h"
+#include "GameFramework/Pawn.h"
+#include "GameFramework/Controller.h"
 #include "Perception/AIPerceptionListenerInterface.h"
 #include "Perception/AIPerceptionComponent.h"
 
+
+IAIPerceptionListenerInterface* FAITouchEvent::GetTouchedActorAsPerceptionListener() const
+{
+	IAIPerceptionListenerInterface* Listener = nullptr;
+	if (TouchReceiver)
+	{
+		Listener = Cast<IAIPerceptionListenerInterface>(TouchReceiver);
+		if (Listener == nullptr)
+		{
+			APawn* ListenerAsPawn = Cast<APawn>(TouchReceiver);
+			if (ListenerAsPawn)
+			{
+				Listener = Cast<IAIPerceptionListenerInterface>(ListenerAsPawn->GetController());
+			}
+		}
+	}
+	return Listener;
+}
+
 UAISense_Touch::UAISense_Touch(const FObjectInitializer& ObjectInitializer) :
 	Super(ObjectInitializer)
 {
@@ -17,7 +38,7 @@ float UAISense_Touch::Update()
 	{
 		if (Event.TouchReceiver != NULL && Event.OtherActor != NULL)
 		{
-			IAIPerceptionListenerInterface* PerceptionListener = Cast<IAIPerceptionListenerInterface>(Event.TouchReceiver);
+			IAIPerceptionListenerInterface* PerceptionListener = Event.GetTouchedActorAsPerceptionListener();
 			if (PerceptionListener != NULL)
 			{
 				UAIPerceptionComponent* PerceptionComponent = PerceptionListener->GetPerceptionComponent();
@@ -46,3 +67,13 @@ void UAISense_Touch::RegisterEvent(const FAITouchEvent& Event)
 
 	RequestImmediateUpdate();
 }
+
+void UAISense_Touch::ReportTouchEvent(UObject* WorldContextObject, AActor* TouchReceiver, AActor* OtherActor, FVector Location)
+{
+	UAIPerceptionSystem* PerceptionSystem = UAIPerceptionSystem::GetCurrent(WorldContextObject);
+	if (PerceptionSystem)
+	{
+		FAITouchEvent Event(TouchReceiver, OtherActor, Location);
+		PerceptionSystem->OnEvent(Event);
+	}
+}
