Fixed FGameplayDebuggerCategory_EQS's FEQSRenderingDebugDrawDelegateHelper not getting set up properly

The issue resulted in EQS's item scores not getting displayed in gameplay debugger.

diff --git a/Engine/Source/Runtime/AIModule/Private/GameplayDebugger/GameplayDebuggerCategory_EQS.cpp b/Engine/Source/Runtime/AIModule/Private/GameplayDebugger/GameplayDebuggerCategory_EQS.cpp
--- a/Engine/Source/Runtime/AIModule/Private/GameplayDebugger/GameplayDebuggerCategory_EQS.cpp
+++ b/Engine/Source/Runtime/AIModule/Private/GameplayDebugger/GameplayDebuggerCategory_EQS.cpp
@@ -107,6 +107,7 @@ FDebugRenderSceneProxy* FGameplayDebuggerCategory_EQS::CreateDebugSceneProxy(con
 		FEQSSceneProxy* EQSSceneProxy = new FEQSSceneProxy(*InComponent, ViewFlagName, QueryData.SolidSpheres, QueryData.Texts);
 
 		auto* OutDelegateHelper2 = new FEQSRenderingDebugDrawDelegateHelper();
+		OutDelegateHelper2->SetupFromProxy(EQSSceneProxy);
 		OutDelegateHelper2->InitDelegateHelper(EQSSceneProxy);
 		OutDelegateHelper = OutDelegateHelper2;
 
