[EQSTestingPawn] refresh display and query on Undo/Redo

diff --git a/Engine/Source/Runtime/AIModule/Classes/EnvironmentQuery/EQSTestingPawn.h b/Engine/Source/Runtime/AIModule/Classes/EnvironmentQuery/EQSTestingPawn.h
--- a/Engine/Source/Runtime/AIModule/Classes/EnvironmentQuery/EQSTestingPawn.h
+++ b/Engine/Source/Runtime/AIModule/Classes/EnvironmentQuery/EQSTestingPawn.h
@@ -86,6 +86,8 @@ public:
 #if WITH_EDITOR
 	virtual void PostEditChangeProperty(struct FPropertyChangedEvent& PropertyChangedEvent) override;
 	virtual void PostEditMove(bool bFinished) override;
+	virtual void PostTransacted(const FTransactionObjectEvent& TransactionEvent) override;
+	void OnPropertyChanged(const FName PropName);
 #endif // WITH_EDITOR
 
 	// IEQSQueryResultSourceInterface start
diff --git a/Engine/Source/Runtime/AIModule/Private/EnvironmentQuery/EQSTestingPawn.cpp b/Engine/Source/Runtime/AIModule/Private/EnvironmentQuery/EQSTestingPawn.cpp
--- a/Engine/Source/Runtime/AIModule/Private/EnvironmentQuery/EQSTestingPawn.cpp
+++ b/Engine/Source/Runtime/AIModule/Private/EnvironmentQuery/EQSTestingPawn.cpp
@@ -312,7 +312,7 @@ const FEnvQueryInstance* AEQSTestingPawn::GetQueryInstance() const
 
 #if WITH_EDITOR
 
-void AEQSTestingPawn::PostEditChangeProperty( FPropertyChangedEvent& PropertyChangedEvent)
+void AEQSTestingPawn::OnPropertyChanged(const FName PropName)
 {
 	static const FName NAME_QueryTemplate = GET_MEMBER_NAME_CHECKED(AEQSTestingPawn, QueryTemplate);
 	static const FName NAME_StepToDebugDraw = GET_MEMBER_NAME_CHECKED(AEQSTestingPawn, StepToDebugDraw);
@@ -320,41 +320,45 @@ void AEQSTestingPawn::PostEditChangeProperty( FPropertyChangedEvent& PropertyCha
 	static const FName NAME_ShouldBeVisibleInGame = GET_MEMBER_NAME_CHECKED(AEQSTestingPawn, bShouldBeVisibleInGame);
 	static const FName NAME_QueryingMode = GET_MEMBER_NAME_CHECKED(AEQSTestingPawn, QueryingMode);
 
-	if (PropertyChangedEvent.Property != NULL)
+	if (PropName == NAME_QueryTemplate || PropName == NAME_QueryConfig)
 	{
-		const FName PropName = PropertyChangedEvent.MemberProperty->GetFName();
-		if (PropName == NAME_QueryTemplate || PropName == NAME_QueryConfig)
+		if (QueryTemplate)
 		{
-			if (QueryTemplate)
-			{
-				QueryTemplate->CollectQueryParams(*this, QueryConfig);
-			}
-
-			RunEQSQuery();
-		}
-		else if (PropName == NAME_StepToDebugDraw)
-		{
-			StepToDebugDraw  = FMath::Clamp(StepToDebugDraw, 0, StepResults.Num() - 1 );
-			UpdateDrawing();
-		}
-		else if (PropName == GET_MEMBER_NAME_CHECKED(AEQSTestingPawn, bDrawFailedItems) ||
-			PropName == GET_MEMBER_NAME_CHECKED(AEQSTestingPawn, HighlightMode))
-		{
-			UpdateDrawing();
+			QueryTemplate->CollectQueryParams(*this, QueryConfig);
 		}
-		else if (PropName == NAME_ShouldBeVisibleInGame)
-		{
-			UBillboardComponent* SpriteComponent = FindComponentByClass<UBillboardComponent>();
-			if (SpriteComponent != NULL)
-			{
-				SpriteComponent->bHiddenInGame = !bShouldBeVisibleInGame;
-			}
-		}
-		else if (PropName == NAME_QueryingMode)
+
+		RunEQSQuery();
+	}
+	else if (PropName == NAME_StepToDebugDraw)
+	{
+		StepToDebugDraw  = FMath::Clamp(StepToDebugDraw, 0, StepResults.Num() - 1 );
+		UpdateDrawing();
+	}
+	else if (PropName == GET_MEMBER_NAME_CHECKED(AEQSTestingPawn, bDrawFailedItems) ||
+		PropName == GET_MEMBER_NAME_CHECKED(AEQSTestingPawn, HighlightMode))
+	{
+		UpdateDrawing();
+	}
+	else if (PropName == NAME_ShouldBeVisibleInGame)
+	{
+		UBillboardComponent* SpriteComponent = FindComponentByClass<UBillboardComponent>();
+		if (SpriteComponent != NULL)
 		{
-			RunEQSQuery();
+			SpriteComponent->bHiddenInGame = !bShouldBeVisibleInGame;
 		}
 	}
+	else if (PropName == NAME_QueryingMode)
+	{
+		RunEQSQuery();
+	}
+}
+
+void AEQSTestingPawn::PostEditChangeProperty( FPropertyChangedEvent& PropertyChangedEvent)
+{
+	if (PropertyChangedEvent.Property != nullptr)
+	{
+		OnPropertyChanged(PropertyChangedEvent.MemberProperty->GetFName());
+	}
 
 	Super::PostEditChangeProperty(PropertyChangedEvent);
 }
@@ -369,4 +373,17 @@ void AEQSTestingPawn::PostEditMove(bool bFinished)
 	}
 }
 
+void AEQSTestingPawn::PostTransacted(const FTransactionObjectEvent& TransactionEvent)
+{
+	Super::PostTransacted(TransactionEvent);
+
+	if (TransactionEvent.GetEventType() == ETransactionObjectEventType::UndoRedo)		
+	{
+		for (const FName PropertyName : TransactionEvent.GetChangedProperties())
+		{
+			OnPropertyChanged(PropertyName);
+		}
+	}
+}
+
 #endif // WITH_EDITOR
