AI GeomUtils
- common geom/math helper functions

diff --git a/Engine/Source/Runtime/AIModule/Private/GeomUtils.cpp b/Engine/Source/Runtime/AIModule/Private/GeomUtils.cpp
new file mode 100644
--- /dev/null
+++ b/Engine/Source/Runtime/AIModule/Private/GeomUtils.cpp
@@ -0,0 +1,131 @@
+﻿// Copyright Epic Games, Inc. All Rights Reserved.
+
+#include "GeomUtils.h"
+
+namespace UE::AI
+{
+
+	bool IntersectSegmentPoly2D(const FVector& Start, const FVector& End, TConstArrayView<FVector> Poly,
+								FVector2D::FReal& OutTMin, FVector2D::FReal& OutTMax, int32& OutSegMin, int32& OutSegMax)
+	{
+		using FReal = FVector::FReal;
+
+		OutTMin = 0.0;
+		OutTMax = 1.0;
+		OutSegMin = INDEX_NONE;
+		OutSegMax = INDEX_NONE;
+
+		const FVector Dir = End - Start;
+		const int32 NumVerts = Poly.Num();
+		
+		for (int NextIndex = 0, Index = NumVerts - 1; NextIndex < NumVerts; Index = NextIndex++)
+		{
+			const FVector Edge = Poly[NextIndex] - Poly[Index];
+			const FVector Diff = Start - Poly[Index];
+			
+			const FReal N = Cross2D(Edge, Diff);
+			const FReal D = Cross2D(Dir, Edge);
+			
+			if (FMath::Abs(D) < UE_KINDA_SMALL_NUMBER)
+			{
+				// S is nearly parallel to this edge
+				if (N > 0.0)
+				{
+					return false;
+				}
+				continue;
+			}
+			
+			const FReal T = N / D;
+			
+			if (D > 0.0)
+			{
+				// segment S is entering across this edge
+				if (T > OutTMin)
+				{
+					OutTMin = T;
+					OutSegMin = Index;
+					// S enters after leaving polygon
+					if (OutTMin > OutTMax)
+					{
+						return false;
+					}
+				}
+			}
+			else
+			{
+				// segment S is leaving across this edge
+				if (T < OutTMax)
+				{
+					OutTMax = T;
+					OutSegMax = Index;
+					// S leaves before entering polygon
+					if (OutTMax < OutTMin)
+					{
+						return false;
+					}
+				}
+			}
+		}
+		
+		return true;
+	}
+
+	FVector2D InvBilinear2D(const FVector Point, const FVector VertexA, const FVector VertexB, const FVector VertexC, const FVector VertexD)
+	{
+		// For more info how inverse bilinear works, see:
+		//  - https://www.reedbeta.com/blog/quadrilateral-interpolation-part-2/
+		//  - https://iquilezles.org/articles/ibilinear/
+		
+		using FReal = FVector::FReal;
+		
+		const FVector E = VertexB - VertexA;
+		const FVector F = VertexD - VertexA;
+		const FVector G = VertexA - VertexB + VertexC - VertexD;
+		const FVector H = Point - VertexA;
+
+		// The algorithm is sensitive to the floating point precisions because we're squaring squares.
+		// Scaling the coefficients will help with the precision.
+		const FReal K = 1.0 / FMath::Max3(1.0, E.SquaredLength(), F.SquaredLength());
+
+		// Coefficient for solving the quadratic equation.
+		const FReal A = Cross2D(G, F) * K;
+		const FReal B = (Cross2D(E, F) + Cross2D(H, G)) * K;
+		const FReal C = Cross2D(H, E) * K; 
+
+		FReal U = 0.0;
+		FReal V = 0.0;
+		
+		// if edges are parallel, this is a linear equation
+		if (FMath::Abs(A) > UE_KINDA_SMALL_NUMBER)
+		{
+			// Quadratic solution
+			const FReal W = B * B - 4.0 * A * C;
+			if (W < 0)
+			{
+				return FVector2D(-1, -1);
+			}
+			V = (-B - FMath::Sqrt(W)) / (2.0 * A);
+		}
+		else if (FMath::Abs(B) > UE_KINDA_SMALL_NUMBER)
+		{
+			// Linear solution
+			V = -C / B;
+		}
+
+		// Calculate U using V, use the larger denominator for better precision.
+		const FReal DenomX = E.X + G.X * V;
+		const FReal DenomY = E.Y + G.Y * V;
+		if (FMath::Abs(DenomX) > FMath::Abs(DenomY))
+		{
+			U = (H.X - F.X * V) / DenomX;
+		}
+		else
+		{
+			U = (H.Y - F.Y * V) / DenomY;
+		}
+		
+		return FVector2D(U, V);
+	}
+
+}; // UE::AI
\ No newline at end of file
diff --git a/Engine/Source/Runtime/AIModule/Public/GeomUtils.h b/Engine/Source/Runtime/AIModule/Public/GeomUtils.h
new file mode 100644
--- /dev/null
+++ b/Engine/Source/Runtime/AIModule/Public/GeomUtils.h
@@ -0,0 +1,173 @@
+// Copyright Epic Games, Inc. All Rights Reserved.
+
+#pragma once
+
+#include "Math/Vector.h"
+#include "Math/Vector2D.h"
+
+namespace UE::AI
+{
+	/** @return 2D cross product. Using X and Y components of 3D vectors. */
+	inline FVector::FReal Cross2D(const FVector& A, const FVector& B)
+	{
+		return A.X * B.Y - A.Y * B.X;
+	}
+
+	/** @return 2D cross product. */
+	inline FVector2D::FReal Cross2D(const FVector2D& A, const FVector2D& B)
+	{
+		return A.X * B.Y - A.Y * B.X;
+	}
+
+	/** @return 2D area of triangle. Using X and Y components of 3D vectors. */
+	inline FVector::FReal TriArea2D(const FVector& A, const FVector& B, const FVector& C)
+	{
+		const FVector AB = B - A;
+		const FVector AC = C - A;
+		return (AC.X * AB.Y - AB.X * AC.Y) * 0.5;
+	}
+
+	/** @return 2D area of triangle. */
+	inline FVector2D::FReal TriArea2D(const FVector2D& A, const FVector2D& B, const FVector2D& C)
+	{
+		const FVector2D AB = B - A;
+		const FVector2D AC = C - A;
+		return (AC.X * AB.Y - AB.X * AC.Y) * 0.5;
+	}
+
+	/** @return value in range [0..1] of the 'Point' project on segment 'Start-End'. Using X and Y components of 3D vectors. */
+	inline FVector2D::FReal ProjectPointOnSegment2D(const FVector Point, const FVector Start, const FVector End)
+	{
+		using FReal = FVector::FReal;
+		
+		const FVector2D Seg(End - Start);
+		const FVector2D Dir(Point - Start);
+		const FReal D = Seg.SquaredLength();
+		const FReal T = FVector2D::DotProduct(Seg, Dir);
+		
+		if (T < 0.0)
+		{
+			return 0.0;
+		}
+		else if (T > D)
+		{
+			return 1.0;
+		}
+		
+		return D > UE_KINDA_SMALL_NUMBER ? (T / D) : 0.0;
+	}
+
+	/** @return value of the 'Point' project on infinite line defined by segment 'Start-End'. Using X and Y components of 3D vectors. */
+	inline FVector::FReal ProjectPointOnLine2D(const FVector Point, const FVector Start, const FVector End)
+	{
+		using FReal = FVector::FReal;
+		
+		const FVector2D Seg(End - Start);
+		const FVector2D Dir(Point - Start);
+		const FReal D = Seg.SquaredLength();
+		const FReal T = FVector2D::DotProduct(Seg, Dir);
+		return D > UE_KINDA_SMALL_NUMBER ? (T / D) : 0.0;
+	}
+
+	/** @return signed distance of the 'Point' to infinite line defined by segment 'Start-End'. Using X and Y components of 3D vectors. */
+	inline FVector::FReal SignedDistancePointLine2D(const FVector Point, const FVector Start, const FVector End)
+	{
+		using FReal = FVector::FReal;
+		
+		const FVector2D Seg(End - Start);
+		const FVector2D Dir(Point - Start);
+		const FReal Nom = Cross2D(Seg, Dir);
+		const FReal Den = Seg.SquaredLength();
+		const FReal Dist = Den > UE_KINDA_SMALL_NUMBER ? (Nom / FMath::Sqrt(Den)) : 0.0;
+		return Dist;
+	}
+
+	/**
+	 * Intersects infinite lines defined by segments A and B in 2D. Using X and Y components of 3D vectors.
+	 * @param StartA start point of segment A
+	 * @param EndA end point of segment A
+	 * @param StartB start point of segment B
+	 * @param EndB end point of segment B
+	 * @param OutTA intersection value along segment A
+	 * @param OutTB intersection value along segment B
+	 * @return if segments A and B intersect in 2D
+	 */
+	inline bool IntersectLineLine2D(const FVector& StartA, const FVector& EndA, const FVector& StartB, const FVector& EndB, FVector2D::FReal& OutTA, FVector2D::FReal& OutTB)
+	{
+		using FReal = FVector::FReal;
+
+		const FVector U = EndA - StartA;
+		const FVector V = EndB - StartB;
+		const FVector W = StartA - StartB;
+		
+		const FReal D = Cross2D(U, V);
+		if (FMath::Abs(D) < UE_KINDA_SMALL_NUMBER)
+		{
+			OutTA = 0.0;
+			OutTB = 0.0;
+			return false;
+		}
+		
+		OutTA = Cross2D(V, W) / D;
+		OutTB = Cross2D(U, W) / D;
+		
+		return true;
+	}
+
+	/**
+	 * Calculates intersection of segment Start-End with convex polygon Poly in 2D. Using X and Y components of 3D vectors.
+	 * @param Start start point of the segment
+	 * @param End end point of the segment
+	 * @param Poly convex polygon
+	 * @param OutTMin value along the segment of the first intersection point [0..1]
+	 * @param OutTMax value along the segment of the second intersection point [0..1]
+	 * @param OutSegMin index of the polygon segment of the first intersection point
+	 * @param OutSegMax index of the polygon segment of the second intersection point
+	 * @return true if the segment inside or intersects with the polygon.
+	 */
+	bool IntersectSegmentPoly2D(const FVector& Start, const FVector& End, TConstArrayView<FVector> Poly,
+								FVector2D::FReal& OutTMin, FVector2D::FReal& OutTMax, int32& OutSegMin, int32& OutSegMax);
+
+	/**
+	 * Interpolates bilinear patch A,B,C,D. U interpolates from A->B, and C->D, and V interpolates from AB->CD.
+	 * @param UV interpolation coordinates [0..1] range
+	 * @param VertexA first corner
+	 * @param VertexB second corner
+	 * @param VertexC third corner
+	 * @param VertexD fourth corner
+	 * @return interpolated value.
+	 */
+	inline FVector Bilinear(const FVector2D UV, const FVector VertexA, const FVector VertexB, const FVector VertexC, const FVector VertexD)
+	{
+		const FVector AB = FMath::Lerp(VertexA, VertexB, UV.X);
+		const FVector CD = FMath::Lerp(VertexD, VertexC, UV.X);
+		return FMath::Lerp(AB, CD, UV.Y);
+	}
+
+	/**
+	 * Finds the UV coordinates of the 'Point' on bilinear patch A,B,C,D. U interpolates from A->B, and C->D, and V interpolates from AB->CD.
+	 * @param Point location inside or close to the bilinear patch
+	 * @param VertexA first corner
+	 * @param VertexB second corner
+	 * @param VertexC third corner
+	 * @param VertexD fourth corner
+	 * @return UV interpolation coordinates of the 'Point'.
+	 */
+	FVector2D InvBilinear2D(const FVector Point, const FVector VertexA, const FVector VertexB, const FVector VertexC, const FVector VertexD);
+
+	/**
+	 * Finds the UV coordinates of the 'Point' on bilinear patch A,B,C,D. U interpolates from A->B, and C->D, and V interpolates from AB->CD.
+	 * The UV coordinate is clamped to [0..1] range after inversion.
+	 * @param Point location inside or close to the bilinear patch
+	 * @param VertexA first corner
+	 * @param VertexB second corner
+	 * @param VertexC third corner
+	 * @param VertexD fourth corner
+	 * @return UV interpolation coordinates of the 'Point' in [0..1] range.
+	 */
+	inline FVector2D InvBilinear2DClamped(const FVector Point, const FVector VertexA, const FVector VertexB, const FVector VertexC, const FVector VertexD)
+	{
+		return InvBilinear2D(Point, VertexA, VertexB, VertexC, VertexD).ClampAxes(0.0, 1.0);
+	}
+
+}; // UE::AI
\ No newline at end of file
