Suspending all decorator branch operation while trasitionning
Support reactivation of a deactivated branch (decorator condition toggled from not pass to pass)

diff --git a/Engine/Source/Developer/AITestSuite/Classes/BehaviorTree/TestBTService_Log.h b/Engine/Source/Developer/AITestSuite/Classes/BehaviorTree/TestBTService_Log.h
--- a/Engine/Source/Developer/AITestSuite/Classes/BehaviorTree/TestBTService_Log.h
+++ b/Engine/Source/Developer/AITestSuite/Classes/BehaviorTree/TestBTService_Log.h
@@ -29,6 +29,9 @@ class UTestBTService_Log : public UBTService
 	UPROPERTY()
 	int32 LogTick;
 
+	UPROPERTY()
+	bool bToggleValue;
+
 	virtual void OnBecomeRelevant(UBehaviorTreeComponent& OwnerComp, uint8* NodeMemory) override;
 	virtual void OnCeaseRelevant(UBehaviorTreeComponent& OwnerComp, uint8* NodeMemory) override;
 	virtual void TickNode(UBehaviorTreeComponent& OwnerComp, uint8* NodeMemory, float DeltaSeconds) override;
diff --git a/Engine/Source/Developer/AITestSuite/Classes/BehaviorTree/TestBTTask_ToggleFlag.h b/Engine/Source/Developer/AITestSuite/Classes/BehaviorTree/TestBTTask_ToggleFlag.h
new file mode 100644
--- /dev/null
+++ b/Engine/Source/Developer/AITestSuite/Classes/BehaviorTree/TestBTTask_ToggleFlag.h
@@ -0,0 +1,28 @@
+// Copyright Epic Games, Inc. All Rights Reserved.
+#pragma once
+
+#include "BehaviorTree/BTTaskNode.h"
+#include "TestBTTask_ToggleFlag.generated.h"
+
+UCLASS(meta=(HiddenNode))
+class UTestBTTask_ToggleFlag : public UBTTaskNode
+{
+	GENERATED_BODY()
+
+public:
+	UTestBTTask_ToggleFlag(const FObjectInitializer& ObjectInitializer);
+
+	virtual EBTNodeResult::Type ExecuteTask(UBehaviorTreeComponent& OwnerComp, uint8* NodeMemory) override;
+
+	UPROPERTY()
+	int32 LogIndex;
+
+	UPROPERTY()
+	FName KeyName;
+
+	UPROPERTY()
+	int32 NumToggles;
+
+	UPROPERTY()
+	TEnumAsByte<EBTNodeResult::Type> TaskResult;
+};
diff --git a/Engine/Source/Developer/AITestSuite/Private/BehaviorTree/TestBTService_Log.cpp b/Engine/Source/Developer/AITestSuite/Private/BehaviorTree/TestBTService_Log.cpp
--- a/Engine/Source/Developer/AITestSuite/Private/BehaviorTree/TestBTService_Log.cpp
+++ b/Engine/Source/Developer/AITestSuite/Private/BehaviorTree/TestBTService_Log.cpp
@@ -16,6 +16,7 @@ UTestBTService_Log::UTestBTService_Log(const FObjectInitializer& ObjectInitializ
 	LogTick = INDEX_NONE;
 	KeyNameBecomeRelevant = NAME_None;
 	KeyNameCeaseRelevant = NAME_None;
+	bToggleValue = false;
 
 	// Force the service to tick every frame
 	Interval = 0.0f;
@@ -28,7 +29,7 @@ void UTestBTService_Log::OnBecomeRelevant(UBehaviorTreeComponent& OwnerComp, uin
 
 	if (KeyNameBecomeRelevant != NAME_None)
 	{
-		OwnerComp.GetBlackboardComponent()->SetValueAsBool(KeyNameBecomeRelevant, true);
+		OwnerComp.GetBlackboardComponent()->SetValueAsBool(KeyNameBecomeRelevant, bToggleValue ? !OwnerComp.GetBlackboardComponent()->GetValueAsBool(KeyNameBecomeRelevant) : true);
 	}
 
 	if (LogActivation >= 0)
@@ -43,7 +44,7 @@ void UTestBTService_Log::OnCeaseRelevant(UBehaviorTreeComponent& OwnerComp, uint
 
 	if (KeyNameCeaseRelevant != NAME_None)
 	{
-		OwnerComp.GetBlackboardComponent()->SetValueAsBool(KeyNameCeaseRelevant, true);
+		OwnerComp.GetBlackboardComponent()->SetValueAsBool(KeyNameCeaseRelevant, bToggleValue ? !OwnerComp.GetBlackboardComponent()->GetValueAsBool(KeyNameCeaseRelevant) : true);
 	}
 
 	if (LogDeactivation >= 0)
@@ -58,7 +59,7 @@ void UTestBTService_Log::TickNode(UBehaviorTreeComponent& OwnerComp, uint8* Node
 
 	if (KeyNameTick != NAME_None)
 	{
-		OwnerComp.GetBlackboardComponent()->SetValueAsBool(KeyNameTick, true);
+		OwnerComp.GetBlackboardComponent()->SetValueAsBool(KeyNameTick, bToggleValue ? !OwnerComp.GetBlackboardComponent()->GetValueAsBool(KeyNameTick) : true);
 	}
 
 	if (LogTick >= 0)
diff --git a/Engine/Source/Developer/AITestSuite/Private/BehaviorTree/TestBTTask_ToggleFlag.cpp b/Engine/Source/Developer/AITestSuite/Private/BehaviorTree/TestBTTask_ToggleFlag.cpp
new file mode 100644
--- /dev/null
+++ b/Engine/Source/Developer/AITestSuite/Private/BehaviorTree/TestBTTask_ToggleFlag.cpp
@@ -0,0 +1,22 @@
+// Copyright Epic Games, Inc. All Rights Reserved.
+
+#include "BehaviorTree/TestBTTask_ToggleFlag.h"
+#include "BehaviorTree/Blackboard/BlackboardKeyType_Bool.h"
+#include "BehaviorTree/BlackboardComponent.h"
+
+UTestBTTask_ToggleFlag::UTestBTTask_ToggleFlag(const FObjectInitializer& ObjectInitializer) : Super(ObjectInitializer)
+{
+	NodeName = "Log";
+	TaskResult = EBTNodeResult::Succeeded;
+	KeyName = TEXT("Bool1");
+	NumToggles = 1;
+}
+
+EBTNodeResult::Type UTestBTTask_ToggleFlag::ExecuteTask(UBehaviorTreeComponent& OwnerComp, uint8* NodeMemory)
+{
+	for (int i = 0; i < NumToggles; i++)
+	{
+		OwnerComp.GetBlackboardComponent()->SetValue<UBlackboardKeyType_Bool>(KeyName, !OwnerComp.GetBlackboardComponent()->GetValue<UBlackboardKeyType_Bool>(KeyName));
+	}
+	return TaskResult;
+}
\ No newline at end of file
diff --git a/Engine/Source/Developer/AITestSuite/Private/Tests/BTTest.cpp b/Engine/Source/Developer/AITestSuite/Private/Tests/BTTest.cpp
--- a/Engine/Source/Developer/AITestSuite/Private/Tests/BTTest.cpp
+++ b/Engine/Source/Developer/AITestSuite/Private/Tests/BTTest.cpp
@@ -450,6 +450,245 @@ struct FAITest_BTAbortingDuringServiceTick : public FAITest_SimpleBT
 };
 IMPLEMENT_AI_LATENT_TEST(FAITest_BTAbortingDuringServiceTick, "System.AI.Behavior Trees.Abort: during service tick")
 
+
+struct FAITest_BTRestartingTheAbortedBranch : public FAITest_SimpleBT
+{
+	FAITest_BTRestartingTheAbortedBranch()
+	{
+		UBTCompositeNode& CompNode = FBTBuilder::AddSelector(*BTAsset);
+		{
+			UBTCompositeNode& CompNode2 = FBTBuilder::AddSelector(CompNode);
+			{
+				FBTBuilder::WithDecoratorBlackboard(CompNode, EBasicKeyOperation::NotSet, EBTFlowAbortMode::Both, TEXT("Bool1"));
+				UBTCompositeNode& CompNode3 = FBTBuilder::AddSelector(CompNode2);
+				{
+					FBTBuilder::AddTaskToggleFlag(CompNode3, EBTNodeResult::InProgress, TEXT("Bool1"), 2/*NumOfToggles*/);
+					FBTBuilder::WithDecoratorBlackboard(CompNode3, EBasicKeyOperation::NotSet, EBTFlowAbortMode::Both, TEXT("Bool2"));
+					FBTBuilder::WithTaskServiceLog(CompNode3, 1 /*ActivationIndex*/, 2 /*DeactivationIndex*/, 3 /*TickIndex*/, NAME_None /*TickBoolKeyName*/, false /*bCallTickOnSearchStart*/,
+											    NAME_None /*BecomeRelevantBoolKeyeyName*/, TEXT("Bool2") /*CeaseRelevantBoolKeyName*/);
+					FBTBuilder::AddTask(CompNode3, 5, EBTNodeResult::Succeeded);
+				}
+			}
+			FBTBuilder::AddTask(CompNode, 6, EBTNodeResult::Succeeded);
+		}
+		ExpectedResult.Add(1);
+		ExpectedResult.Add(3);
+		ExpectedResult.Add(3);
+		ExpectedResult.Add(2);
+		ExpectedResult.Add(1);
+		ExpectedResult.Add(3);
+		ExpectedResult.Add(3);
+		ExpectedResult.Add(2);
+		ExpectedResult.Add(5);
+	}
+};
+IMPLEMENT_AI_LATENT_TEST(FAITest_BTRestartingTheAbortedBranch, "System.AI.Behavior Trees.Toggle: restart the aborted branch")
+
+
+struct FAITest_BTAbortingTheRestarted : public FAITest_SimpleBT
+{
+	FAITest_BTAbortingTheRestarted()
+	{
+		UBTCompositeNode& CompNode = FBTBuilder::AddSelector(*BTAsset);
+		{
+			UBTCompositeNode& CompNode2 = FBTBuilder::AddSelector(CompNode);
+			{
+				FBTBuilder::WithDecoratorBlackboard(CompNode, EBasicKeyOperation::NotSet, EBTFlowAbortMode::Both, TEXT("Bool1"));
+				UBTCompositeNode& CompNode3 = FBTBuilder::AddSelector(CompNode2);
+				{
+					FBTBuilder::AddTaskToggleFlag(CompNode3, EBTNodeResult::InProgress, TEXT("Bool1"), 3/*NumOfToggles*/);
+					FBTBuilder::WithDecoratorBlackboard(CompNode3, EBasicKeyOperation::NotSet, EBTFlowAbortMode::Both, TEXT("Bool2"));
+					FBTBuilder::WithTaskServiceLog(CompNode3, 1 /*ActivationIndex*/, 2 /*DeactivationIndex*/, 3 /*TickIndex*/);
+					FBTBuilder::AddTask(CompNode3, 5, EBTNodeResult::Succeeded);
+				}
+			}
+			FBTBuilder::AddTask(CompNode, 6, EBTNodeResult::Succeeded);
+		}
+		ExpectedResult.Add(1);
+		ExpectedResult.Add(3);
+		ExpectedResult.Add(3);
+		ExpectedResult.Add(2);
+		ExpectedResult.Add(6);
+	}
+};
+IMPLEMENT_AI_LATENT_TEST(FAITest_BTAbortingTheRestarted, "System.AI.Behavior Trees.Toggle: reabort the restarted branch")
+
+
+struct FAITest_BTRestartingTheReaborted : public FAITest_SimpleBT
+{
+	FAITest_BTRestartingTheReaborted()
+	{
+		UBTCompositeNode& CompNode = FBTBuilder::AddSelector(*BTAsset);
+		{
+			UBTCompositeNode& CompNode2 = FBTBuilder::AddSelector(CompNode);
+			{
+				FBTBuilder::WithDecoratorBlackboard(CompNode, EBasicKeyOperation::NotSet, EBTFlowAbortMode::Both, TEXT("Bool1"));
+				UBTCompositeNode& CompNode3 = FBTBuilder::AddSelector(CompNode2);
+				{
+					FBTBuilder::AddTaskToggleFlag(CompNode3, EBTNodeResult::InProgress, TEXT("Bool1"), 4/*NumOfToggles*/);
+					FBTBuilder::WithDecoratorBlackboard(CompNode3, EBasicKeyOperation::NotSet, EBTFlowAbortMode::Both, TEXT("Bool2"));
+					FBTBuilder::WithTaskServiceLog(CompNode3, 1 /*ActivationIndex*/, 2 /*DeactivationIndex*/, 3 /*TickIndex*/, NAME_None /*TickBoolKeyName*/, false /*bCallTickOnSearchStart*/,
+						NAME_None /*BecomeRelevantBoolKeyeyName*/, TEXT("Bool2") /*CeaseRelevantBoolKeyName*/);
+					FBTBuilder::AddTask(CompNode3, 5, EBTNodeResult::Succeeded);
+				}
+			}
+			FBTBuilder::AddTask(CompNode, 6, EBTNodeResult::Succeeded);
+		}
+		ExpectedResult.Add(1);
+		ExpectedResult.Add(3);
+		ExpectedResult.Add(3);
+		ExpectedResult.Add(2);
+		ExpectedResult.Add(1);
+		ExpectedResult.Add(3);
+		ExpectedResult.Add(3);
+		ExpectedResult.Add(2);
+		ExpectedResult.Add(5);
+	}
+};
+IMPLEMENT_AI_LATENT_TEST(FAITest_BTRestartingTheReaborted, "System.AI.Behavior Trees.Toggle: restart the reaborted branch")
+
+
+struct FAITest_BTRestartingTheAbortedBranchDuringServiceTick : public FAITest_SimpleBT
+{
+	FAITest_BTRestartingTheAbortedBranchDuringServiceTick()
+	{
+		UBTCompositeNode& CompNode = FBTBuilder::AddSelector(*BTAsset);
+		{
+			UBTCompositeNode& CompNode2 = FBTBuilder::AddSelector(CompNode);
+			{
+				FBTBuilder::WithDecoratorBlackboard(CompNode, EBasicKeyOperation::NotSet, EBTFlowAbortMode::Both, TEXT("Bool1"));
+				UBTCompositeNode& CompNode3 = FBTBuilder::AddSelector(CompNode2);
+				{
+					FBTBuilder::AddTask(CompNode3, 4, EBTNodeResult::Succeeded);
+					FBTBuilder::WithDecoratorBlackboard(CompNode3, EBasicKeyOperation::NotSet, EBTFlowAbortMode::Both, TEXT("Bool2"));
+					FBTBuilder::WithTaskServiceLog(CompNode3, -1 /*ActivationIndex*/, -1 /*DeactivationIndex*/, -1 /*TickIndex*/, TEXT("Bool1") /*TickBoolKeyName*/, false /*bCallTickOnSearchStart*/,
+						NAME_None /*BecomeRelevantBoolKeyeyName*/, NAME_None /*CeaseRelevantBoolKeyName*/, true /*bToggleValues*/);
+					FBTBuilder::WithTaskServiceLog(CompNode3, -1 /*ActivationIndex*/, -1 /*DeactivationIndex*/, -1 /*TickIndex*/, TEXT("Bool1") /*TickBoolKeyName*/, false /*bCallTickOnSearchStart*/,
+						NAME_None /*BecomeRelevantBoolKeyeyName*/, NAME_None /*CeaseRelevantBoolKeyName*/, true /*bToggleValues*/);
+
+					FBTBuilder::WithTaskServiceLog(CompNode3, 1 /*ActivationIndex*/, 2 /*DeactivationIndex*/, 3 /*TickIndex*/, NAME_None /*TickBoolKeyName*/, false /*bCallTickOnSearchStart*/,
+						NAME_None /*BecomeRelevantBoolKeyeyName*/, TEXT("Bool2") /*CeaseRelevantBoolKeyName*/);
+					FBTBuilder::AddTask(CompNode3, 5, EBTNodeResult::Succeeded);
+				}
+			}
+			FBTBuilder::AddTask(CompNode, 6, EBTNodeResult::Succeeded);
+		}
+		ExpectedResult.Add(1);
+		ExpectedResult.Add(3);
+		ExpectedResult.Add(4);
+		ExpectedResult.Add(3);
+		ExpectedResult.Add(2);
+		ExpectedResult.Add(1);
+		ExpectedResult.Add(3);
+		ExpectedResult.Add(4);
+		ExpectedResult.Add(3);
+		ExpectedResult.Add(2);
+		ExpectedResult.Add(5);
+	}
+};
+IMPLEMENT_AI_LATENT_TEST(FAITest_BTRestartingTheAbortedBranchDuringServiceTick, "System.AI.Behavior Trees.Toggle: restart the aborted branch during service tick")
+
+
+struct FAITest_BTRestartTheAbortedBranchDuringServiceCeaseRelevant : public FAITest_SimpleBT
+{
+	FAITest_BTRestartTheAbortedBranchDuringServiceCeaseRelevant()
+	{
+		UBTCompositeNode& CompNode = FBTBuilder::AddSelector(*BTAsset);
+		{
+			UBTCompositeNode& CompNode2 = FBTBuilder::AddSelector(CompNode);
+			{
+				FBTBuilder::WithDecoratorBlackboard(CompNode, EBasicKeyOperation::NotSet, EBTFlowAbortMode::Both, TEXT("Bool1"));
+				UBTCompositeNode& CompNode3 = FBTBuilder::AddSelector(CompNode2);
+				{
+					FBTBuilder::AddTaskToggleFlag(CompNode3, EBTNodeResult::InProgress, TEXT("Bool1"), 1/*NumOfToggles*/);
+					FBTBuilder::WithDecoratorBlackboard(CompNode3, EBasicKeyOperation::NotSet, EBTFlowAbortMode::Both, TEXT("Bool2"));
+					FBTBuilder::WithTaskServiceLog(CompNode3, -1 /*ActivationIndex*/, -1 /*DeactivationIndex*/, -1 /*TickIndex*/, NAME_None /*TickBoolKeyName*/, false /*bCallTickOnSearchStart*/,
+						NAME_None /*BecomeRelevantBoolKeyeyName*/, TEXT("Bool1") /*CeaseRelevantBoolKeyName*/, true /*bToggleValues*/);
+
+					FBTBuilder::WithTaskServiceLog(CompNode3, 1 /*ActivationIndex*/, 2 /*DeactivationIndex*/, 3 /*TickIndex*/, NAME_None /*TickBoolKeyName*/, false /*bCallTickOnSearchStart*/,
+						NAME_None /*BecomeRelevantBoolKeyeyName*/, TEXT("Bool2") /*CeaseRelevantBoolKeyName*/);
+					FBTBuilder::AddTask(CompNode3, 5, EBTNodeResult::Succeeded);
+				}
+			}
+			FBTBuilder::AddTask(CompNode, 6, EBTNodeResult::Succeeded);
+		}
+		ExpectedResult.Add(1);
+		ExpectedResult.Add(3);
+		ExpectedResult.Add(3);
+		ExpectedResult.Add(2);
+		ExpectedResult.Add(6);
+		ExpectedResult.Add(5);
+	}
+};
+IMPLEMENT_AI_LATENT_TEST(FAITest_BTRestartTheAbortedBranchDuringServiceCeaseRelevant, "System.AI.Behavior Trees.Toggle: restart the aborted branch during service cease relevant")
+
+struct FAITest_BTRestartTheReabortedBranchDuringServiceCeaseRelevant : public FAITest_SimpleBT
+{
+	FAITest_BTRestartTheReabortedBranchDuringServiceCeaseRelevant()
+	{
+		UBTCompositeNode& CompNode = FBTBuilder::AddSelector(*BTAsset);
+		{
+			UBTCompositeNode& CompNode2 = FBTBuilder::AddSelector(CompNode);
+			{
+				FBTBuilder::WithDecoratorBlackboard(CompNode, EBasicKeyOperation::NotSet, EBTFlowAbortMode::Both, TEXT("Bool1"));
+				UBTCompositeNode& CompNode3 = FBTBuilder::AddSelector(CompNode2);
+				{
+					FBTBuilder::AddTaskToggleFlag(CompNode3, EBTNodeResult::InProgress, TEXT("Bool1"), 1/*NumOfToggles*/);
+					FBTBuilder::WithDecoratorBlackboard(CompNode3, EBasicKeyOperation::NotSet, EBTFlowAbortMode::Both, TEXT("Bool2"));
+					FBTBuilder::WithTaskServiceLog(CompNode3, -1 /*ActivationIndex*/, -1 /*DeactivationIndex*/, -1 /*TickIndex*/, NAME_None /*TickBoolKeyName*/, false /*bCallTickOnSearchStart*/,
+						NAME_None /*BecomeRelevantBoolKeyeyName*/, TEXT("Bool1") /*CeaseRelevantBoolKeyName*/, true /*bToggleValues*/);
+					FBTBuilder::WithTaskServiceLog(CompNode3, -1 /*ActivationIndex*/, -1 /*DeactivationIndex*/, -1 /*TickIndex*/, NAME_None /*TickBoolKeyName*/, false /*bCallTickOnSearchStart*/,
+						NAME_None /*BecomeRelevantBoolKeyeyName*/, TEXT("Bool1") /*CeaseRelevantBoolKeyName*/, true /*bToggleValues*/);
+					FBTBuilder::WithTaskServiceLog(CompNode3, -1 /*ActivationIndex*/, -1 /*DeactivationIndex*/, -1 /*TickIndex*/, NAME_None /*TickBoolKeyName*/, false /*bCallTickOnSearchStart*/,
+						NAME_None /*BecomeRelevantBoolKeyeyName*/, TEXT("Bool1") /*CeaseRelevantBoolKeyName*/, true /*bToggleValues*/);
+
+					FBTBuilder::WithTaskServiceLog(CompNode3, 1 /*ActivationIndex*/, 2 /*DeactivationIndex*/, 3 /*TickIndex*/, NAME_None /*TickBoolKeyName*/, false /*bCallTickOnSearchStart*/,
+						NAME_None /*BecomeRelevantBoolKeyeyName*/, TEXT("Bool2") /*CeaseRelevantBoolKeyName*/);
+					FBTBuilder::AddTask(CompNode3, 5, EBTNodeResult::Succeeded);
+				}
+			}
+			FBTBuilder::AddTask(CompNode, 6, EBTNodeResult::Succeeded);
+		}
+		ExpectedResult.Add(1);
+		ExpectedResult.Add(3);
+		ExpectedResult.Add(3);
+		ExpectedResult.Add(2);
+		ExpectedResult.Add(6);
+		ExpectedResult.Add(5);
+	}
+};
+IMPLEMENT_AI_LATENT_TEST(FAITest_BTRestartTheReabortedBranchDuringServiceCeaseRelevant, "System.AI.Behavior Trees.Toggle: restart the reaborted branch during service cease relevant")
+
+struct FAITest_BTStartBranchDuringServiceCeaseRelevant : public FAITest_SimpleBT
+{
+	FAITest_BTStartBranchDuringServiceCeaseRelevant()
+	{
+		UBTCompositeNode& CompNode = FBTBuilder::AddSelector(*BTAsset);
+		{
+			UBTCompositeNode& CompNode2 = FBTBuilder::AddSelector(CompNode);
+			{
+				FBTBuilder::WithDecoratorBlackboard(CompNode, EBasicKeyOperation::NotSet, EBTFlowAbortMode::Both, TEXT("Bool1"));
+				UBTCompositeNode& CompNode3 = FBTBuilder::AddSelector(CompNode2);
+				{
+					FBTBuilder::AddTask(CompNode3, 5, EBTNodeResult::Succeeded);
+					FBTBuilder::WithDecoratorBlackboard(CompNode3, EBasicKeyOperation::Set, EBTFlowAbortMode::Both, TEXT("Bool2"));
+
+					FBTBuilder::AddTaskToggleFlag(CompNode3, EBTNodeResult::InProgress, TEXT("Bool1"), 1/*NumOfToggles*/);
+					FBTBuilder::WithTaskServiceLog(CompNode3, 1 /*ActivationIndex*/, 2 /*DeactivationIndex*/, 3 /*TickIndex*/, NAME_None /*TickBoolKeyName*/, false /*bCallTickOnSearchStart*/,
+						NAME_None /*BecomeRelevantBoolKeyeyName*/, TEXT("Bool2") /*CeaseRelevantBoolKeyName*/);
+				}
+			}
+			FBTBuilder::AddTask(CompNode, 6, EBTNodeResult::Succeeded);
+		}
+		ExpectedResult.Add(1);
+		ExpectedResult.Add(3);
+		ExpectedResult.Add(3);
+		ExpectedResult.Add(2);
+		ExpectedResult.Add(6);
+	}
+};
+IMPLEMENT_AI_LATENT_TEST(FAITest_BTStartBranchDuringServiceCeaseRelevant, "System.AI.Behavior Trees.Toggle: start new branch during service cease relevant")
+
 struct FAITest_BTAbortingDuringServiceBecomeRelevant : public FAITest_SimpleBT
 {
 	FAITest_BTAbortingDuringServiceBecomeRelevant()
diff --git a/Engine/Source/Developer/AITestSuite/Public/BTBuilder.h b/Engine/Source/Developer/AITestSuite/Public/BTBuilder.h
--- a/Engine/Source/Developer/AITestSuite/Public/BTBuilder.h
+++ b/Engine/Source/Developer/AITestSuite/Public/BTBuilder.h
@@ -17,6 +17,7 @@
 
 #include "BehaviorTree/TestBTTask_Log.h"
 #include "BehaviorTree/TestBTTask_SetFlag.h"
+#include "BehaviorTree/TestBTTask_ToggleFlag.h"
 #include "BehaviorTree/TestBTTask_SetValue.h"
 #include "BehaviorTree/TestBTTask_LatentWithFlags.h"
 #include "BehaviorTree/TestBTDecorator_Blackboard.h"
@@ -162,6 +163,17 @@ struct FBTBuilder
 		ParentNode.Children[ChildIdx].ChildTask = TaskNode;
 	}
 
+	static void AddTaskToggleFlag(UBTCompositeNode& ParentNode, EBTNodeResult::Type NodeResult, FName BoolKeyName, int32 NumToggles)
+	{
+		UTestBTTask_ToggleFlag* TaskNode = NewObject<UTestBTTask_ToggleFlag>(ParentNode.GetTreeAsset());
+		TaskNode->TaskResult = NodeResult;
+		TaskNode->KeyName = BoolKeyName;
+		TaskNode->NumToggles = NumToggles;
+
+		const int32 ChildIdx = ParentNode.Children.AddZeroed(1);
+		ParentNode.Children[ChildIdx].ChildTask = TaskNode;
+	}
+
 	static void AddTaskValueChange(UBTCompositeNode& ParentNode, int32 Value, EBTNodeResult::Type NodeResult, FName IntKeyName = TEXT("Int"), FName IntOnAbortKeyName = FName(), int32 OnAbortValue = 0)
 	{
 		UTestBTTask_SetValue* TaskNode = NewObject<UTestBTTask_SetValue>(ParentNode.GetTreeAsset());
@@ -290,7 +302,7 @@ struct FBTBuilder
 		return *ServiceOb;
 	}
 
-	static void WithServiceLog(UBTCompositeNode& ParentNode, int32 ActivationIndex, int32 DeactivationIndex, int32 TickIndex = INDEX_NONE, FName TickBoolKeyName = NAME_None, bool bCallTickOnSearchStart = false, FName BecomeRelevantBoolKeyName = NAME_None, FName CeaseRelevantBoolKeyName = NAME_None)
+	static void WithServiceLog(UBTCompositeNode& ParentNode, int32 ActivationIndex, int32 DeactivationIndex, int32 TickIndex = INDEX_NONE, FName TickBoolKeyName = NAME_None, bool bCallTickOnSearchStart = false, FName BecomeRelevantBoolKeyName = NAME_None, FName CeaseRelevantBoolKeyName = NAME_None, bool bToggleValue = false)
 	{
 		UTestBTService_Log& LogService = WithService<UTestBTService_Log>(ParentNode);
 		LogService.LogActivation = ActivationIndex;
@@ -299,6 +311,7 @@ struct FBTBuilder
 		LogService.SetFlagOnTick(TickBoolKeyName, bCallTickOnSearchStart);
 		LogService.KeyNameBecomeRelevant = BecomeRelevantBoolKeyName;
 		LogService.KeyNameCeaseRelevant = CeaseRelevantBoolKeyName;
+		LogService.bToggleValue = bToggleValue;
 	}
 
 	template<class T>
@@ -313,7 +326,7 @@ struct FBTBuilder
 		return *ServiceOb;
 	}
 
-	static void WithTaskServiceLog(UBTCompositeNode& ParentNode, int32 ActivationIndex, int32 DeactivationIndex, int32 TickIndex = INDEX_NONE, FName TickBoolKeyName = NAME_None, bool bCallTickOnSearchStart = false, FName BecomeRelevantBoolKeyName = NAME_None, FName CeaseRelevantBoolKeyName = NAME_None)
+	static void WithTaskServiceLog(UBTCompositeNode& ParentNode, int32 ActivationIndex, int32 DeactivationIndex, int32 TickIndex = INDEX_NONE, FName TickBoolKeyName = NAME_None, bool bCallTickOnSearchStart = false, FName BecomeRelevantBoolKeyName = NAME_None, FName CeaseRelevantBoolKeyName = NAME_None, bool bToggleValue = false)
 	{
 		UTestBTService_Log& LogService = WithTaskService<UTestBTService_Log>(ParentNode);
 		LogService.LogActivation = ActivationIndex;
@@ -322,5 +335,6 @@ struct FBTBuilder
 		LogService.SetFlagOnTick(TickBoolKeyName, bCallTickOnSearchStart);
 		LogService.KeyNameBecomeRelevant = BecomeRelevantBoolKeyName;
 		LogService.KeyNameCeaseRelevant = CeaseRelevantBoolKeyName;
+		LogService.bToggleValue = bToggleValue;
 	}
 };
diff --git a/Engine/Source/Runtime/AIModule/Classes/BehaviorTree/BehaviorTreeComponent.h b/Engine/Source/Runtime/AIModule/Classes/BehaviorTree/BehaviorTreeComponent.h
--- a/Engine/Source/Runtime/AIModule/Classes/BehaviorTree/BehaviorTreeComponent.h
+++ b/Engine/Source/Runtime/AIModule/Classes/BehaviorTree/BehaviorTreeComponent.h
@@ -79,6 +79,19 @@ struct FBTTreeStartInfo
 	bool HasPendingInitialize() const { return bPendingInitialize && IsSet(); }
 };
 
+enum class EBTBranchAction : uint8
+{
+	None = 0x0,
+	DecoratorEvaluate = 0x1,
+	DecoratorActivate_IfNotExecuting = 0x2,
+	DecoratorActivate_EvenIfExecuting = 0x4,
+	DecoratorActivate = DecoratorActivate_IfNotExecuting | DecoratorActivate_EvenIfExecuting,
+	DecoratorDeactivate = 0x8,
+	UnregisterAuxNodes = 0x10,
+	All = DecoratorEvaluate | DecoratorActivate_IfNotExecuting | DecoratorActivate_EvenIfExecuting | DecoratorDeactivate | UnregisterAuxNodes,
+};
+ENUM_CLASS_FLAGS(EBTBranchAction);
+
 UCLASS(ClassGroup = AI, meta = (BlueprintSpawnableComponent))
 class AIMODULE_API UBehaviorTreeComponent : public UBrainComponent
 {
@@ -136,6 +149,12 @@ public:
 	UE_DEPRECATED(5.0, "This function is deprecated. Please use RequestBranchDeactivation instead.")
 	void RequestUnregisterAuxNodesInBranch(const UBTCompositeNode* Node);
 		
+	/** request branch evaluation: helper for decorator */
+	void RequestBranchEvaluation(const UBTDecorator& RequestedBy);
+
+	/** request branch activation: helper for decorator */
+	void RequestBranchActivation(const UBTDecorator& RequestedBy, const bool bRequestEvenIfExecuting);
+
 	/** request branch deactivation: helper for decorator */
 	void RequestBranchDeactivation(const UBTDecorator& RequestedBy);
 
@@ -268,8 +287,21 @@ protected:
 	/** result of ExecutionRequest, will be applied when current task finish aborting */
 	FBTPendingExecutionInfo PendingExecution;
 
-	/** list of all pending branch deactivation requests */
-	TArray<const UBTNode*> PendingBranchesToDeactivate;
+	struct FBranchActionInfo
+	{
+		FBranchActionInfo(const UBTNode* InNode, const EBTBranchAction InAction)
+			: Node(InNode)
+			, Action(InAction)
+		{}
+		const UBTNode* Node;
+		EBTBranchAction Action;
+	};
+
+	/* Type of supended branch action */
+	EBTBranchAction SuspendedBranchActions;
+
+	/** list of all pending branch action requests */
+	TArray<FBranchActionInfo> PendingBranchActionRequests;
 
 	/** stored data for starting new tree, waits until previously running finishes aborting */
 	FBTTreeStartInfo TreeStartInfo;
@@ -328,9 +360,6 @@ protected:
 	/** if set, execution requests will be postponed */
 	uint8 bIsPaused : 1;
 
-	/** if set, all branch deactivation requests will be queued up */
-	uint8 bBranchDeactivationSuspended : 1;
-
 	/** push behavior tree instance on execution stack
 	 *	@NOTE: should never be called out-side of BT execution, meaning only BT tasks can push another BT instance! */
 	bool PushInstance(UBehaviorTree& TreeAsset);
@@ -418,29 +447,44 @@ protected:
 	/** Return NodeA's relative priority in regards to NodeB */
 	EBTNodeRelativePriority CalculateRelativePriority(const UBTNode* NodeA, const UBTNode* NodeB) const;
 
-	/** Deactivate a branch as the decorator condition is not passing anymore */
+	/** Activate a branch as the decorator conditions have changed */
+	void EvaluateBranch(const UBTDecorator& RequestedBy);
+
+	/** Activate a branch as the decorator conditions are now passing */
+	void ActivateBranch(const UBTDecorator& RequestedBy, bool bRequestEvenIfNotExecuting);
+
+	/** Deactivate a branch as the decorator conditions are not passing anymore */
 	void DeactivateBranch(const UBTDecorator& RequestedBy);
 
-	/** Suspend any branch deactivation and queue them to be processed later by ResumeBranchDeactivation() */
-	void SuspendBranchDeactivation();
+	/** Suspend any branch actions and queue them to be processed later by ResumeBranchActions() */
+	void SuspendBranchActions(EBTBranchAction BranchActions = EBTBranchAction::All);
 
-	/** Resume branch deactivation and execute all the queued up ones */
-	void ResumeBranchDeactivation();
+	/** Resume branch actions and execute all the queued up ones */
+	void ResumeBranchActions();
 
-	struct FBTSuspendBranchDeactivationScoped
+	UE_DEPRECATED(5.1, "This function is deprecated. Please use SuspendBranchActions instead.")
+	void SuspendBranchDeactivation() { SuspendBranchActions(); }
+
+	UE_DEPRECATED(5.1, "This function is deprecated. Please use ResumeBranchActions instead.")
+	void ResumeBranchDeactivation() { ResumeBranchActions(); }
+
+	struct FBTSuspendBranchActionsScoped
 	{
-		FBTSuspendBranchDeactivationScoped(UBehaviorTreeComponent& InBTComp)
+		FBTSuspendBranchActionsScoped(UBehaviorTreeComponent& InBTComp, EBTBranchAction BranchActions = EBTBranchAction::All)
 			: BTComp(InBTComp)
 		{
-			BTComp.SuspendBranchDeactivation();
+			BTComp.SuspendBranchActions(BranchActions);
 		}
-		~FBTSuspendBranchDeactivationScoped()
+		~FBTSuspendBranchActionsScoped()
 		{
-			BTComp.ResumeBranchDeactivation();
+			BTComp.ResumeBranchActions();
 		}
 		UBehaviorTreeComponent& BTComp;
 	};
 
+	UE_DEPRECATED(5.1, "This struct is deprecated. Please use FBTSuspendBranchActionsScoped instead.")
+	typedef FBTSuspendBranchActionsScoped FBTSuspendBranchDeactivationScoped;
+
 	friend UBTNode;
 	friend UBTCompositeNode;
 	friend UBTTaskNode;
diff --git a/Engine/Source/Runtime/AIModule/Classes/BehaviorTree/BehaviorTreeTypes.h b/Engine/Source/Runtime/AIModule/Classes/BehaviorTree/BehaviorTreeTypes.h
--- a/Engine/Source/Runtime/AIModule/Classes/BehaviorTree/BehaviorTreeTypes.h
+++ b/Engine/Source/Runtime/AIModule/Classes/BehaviorTree/BehaviorTreeTypes.h
@@ -466,12 +466,17 @@ struct FBehaviorTreeSearchUpdate
 	/** if set, this entry will be applied AFTER other are processed */
 	uint8 bPostUpdate : 1;
 
-	FBehaviorTreeSearchUpdate() : AuxNode(0), TaskNode(0), InstanceIndex(0), Mode(EBTNodeUpdateMode::Unknown), bPostUpdate(false) {}
-	FBehaviorTreeSearchUpdate(const UBTAuxiliaryNode* InAuxNode, uint16 InInstanceIndex, EBTNodeUpdateMode::Type InMode) :
-		AuxNode((UBTAuxiliaryNode*)InAuxNode), TaskNode(0), InstanceIndex(InInstanceIndex), Mode(InMode), bPostUpdate(false)
+	/** if set the action was unnecessary (Remove and was already inactive, Add and was already active) */
+	mutable uint8 bApplySkipped : 1;
+
+	FBehaviorTreeSearchUpdate() 
+		: AuxNode(0), TaskNode(0), InstanceIndex(0), Mode(EBTNodeUpdateMode::Unknown), bPostUpdate(false), bApplySkipped(false)
+	{}
+	FBehaviorTreeSearchUpdate(const UBTAuxiliaryNode* InAuxNode, uint16 InInstanceIndex, EBTNodeUpdateMode::Type InMode)
+		: AuxNode((UBTAuxiliaryNode*)InAuxNode), TaskNode(0), InstanceIndex(InInstanceIndex), Mode(InMode), bPostUpdate(false), bApplySkipped(false)
 	{}
-	FBehaviorTreeSearchUpdate(const UBTTaskNode* InTaskNode, uint16 InInstanceIndex, EBTNodeUpdateMode::Type InMode) :
-		AuxNode(0), TaskNode((UBTTaskNode*)InTaskNode), InstanceIndex(InInstanceIndex), Mode(InMode), bPostUpdate(false)
+	FBehaviorTreeSearchUpdate(const UBTTaskNode* InTaskNode, uint16 InInstanceIndex, EBTNodeUpdateMode::Type InMode)
+		: AuxNode(0), TaskNode((UBTTaskNode*)InTaskNode), InstanceIndex(InInstanceIndex), Mode(InMode), bPostUpdate(false), bApplySkipped(false)
 	{}
 };
 
diff --git a/Engine/Source/Runtime/AIModule/Private/BehaviorTree/BTAuxiliaryNode.cpp b/Engine/Source/Runtime/AIModule/Private/BehaviorTree/BTAuxiliaryNode.cpp
--- a/Engine/Source/Runtime/AIModule/Private/BehaviorTree/BTAuxiliaryNode.cpp
+++ b/Engine/Source/Runtime/AIModule/Private/BehaviorTree/BTAuxiliaryNode.cpp
@@ -2,6 +2,7 @@
 
 #include "BehaviorTree/BTAuxiliaryNode.h"
 #include "BehaviorTree/BTCompositeNode.h"
+#include "VisualLogger/VisualLogger.h"
 
 UBTAuxiliaryNode::UBTAuxiliaryNode(const FObjectInitializer& ObjectInitializer) : Super(ObjectInitializer)
 {
@@ -60,6 +61,8 @@ bool UBTAuxiliaryNode::WrappedTickNode(UBehaviorTreeComponent& OwnerComp, uint8*
 				    UseDeltaTime = AuxMemory->AccumulatedDeltaTime;
 				    AuxMemory->AccumulatedDeltaTime = 0.0f;
     
+					UE_VLOG(OwnerComp.GetOwner(), LogBehaviorTree, Verbose, TEXT("Ticking aux node: %s"), *UBehaviorTreeTypes::DescribeNodeHelper(this));
+
 				    const_cast<UBTAuxiliaryNode*>(NodeOb)->TickNode(OwnerComp, NodeMemory, UseDeltaTime);
 				}
 
@@ -71,6 +74,8 @@ bool UBTAuxiliaryNode::WrappedTickNode(UBehaviorTreeComponent& OwnerComp, uint8*
 			}
 			else
 			{
+				UE_VLOG(OwnerComp.GetOwner(), LogBehaviorTree, Verbose, TEXT("Ticking aux node: %s"), *UBehaviorTreeTypes::DescribeNodeHelper(this));
+
 				const_cast<UBTAuxiliaryNode*>(NodeOb)->TickNode(OwnerComp, NodeMemory, UseDeltaTime);
 				NextNeededDeltaTime = 0.0f;
 				return true;
diff --git a/Engine/Source/Runtime/AIModule/Private/BehaviorTree/BTDecorator.cpp b/Engine/Source/Runtime/AIModule/Private/BehaviorTree/BTDecorator.cpp
--- a/Engine/Source/Runtime/AIModule/Private/BehaviorTree/BTDecorator.cpp
+++ b/Engine/Source/Runtime/AIModule/Private/BehaviorTree/BTDecorator.cpp
@@ -97,32 +97,21 @@ void UBTDecorator::ConditionalFlowAbort(UBehaviorTreeComponent& OwnerComp, EBTDe
 
 	uint8* NodeMemory = OwnerComp.GetNodeMemory((UBTNode*)this, InstanceIdx);
 
-	const bool bIsExecutingBranch = OwnerComp.IsExecutingBranch(this, GetChildIndex());
 	const bool bPass = WrappedCanExecute(OwnerComp, NodeMemory);
-	const bool bAbortPending = OwnerComp.IsAbortPending();
 	const bool bAlwaysRequestWhenPassing = (RequestMode == EBTDecoratorAbortRequest::ConditionPassing);
 
-	const bool bLogRequestExecution = (!bIsExecutingBranch && bPass) || (bIsExecutingBranch && bPass && (bAlwaysRequestWhenPassing || bAbortPending));
-	UE_VLOG(OwnerComp.GetOwner(), LogBehaviorTree, Verbose, TEXT("%s, ConditionalFlowAbort(%s) pass:%d executingBranch:%d abortPending:%d => %s"),
+	UE_VLOG(OwnerComp.GetOwner(), LogBehaviorTree, Verbose, TEXT("%s, ConditionalFlowAbort(%s) pass:%d"),
 		*UBehaviorTreeTypes::DescribeNodeHelper(this),
-		bAlwaysRequestWhenPassing ? TEXT("always when passing") : TEXT("on change"),
-		bPass ? 1 : 0,
-		bIsExecutingBranch ? 1 : 0,
-		bAbortPending ? 1 : 0,
-		!bPass ? TEXT("request branch deactivation") : (bLogRequestExecution ? TEXT("request execution") : TEXT("skip")));
+		bAlwaysRequestWhenPassing ? TEXT("always when passing") : TEXT("on change"));
+
 
 	if (!bPass)
 	{
 		OwnerComp.RequestBranchDeactivation(*this);
 	}
-	else if (!bIsExecutingBranch)
-	{
-		OwnerComp.RequestExecution(this);
-	}
-	else if (bAlwaysRequestWhenPassing || bAbortPending)
+	else
 	{
-		// force result Aborted to restart from this decorator
-		OwnerComp.RequestExecution(GetParentNode(), InstanceIdx, this, GetChildIndex(), EBTNodeResult::Aborted);
+		OwnerComp.RequestBranchActivation(*this, bAlwaysRequestWhenPassing);
 	}
 }
 
diff --git a/Engine/Source/Runtime/AIModule/Private/BehaviorTree/BehaviorTreeComponent.cpp b/Engine/Source/Runtime/AIModule/Private/BehaviorTree/BehaviorTreeComponent.cpp
--- a/Engine/Source/Runtime/AIModule/Private/BehaviorTree/BehaviorTreeComponent.cpp
+++ b/Engine/Source/Runtime/AIModule/Private/BehaviorTree/BehaviorTreeComponent.cpp
@@ -71,7 +71,7 @@ UBehaviorTreeComponent::UBehaviorTreeComponent(const FObjectInitializer& ObjectI
 	bWantsInitializeComponent = true; 
 	bIsRunning = false;
 	bIsPaused = false;
-	bBranchDeactivationSuspended = false;
+	SuspendedBranchActions = EBTBranchAction::None;
 
 	// Adding hook for bespoke framepro BT timings for BR
 #if !UE_BUILD_SHIPPING
@@ -305,7 +305,7 @@ void UBehaviorTreeComponent::StopTree(EBTStopMode::Type StopMode)
 
 			// notify active aux nodes
 			{
-				FBTSuspendBranchDeactivationScoped ScopedSuspend(*this);
+				FBTSuspendBranchActionsScoped ScopedSuspend(*this, EBTBranchAction::All);
 				InstanceInfo.ExecuteOnEachAuxNode([&InstanceInfo, this](const UBTAuxiliaryNode& AuxNode)
 				{
 					uint8* NodeMemory = AuxNode.GetNodeMemory<uint8>(InstanceInfo);
@@ -672,9 +672,10 @@ void UBehaviorTreeComponent::RequestUnregisterAuxNodesInBranch(const UBTComposit
 		return;
 	}
 
-	if (bBranchDeactivationSuspended)
+	if ((SuspendedBranchActions & EBTBranchAction::UnregisterAuxNodes) != EBTBranchAction::None)
 	{
-		PendingBranchesToDeactivate.Add(Node);
+		UE_VLOG(GetOwner(), LogBehaviorTree, Verbose, TEXT("%s, unregister aux nodes in branch queued up"), *UBehaviorTreeTypes::DescribeNodeHelper(Node));
+		PendingBranchActionRequests.Emplace(Node, EBTBranchAction::UnregisterAuxNodes);
 		return;
 	}
 
@@ -685,8 +686,8 @@ void UBehaviorTreeComponent::DeactivateBranch(const UBTDecorator& RequestedBy)
 {
 	if (IsExecutingBranch(&RequestedBy, RequestedBy.GetChildIndex()))
 	{
-		UE_VLOG(GetOwner(), LogBehaviorTree, Verbose, TEXT("%s, Branch deactivation resulted in new execution request"), *UBehaviorTreeTypes::DescribeNodeHelper(&RequestedBy));
-		RequestExecution(&RequestedBy);
+		UE_VLOG(GetOwner(), LogBehaviorTree, Verbose, TEXT("%s, Branch deactivation resulted in an evaluate branch"), *UBehaviorTreeTypes::DescribeNodeHelper(&RequestedBy));
+		EvaluateBranch(RequestedBy);
 	}
 	else if (ensureMsgf(RequestedBy.GetParentNode() && RequestedBy.GetParentNode()->Children.IsValidIndex(RequestedBy.GetChildIndex()), 
 				TEXT("The decorator %s does not have a parent or is not a valid child."), *UBehaviorTreeTypes::DescribeNodeHelper(&RequestedBy)))
@@ -694,7 +695,15 @@ void UBehaviorTreeComponent::DeactivateBranch(const UBTDecorator& RequestedBy)
 		UE_VLOG(GetOwner(), LogBehaviorTree, Verbose, TEXT("%s, Branch deactivation resulted in aux nodes unregistration"), *UBehaviorTreeTypes::DescribeNodeHelper(&RequestedBy));
 		if (const UBTCompositeNode* BranchRoot = RequestedBy.GetParentNode()->Children[RequestedBy.GetChildIndex()].ChildComposite)
 		{
-			UnregisterAuxNodesInBranch(BranchRoot, true/*bApplyImmediately*/);
+			if ((SuspendedBranchActions & EBTBranchAction::UnregisterAuxNodes) != EBTBranchAction::None)
+			{
+				UE_VLOG(GetOwner(), LogBehaviorTree, Verbose, TEXT("%s, unregister aux nodes in branch queued up"), *UBehaviorTreeTypes::DescribeNodeHelper(&RequestedBy));
+				PendingBranchActionRequests.Emplace(BranchRoot, EBTBranchAction::UnregisterAuxNodes);
+			}
+			else
+			{
+				UnregisterAuxNodesInBranch(BranchRoot, true/*bApplyImmediately*/);
+			}
 		}
 	}
 	else
@@ -705,49 +714,65 @@ void UBehaviorTreeComponent::DeactivateBranch(const UBTDecorator& RequestedBy)
 
 void UBehaviorTreeComponent::RequestBranchDeactivation(const UBTDecorator& RequestedBy)
 {
-	if (bBranchDeactivationSuspended)
+	if ((SuspendedBranchActions & EBTBranchAction::DecoratorDeactivate) != EBTBranchAction::None)
 	{
 		UE_VLOG(GetOwner(), LogBehaviorTree, Verbose, TEXT("%s, Branch deactivation queued up"), *UBehaviorTreeTypes::DescribeNodeHelper(&RequestedBy));
-		PendingBranchesToDeactivate.Add(&RequestedBy);
+		PendingBranchActionRequests.Emplace(&RequestedBy, EBTBranchAction::DecoratorDeactivate);
 		return;
 	}
 
 	DeactivateBranch(RequestedBy);
 }
 
-void UBehaviorTreeComponent::SuspendBranchDeactivation()
+void UBehaviorTreeComponent::SuspendBranchActions(EBTBranchAction BranchActions)
 {
-	UE_VLOG(GetOwner(), LogBehaviorTree, VeryVerbose, TEXT("Suspending branch deactivation."));
-	checkf(!bBranchDeactivationSuspended, TEXT("This logic does not support re-entrance"));
-	bBranchDeactivationSuspended = true;
+	UE_VLOG(GetOwner(), LogBehaviorTree, VeryVerbose, TEXT("Suspending branch actions."));
+	checkf(SuspendedBranchActions == EBTBranchAction::None, TEXT("This logic does not support re-entrance"));
+	SuspendedBranchActions = BranchActions;
 }
 
-void UBehaviorTreeComponent::ResumeBranchDeactivation()
+void UBehaviorTreeComponent::ResumeBranchActions()
 {
-	UE_VLOG(GetOwner(), LogBehaviorTree, VeryVerbose, TEXT("Resuming branch deactivation."));
-	checkf(bBranchDeactivationSuspended, TEXT("Expecting SuspendBranchDeactivation() be called before calling resume"));
-	bBranchDeactivationSuspended = false;
+	UE_VLOG(GetOwner(), LogBehaviorTree, VeryVerbose, TEXT("Resuming branch actions."));
+	checkf(SuspendedBranchActions != EBTBranchAction::None, TEXT("Expecting SuspendBranchActions() be called before calling resume"));
+	SuspendedBranchActions = EBTBranchAction::None;
 
-	// Flushing any pending branch deactivations
-	while (PendingBranchesToDeactivate.Num() > 0)
+	// Flushing any pending branch actions
+	while (PendingBranchActionRequests.Num() > 0)
 	{
-		TArray<const UBTNode*> PendingBranchesToDeactivateToProcess(MoveTemp(PendingBranchesToDeactivate));
-		PendingBranchesToDeactivate.Reset();
-		for (const UBTNode* Node : PendingBranchesToDeactivateToProcess)
+		TArray<FBranchActionInfo> PendingBranchActionRequestsToProcess(MoveTemp(PendingBranchActionRequests));
+		PendingBranchActionRequests.Reset();
+		for (const FBranchActionInfo& Info : PendingBranchActionRequestsToProcess)
 		{
-			if (const UBTDecorator* RequestedDecorator = Cast<UBTDecorator>(Node))
-			{
-				DeactivateBranch(*RequestedDecorator);
-			}
-			//////////////////////////////////////////////////////////////////////
-			// An attempt at backward compatibility 
-			// If anybody is still using deprecated method RequestUnregisterAuxNodesInBranch
-			else if (const UBTCompositeNode* BranchRoot = Cast<UBTCompositeNode>(Node))
+			switch (Info.Action)
 			{
-				UnregisterAuxNodesInBranch(BranchRoot, true/*bApplyImmediately*/);
+				case EBTBranchAction::DecoratorEvaluate:
+				{
+					const UBTDecorator* RequestedBy = CastChecked<UBTDecorator>(Info.Node);
+					EvaluateBranch(*RequestedBy);
+					break;
+				}
+
+				case EBTBranchAction::DecoratorActivate_IfNotExecuting:
+				case EBTBranchAction::DecoratorActivate_EvenIfExecuting:
+				{
+					const UBTDecorator* RequestedBy = CastChecked<UBTDecorator>(Info.Node);
+					ActivateBranch(*RequestedBy, Info.Action == EBTBranchAction::DecoratorActivate_EvenIfExecuting /*bForceRequestEvenIfExecuting*/);
+					break;
+				}
+				case EBTBranchAction::DecoratorDeactivate:
+				{
+					const UBTDecorator* RequestedBy = CastChecked<UBTDecorator>(Info.Node);
+					DeactivateBranch(*RequestedBy);
+					break;
+				}
+				case EBTBranchAction::UnregisterAuxNodes:
+				{
+					const UBTCompositeNode* BranchRoot = CastChecked<UBTCompositeNode>(Info.Node);
+					UnregisterAuxNodesInBranch(BranchRoot, true/*bApplyImmediately*/);
+					break;
+				}
 			}
-			//
-			//////////////////////////////////////////////////////////////////////
 		}
 	}
 }
@@ -755,6 +780,23 @@ void UBehaviorTreeComponent::ResumeBranchDeactivation()
 void UBehaviorTreeComponent::RequestExecution(const UBTDecorator* RequestedBy)
 {
 	check(RequestedBy);
+	RequestBranchEvaluation(*RequestedBy);
+}
+
+void UBehaviorTreeComponent::RequestBranchEvaluation(const UBTDecorator& RequestedBy)
+{
+	if ((SuspendedBranchActions & EBTBranchAction::DecoratorEvaluate) != EBTBranchAction::None)
+	{
+		UE_VLOG(GetOwner(), LogBehaviorTree, Verbose, TEXT("%s, Branch evaluation queued up"), *UBehaviorTreeTypes::DescribeNodeHelper(&RequestedBy));
+		PendingBranchActionRequests.Emplace(&RequestedBy, EBTBranchAction::DecoratorEvaluate);
+		return;
+	}
+
+	EvaluateBranch(RequestedBy);
+}
+
+void UBehaviorTreeComponent::EvaluateBranch(const UBTDecorator& RequestedBy)
+{
 	// search range depends on decorator's FlowAbortMode:
 	//
 	// - LowerPri: try entering branch = search only nodes under decorator
@@ -764,13 +806,13 @@ void UBehaviorTreeComponent::RequestExecution(const UBTDecorator* RequestedBy)
 	// - Both: check if active node is within inner child nodes and choose Self or LowerPri
 	//
 
-	EBTFlowAbortMode::Type AbortMode = RequestedBy->GetFlowAbortMode();
+	EBTFlowAbortMode::Type AbortMode = RequestedBy.GetFlowAbortMode();
 	if (AbortMode == EBTFlowAbortMode::None)
 	{
 		return;
 	}
 
-	const int32 InstanceIdx = FindInstanceContainingNode(RequestedBy->GetParentNode());
+	const int32 InstanceIdx = FindInstanceContainingNode(RequestedBy.GetParentNode());
 	if (InstanceIdx == INDEX_NONE)
 	{
 		return;
@@ -780,13 +822,13 @@ void UBehaviorTreeComponent::RequestExecution(const UBTDecorator* RequestedBy)
 	const FBehaviorTreeInstance& ActiveInstance = InstanceStack.Last();
 	if (ActiveInstance.ActiveNodeType == EBTActiveNode::ActiveTask)
 	{
-		EBTNodeRelativePriority RelativePriority = CalculateRelativePriority(RequestedBy, ActiveInstance.ActiveNode);
+		EBTNodeRelativePriority RelativePriority = CalculateRelativePriority(&RequestedBy, ActiveInstance.ActiveNode);
 
 		if (RelativePriority < EBTNodeRelativePriority::Same)
 		{
 			const FString ErrorMsg(FString::Printf(TEXT("%s: decorator %s requesting restart has lower priority than Current Task %s"),
 				ANSI_TO_TCHAR(__FUNCTION__),
-				*UBehaviorTreeTypes::DescribeNodeHelper(RequestedBy),
+				*UBehaviorTreeTypes::DescribeNodeHelper(&RequestedBy),
 				*UBehaviorTreeTypes::DescribeNodeHelper(ActiveInstance.ActiveNode)));
 
 			UE_VLOG(GetOwner(), LogBehaviorTree, Error, TEXT("%s"), *ErrorMsg);
@@ -797,12 +839,61 @@ void UBehaviorTreeComponent::RequestExecution(const UBTDecorator* RequestedBy)
 
 	if (AbortMode == EBTFlowAbortMode::Both)
 	{
-		const bool bIsExecutingChildNodes = IsExecutingBranch(RequestedBy, RequestedBy->GetChildIndex());
+		const bool bIsExecutingChildNodes = IsExecutingBranch(&RequestedBy, RequestedBy.GetChildIndex());
 		AbortMode = bIsExecutingChildNodes ? EBTFlowAbortMode::Self : EBTFlowAbortMode::LowerPriority;
 	}
 
+	UE_VLOG(GetOwner(), LogBehaviorTree, Verbose, TEXT("%s, EvaluateBranch requested a new execution"), *UBehaviorTreeTypes::DescribeNodeHelper(&RequestedBy));
+
 	EBTNodeResult::Type ContinueResult = (AbortMode == EBTFlowAbortMode::Self) ? EBTNodeResult::Failed : EBTNodeResult::Aborted;
-	RequestExecution(RequestedBy->GetParentNode(), InstanceIdx, RequestedBy, RequestedBy->GetChildIndex(), ContinueResult);
+	RequestExecution(RequestedBy.GetParentNode(), InstanceIdx, &RequestedBy, RequestedBy.GetChildIndex(), ContinueResult);
+}
+
+
+void UBehaviorTreeComponent::RequestBranchActivation(const UBTDecorator& RequestedBy, const bool bRequestEvenIfExecuting)
+{
+	if ((SuspendedBranchActions & EBTBranchAction::DecoratorActivate) != EBTBranchAction::None)
+	{
+		UE_VLOG(GetOwner(), LogBehaviorTree, Verbose, TEXT("%s, Branch activation queued up"), *UBehaviorTreeTypes::DescribeNodeHelper(&RequestedBy));
+		PendingBranchActionRequests.Emplace(&RequestedBy, bRequestEvenIfExecuting ? EBTBranchAction::DecoratorActivate_EvenIfExecuting : EBTBranchAction::DecoratorActivate_IfNotExecuting);
+		return;
+	}
+
+	ActivateBranch(RequestedBy, bRequestEvenIfExecuting);
+}
+
+
+void UBehaviorTreeComponent::ActivateBranch(const UBTDecorator& RequestedBy, const bool bForceResquestEvenIfExecuting)
+{
+	const int32 InstanceIdx = FindInstanceContainingNode(&RequestedBy);
+	if (InstanceIdx == INDEX_NONE)
+	{
+		return;
+	}
+
+	const bool bIsExecutingBranch = IsExecutingBranch(&RequestedBy, RequestedBy.GetChildIndex());
+	const bool bAbortPending = IsAbortPending();
+	const bool bIsDeactivatingBranchRoot = ExecutionRequest.ContinueWithResult == EBTNodeResult::Failed && ExecutionRequest.SearchStart == FBTNodeIndex(InstanceIdx, RequestedBy.GetExecutionIndex());
+
+	const bool bLogRequestExecution = !bIsExecutingBranch || (bForceResquestEvenIfExecuting || bAbortPending || bIsDeactivatingBranchRoot);
+	UE_VLOG(GetOwner(), LogBehaviorTree, Verbose, TEXT("%s, ActivateBranch(%s) executingBranch:%d abortPending:%d deactivatingBranchRoot:%d => %s"),
+		*UBehaviorTreeTypes::DescribeNodeHelper(&RequestedBy),
+		bForceResquestEvenIfExecuting ? TEXT("request even if executing") : TEXT("request if not executing"),
+		bIsExecutingBranch ? 1 : 0,
+		bAbortPending ? 1 : 0,
+		bIsDeactivatingBranchRoot ? 1 : 0,
+		bLogRequestExecution ? TEXT("request execution") : TEXT("skip"));
+
+	if (!bIsExecutingBranch)
+	{
+		// This should endup into a branch activation
+		EvaluateBranch(RequestedBy);
+	}
+	else if (bForceResquestEvenIfExecuting || bAbortPending || bIsDeactivatingBranchRoot)
+	{
+		// force result Aborted to restart from this branch
+		RequestExecution(RequestedBy.GetParentNode(), InstanceIdx, &RequestedBy, RequestedBy.GetChildIndex(), EBTNodeResult::Aborted);
+	}
 }
 
 EBTNodeRelativePriority UBehaviorTreeComponent::CalculateRelativePriority(const UBTNode* NodeA, const UBTNode* NodeB) const
@@ -950,29 +1041,32 @@ void UBehaviorTreeComponent::RequestExecution(UBTCompositeNode* RequestedOn, int
 	ExecutionIdx.ExecutionIndex = RequestedBy->GetExecutionIndex();
 	uint16 LastExecutionIndex = MAX_uint16;
 
-	// make sure that the request is not coming from a node that has pending branch deactivation since it won't be accessible anymore
-	if (bBranchDeactivationSuspended)
+	// make sure that the request is not coming from a node that has pending branch actions since it won't be accessible anymore
+	if (SuspendedBranchActions != EBTBranchAction::None)
 	{
-		for (const UBTNode* Node : PendingBranchesToDeactivate)
+		UE_CVLOG(SuspendedBranchActions != EBTBranchAction::UnregisterAuxNodes, GetOwner(), LogBehaviorTree, Warning, TEXT("Caller should be converted to new Evaluate/Activate/DeactivateBranch API instead of using this RequestExecution directly"));
+		for (const FBranchActionInfo& Info : PendingBranchActionRequests)
 		{
 			const UBTCompositeNode* BranchRoot = nullptr;
-			if (const UBTDecorator* Decorator = Cast<UBTDecorator>(Node))
+			switch (Info.Action)
 			{
-				// This check was already in previous version, it will ensure and output a vlog later in the DeactivateBranch, no need to do anything now.
-				if (Decorator->GetParentNode() && Decorator->GetParentNode()->Children.IsValidIndex(Decorator->GetChildIndex()))
-				{
-					BranchRoot = Decorator->GetParentNode()->Children[Decorator->GetChildIndex()].ChildComposite;
-				}
-			}
-			//////////////////////////////////////////////////////////////////////
-			// An attempt at backward compatibility 
-			// If anybody is still using deprecated method RequestUnregisterAuxNodesInBranch
-			else if (const UBTCompositeNode* CompNode = Cast<UBTCompositeNode>(Node))
-			{
-				BranchRoot = CompNode;
+				case EBTBranchAction::DecoratorDeactivate:
+					if (const UBTDecorator* Decorator = Cast<UBTDecorator>(Info.Node))
+					{
+						// This check was already in previous version, it will ensure and output a vlog later in the DeactivateBranch, no need to do anything now.
+						if (Decorator->GetParentNode() && Decorator->GetParentNode()->Children.IsValidIndex(Decorator->GetChildIndex()))
+						{
+							BranchRoot = Decorator->GetParentNode()->Children[Decorator->GetChildIndex()].ChildComposite;
+						}
+					}
+					break;
+				case EBTBranchAction::UnregisterAuxNodes:
+					if (const UBTCompositeNode* CompNode = Cast<UBTCompositeNode>(Info.Node))
+					{
+						BranchRoot = CompNode;
+					}
+					break;
 			}
-			//
-			//////////////////////////////////////////////////////////////////////
 
 			if (BranchRoot)
 			{
@@ -982,7 +1076,7 @@ void UBehaviorTreeComponent::RequestExecution(UBTCompositeNode* RequestedOn, int
 					FBTNodeIndexRange Range( FBTNodeIndex(BranchRootInstanceIdx, BranchRoot->GetExecutionIndex()), FBTNodeIndex(BranchRootInstanceIdx, BranchRoot->GetLastExecutionIndex()));
 					if (Range.Contains(ExecutionIdx))
 					{
-						UE_VLOG(GetOwner(), LogBehaviorTree, Log, TEXT("> skip: request by %s(%s) is in deactivated branch %s(%s) and was deactivated by %s"), *UBehaviorTreeTypes::DescribeNodeHelper(RequestedBy), *ExecutionIdx.Describe(), *UBehaviorTreeTypes::DescribeNodeHelper(BranchRoot), *Range.Describe(), *UBehaviorTreeTypes::DescribeNodeHelper(Node));
+						UE_VLOG(GetOwner(), LogBehaviorTree, Log, TEXT("> skip: request by %s(%s) is in deactivated branch %s(%s) and was deactivated by %s"), *UBehaviorTreeTypes::DescribeNodeHelper(RequestedBy), *ExecutionIdx.Describe(), *UBehaviorTreeTypes::DescribeNodeHelper(BranchRoot), *Range.Describe(), *UBehaviorTreeTypes::DescribeNodeHelper(Info.Node));
 						return;
 					}
 				}
@@ -991,7 +1085,7 @@ void UBehaviorTreeComponent::RequestExecution(UBTCompositeNode* RequestedOn, int
 	}
 	else
 	{
-		checkf(PendingBranchesToDeactivate.Num() == 0, TEXT("All pending branches should have been flushed before requesting an execution"));
+		checkf(PendingBranchActionRequests.Num() == 0, TEXT("All pending branches should have been flushed before requesting an execution"));
 	}
 
 	if (bSwitchToHigherPriority && RequestedByChildIndex >= 0)
@@ -1210,6 +1304,12 @@ void UBehaviorTreeComponent::ApplySearchUpdates(const TArray<FBehaviorTreeSearch
 	for (int32 Index = 0; Index < UpdateList.Num(); Index++)
 	{
 		const FBehaviorTreeSearchUpdate& UpdateInfo = UpdateList[Index];
+		// Check if we are in the right pass...
+		if (UpdateInfo.bPostUpdate != bPostUpdate)
+		{
+			continue;
+		}
+
 		if (!InstanceStack.IsValidIndex(UpdateInfo.InstanceIndex))
 		{
 			continue;
@@ -1233,13 +1333,18 @@ void UBehaviorTreeComponent::ApplySearchUpdates(const TArray<FBehaviorTreeSearch
 		checkSlow(UpdateNode);
 
 		if ((UpdateInfo.Mode == EBTNodeUpdateMode::Remove && !bIsComponentActive) ||
-			(UpdateInfo.Mode == EBTNodeUpdateMode::Add && (bIsComponentActive || UpdateNode->GetExecutionIndex() > NewNodeExecutionIndex)) ||
-			(UpdateInfo.bPostUpdate != bPostUpdate))
+			(UpdateInfo.Mode == EBTNodeUpdateMode::Add && (bIsComponentActive || UpdateNode->GetExecutionIndex() > NewNodeExecutionIndex)) )
 		{
+			UpdateInfo.bApplySkipped = true;
+			UE_VLOG(GetOwner(), LogBehaviorTree, Log, TEXT("Skipping: %s for %s: %s (already in the right state)"),
+				*UBehaviorTreeTypes::DescribeNodeUpdateMode(UpdateInfo.Mode),
+				UpdateInfo.AuxNode ? TEXT("auxiliary node") : TEXT("parallel's main task"),
+				*UBehaviorTreeTypes::DescribeNodeHelper(UpdateNode));
+
 			continue;
 		}
 
-		UE_VLOG(GetOwner(), LogBehaviorTree, Log, TEXT("Update: %s for %s: %s"),
+		UE_VLOG(GetOwner(), LogBehaviorTree, Log, TEXT("Applying: %s for %s: %s"),
 			*UBehaviorTreeTypes::DescribeNodeUpdateMode(UpdateInfo.Mode),
 			UpdateInfo.AuxNode ? TEXT("auxiliary node") : TEXT("parallel's main task"),
 			*UBehaviorTreeTypes::DescribeNodeHelper(UpdateNode));
@@ -1347,7 +1452,7 @@ void UBehaviorTreeComponent::ApplySearchData(UBTNode* NewActiveNode)
 	for (int32 Idx = 0; Idx < SearchData.PendingUpdates.Num(); Idx++)
 	{
 		const FBehaviorTreeSearchUpdate& UpdateInfo = SearchData.PendingUpdates[Idx];
-		if (UpdateInfo.Mode == EBTNodeUpdateMode::Add && UpdateInfo.AuxNode && InstanceStack.IsValidIndex(UpdateInfo.InstanceIndex))
+		if (!UpdateInfo.bApplySkipped && UpdateInfo.Mode == EBTNodeUpdateMode::Add && UpdateInfo.AuxNode && InstanceStack.IsValidIndex(UpdateInfo.InstanceIndex))
 		{
 			FBehaviorTreeInstance& InstanceInfo = InstanceStack[UpdateInfo.InstanceIndex];
 			uint8* NodeMemory = UpdateInfo.AuxNode->GetNodeMemory<uint8>(InstanceInfo);
@@ -1424,13 +1529,13 @@ void UBehaviorTreeComponent::TickComponent(float DeltaTime, enum ELevelTick Tick
 	check(IsValid(this));
 	float NextNeededDeltaTime = FLT_MAX;
 
-	checkf(PendingBranchesToDeactivate.Num() == 0, TEXT("Pending branches should always be flushed immediately with the new system"))
+	checkf(PendingBranchActionRequests.Num() == 0, TEXT("Pending branches action requests should always be flushed immediately with the new system"))
 
 	// tick active auxiliary nodes (in execution order, before task)
 	// do it before processing execution request to give BP driven logic chance to accumulate execution requests
 	// newly added aux nodes are ticked as part of SearchData application
 	{
-		FBTSuspendBranchDeactivationScoped ScopedSuspend(*this);
+		FBTSuspendBranchActionsScoped ScopedSuspend(*this, EBTBranchAction::UnregisterAuxNodes);
 		for (int32 InstanceIndex = 0; InstanceIndex < InstanceStack.Num(); InstanceIndex++)
 		{
 			FBehaviorTreeInstance& InstanceInfo = InstanceStack[InstanceIndex];
@@ -1910,8 +2015,8 @@ void UBehaviorTreeComponent::ProcessPendingExecution()
 	const FBTNodeIndex NextTaskIdx = SavedInfo.NextTask ? FBTNodeIndex(ActiveInstanceIdx, SavedInfo.NextTask->GetExecutionIndex()) : FBTNodeIndex(0, 0);
 	UnregisterAuxNodesUpTo(NextTaskIdx);
 
-	// Suspending any branch deactivation as it is impossible for decorators to have the right answer if they are in an executing branch or not.
-	SuspendBranchDeactivation();
+	// Suspending any all branch actions as it is impossible for decorators to have the right answer if they are in an executing branch or not.
+	SuspendBranchActions(EBTBranchAction::All);
 
 	// change aux nodes
 	ApplySearchData(SavedInfo.NextTask);
@@ -1931,12 +2036,12 @@ void UBehaviorTreeComponent::ProcessPendingExecution()
 	// validate active instance as well, execution can be delayed AND can have AbortCurrentTask call before using instance index
 	if (SavedInfo.NextTask && InstanceStack.IsValidIndex(ActiveInstanceIdx))
 	{
-		// ResumeBranchDeactivation() is done inside ExecuteTask after the active task is set but before we execute the task.
+		// ResumeBranchActions() is done inside ExecuteTask after the active task is set but before we execute the task.
 		ExecuteTask(SavedInfo.NextTask);
 	}
 	else
 	{
-		ResumeBranchDeactivation();
+		ResumeBranchActions();
 		OnTreeFinished();
 	}
 }
@@ -2100,7 +2205,7 @@ void UBehaviorTreeComponent::UnregisterAuxNodesInBranch(const UBTCompositeNode*
 
 		if (bApplyImmediately)
 		{
-			FBTSuspendBranchDeactivationScoped ScopedSuspend(*this);
+			FBTSuspendBranchActionsScoped ScopedSuspend(*this, EBTBranchAction::All);
 			ApplySearchUpdates(SearchData.PendingUpdates, 0);
 			SearchData.PendingUpdates = UpdateListCopy;
 		}
@@ -2137,8 +2242,6 @@ void UBehaviorTreeComponent::ExecuteTask(UBTTaskNode* TaskNode)
 	{
 		uint8* NodeMemory = (uint8*)ServiceNode->GetNodeMemory<uint8>(ActiveInstance);
 
-		UE_VLOG(GetOwner(), LogBehaviorTree, Log, TEXT("Ticking task service: %s"), *UBehaviorTreeTypes::DescribeNodeHelper(ServiceNode));
-
 		// We do not care about the next needed DeltaTime, it will be recalculated in the tick later.
 		float NextNeededDeltaTime = 0.0f;
 		ServiceNode->WrappedTickNode(*this, NodeMemory, CurrentFrameDeltaSeconds, NextNeededDeltaTime);
@@ -2147,9 +2250,9 @@ void UBehaviorTreeComponent::ExecuteTask(UBTTaskNode* TaskNode)
 	ActiveInstance.ActiveNode = TaskNode;
 	ActiveInstance.ActiveNodeType = EBTActiveNode::ActiveTask;
 
-	// Is is now ok to resume any branch deactivation as the new active instance is set.
+	// Is is now ok to resume any branch acitons as the new active instance is set.
 	// Before that, decorators evaluating the IsExecutingBranch would be wrong.
-	ResumeBranchDeactivation();
+	ResumeBranchActions();
 
 	// make a snapshot for debugger
 	StoreDebuggerExecutionStep(EBTExecutionSnap::Regular);
