AI Geom Utils: Added extern to the functions implemented in cpp

diff --git a/Engine/Source/Runtime/AIModule/Public/GeomUtils.h b/Engine/Source/Runtime/AIModule/Public/GeomUtils.h
--- a/Engine/Source/Runtime/AIModule/Public/GeomUtils.h
+++ b/Engine/Source/Runtime/AIModule/Public/GeomUtils.h
@@ -125,7 +125,7 @@ namespace UE::AI
 	 * @param OutSegMax index of the polygon segment of the second intersection point
 	 * @return true if the segment inside or intersects with the polygon.
 	 */
-	bool IntersectSegmentPoly2D(const FVector& Start, const FVector& End, TConstArrayView<FVector> Poly,
+	extern AIMODULE_API bool IntersectSegmentPoly2D(const FVector& Start, const FVector& End, TConstArrayView<FVector> Poly,
 								FVector2D::FReal& OutTMin, FVector2D::FReal& OutTMax, int32& OutSegMin, int32& OutSegMax);
 
 	/**
@@ -153,7 +153,7 @@ namespace UE::AI
 	 * @param VertexD fourth corner
 	 * @return UV interpolation coordinates of the 'Point'.
 	 */
-	FVector2D InvBilinear2D(const FVector Point, const FVector VertexA, const FVector VertexB, const FVector VertexC, const FVector VertexD);
+	extern AIMODULE_API FVector2D InvBilinear2D(const FVector Point, const FVector VertexA, const FVector VertexB, const FVector VertexC, const FVector VertexD);
 
 	/**
 	 * Finds the UV coordinates of the 'Point' on bilinear patch A,B,C,D. U interpolates from A->B, and C->D, and V interpolates from AB->CD.
