Prevent infinite loop when the Loop Condition is still valid but no nodes under have valid conditions. Took the same approach as the LoopDecorator with infinite loop settings.

diff --git a/Engine/Source/Runtime/AIModule/Classes/BehaviorTree/Decorators/BTDecorator_ConditionalLoop.h b/Engine/Source/Runtime/AIModule/Classes/BehaviorTree/Decorators/BTDecorator_ConditionalLoop.h
--- a/Engine/Source/Runtime/AIModule/Classes/BehaviorTree/Decorators/BTDecorator_ConditionalLoop.h
+++ b/Engine/Source/Runtime/AIModule/Classes/BehaviorTree/Decorators/BTDecorator_ConditionalLoop.h
@@ -10,6 +10,11 @@
 
 class UBlackboardComponent;
 
+struct FBTConditionalLoopDecoratorMemory
+{
+	int32 SearchId;
+};
+
 /**
  * Conditional loop decorator node.
  * A decorator node that loops execution as long as condition is satisfied.
@@ -23,6 +28,8 @@ class AIMODULE_API UBTDecorator_ConditionalLoop : public UBTDecorator_Blackboard
 	virtual EBlackboardNotificationResult OnBlackboardKeyValueChange(const UBlackboardComponent& Blackboard, FBlackboard::FKey ChangedKeyID) override;
 	virtual void OnNodeDeactivation(FBehaviorTreeSearchData& SearchData, EBTNodeResult::Type NodeResult) override;
 
+	virtual uint16 GetInstanceMemorySize() const override;
+
 #if WITH_EDITOR
 	virtual FName GetNodeIconName() const override;
 #endif // WITH_EDITOR
diff --git a/Engine/Source/Runtime/AIModule/Private/BehaviorTree/Decorators/BTDecorator_ConditionalLoop.cpp b/Engine/Source/Runtime/AIModule/Private/BehaviorTree/Decorators/BTDecorator_ConditionalLoop.cpp
--- a/Engine/Source/Runtime/AIModule/Private/BehaviorTree/Decorators/BTDecorator_ConditionalLoop.cpp
+++ b/Engine/Source/Runtime/AIModule/Private/BehaviorTree/Decorators/BTDecorator_ConditionalLoop.cpp
@@ -29,20 +29,34 @@ EBlackboardNotificationResult UBTDecorator_ConditionalLoop::OnBlackboardKeyValue
 
 void UBTDecorator_ConditionalLoop::OnNodeDeactivation(FBehaviorTreeSearchData& SearchData, EBTNodeResult::Type NodeResult)
 {
-	if (NodeResult != EBTNodeResult::Aborted)
+	FBTConditionalLoopDecoratorMemory* DecoratorMemory = GetNodeMemory<FBTConditionalLoopDecoratorMemory>(SearchData);
+	checkf(DecoratorMemory, TEXT("Expecting to always have decorator memory available"));
+
+	// protect from infinite loop within single search
+	if (DecoratorMemory->SearchId != SearchData.SearchId)
 	{
-		const UBlackboardComponent* BlackboardComp = SearchData.OwnerComp.GetBlackboardComponent();
-		const bool bEvalResult = BlackboardComp && EvaluateOnBlackboard(*BlackboardComp);
-		UE_VLOG(SearchData.OwnerComp.GetOwner(), LogBehaviorTree, Verbose, TEXT("Loop condition: %s -> %s"),
-			bEvalResult ? TEXT("true") : TEXT("false"), (bEvalResult != IsInversed()) ? TEXT("run again!") : TEXT("break"));
+		DecoratorMemory->SearchId = SearchData.SearchId;
 
-		if (bEvalResult != IsInversed())
+		if (NodeResult != EBTNodeResult::Aborted)
 		{
-			GetParentNode()->SetChildOverride(SearchData, GetChildIndex());
+			const UBlackboardComponent* BlackboardComp = SearchData.OwnerComp.GetBlackboardComponent();
+			const bool bEvalResult = BlackboardComp && EvaluateOnBlackboard(*BlackboardComp);
+			UE_VLOG(SearchData.OwnerComp.GetOwner(), LogBehaviorTree, Verbose, TEXT("Loop condition: %s -> %s"),
+				bEvalResult ? TEXT("true") : TEXT("false"), (bEvalResult != IsInversed()) ? TEXT("run again!") : TEXT("break"));
+
+			if (bEvalResult != IsInversed())
+			{
+				GetParentNode()->SetChildOverride(SearchData, GetChildIndex());
+			}
 		}
 	}
 }
 
+uint16 UBTDecorator_ConditionalLoop::GetInstanceMemorySize() const
+{
+	return sizeof(FBTConditionalLoopDecoratorMemory);
+}
+
 #if WITH_EDITOR
 
 FName UBTDecorator_ConditionalLoop::GetNodeIconName() const
