Fixes two issues in AI Perception UnregisterSource involving pending adds (sources that were added this frame):
1. The function would erronously UE_VLOG "source doesn't seem to be registered" even if it did remove pending adds
2. The code for unregistering pending adds did not treat Sense == null as documented, i.e. Sense == null should target all senses.

diff --git a/Engine/Source/Runtime/AIModule/Private/Perception/AIPerceptionSystem.cpp b/Engine/Source/Runtime/AIModule/Private/Perception/AIPerceptionSystem.cpp
--- a/Engine/Source/Runtime/AIModule/Private/Perception/AIPerceptionSystem.cpp
+++ b/Engine/Source/Runtime/AIModule/Private/Perception/AIPerceptionSystem.cpp
@@ -357,13 +357,20 @@ void UAIPerceptionSystem::UnregisterListener(UAIPerceptionComponent& Listener)
 }
 
 void UAIPerceptionSystem::UnregisterSource(AActor& SourceActor, const TSubclassOf<UAISense> Sense)
-{	
+{
+	// Log a message if it turns out the source actor was not registered nor pending registration
+	bool bSourceWasKnown = false;
+	
 	FPerceptionStimuliSource* StimuliSource = RegisteredStimuliSources.Find(&SourceActor);
 	if (StimuliSource)
 	{
-		// single sense case
+		// Source actor was registered
+		bSourceWasKnown = true;
+
+		// A single sense can be targeted, or Sense == null for all senses
 		if (Sense)
 		{
+			// Unregister the source actor from a single sense
 			const FAISenseID SenseID = UAISense::GetSenseID(Sense);
 			if (Senses[SenseID] != nullptr && StimuliSource->RelevantSenses.ShouldRespondToChannel(Senses[SenseID]->GetSenseID()))
 			{
@@ -371,8 +378,9 @@ void UAIPerceptionSystem::UnregisterSource(AActor& SourceActor, const TSubclassO
 				StimuliSource->RelevantSenses.FilterOutChannel(SenseID);
 			}
 		}
-		else // unregister from all senses
+		else
 		{
+			// Unregister the source actor from all senses
 			for (UAISense* const SenseInstance : Senses)
 			{
 				if (SenseInstance != nullptr && StimuliSource->RelevantSenses.ShouldRespondToChannel(SenseInstance->GetSenseID()))
@@ -383,25 +391,35 @@ void UAIPerceptionSystem::UnregisterSource(AActor& SourceActor, const TSubclassO
 			StimuliSource->RelevantSenses.Clear();
 		}
 
+		// If the source actor is no longer relevant for any senses, we can remove its stimuli source entry
 		if (StimuliSource->RelevantSenses.IsEmpty())
 		{
 			SourceActor.OnEndPlay.Remove(StimuliSourceEndPlayDelegate);
 			RegisteredStimuliSources.Remove(&SourceActor);
 		}
 	}
-	else
-	{
-		UE_VLOG(this, LogAIPerception, Log, TEXT("UnregisterSource called for %s but it doesn't seem to be registered as a source"), *SourceActor.GetName());
-	}
+	
 	// Remove this from any pending adds (add/remove same frame)
 	for (int32 RemoveIndex = SourcesToRegister.Num() - 1; RemoveIndex >= 0; RemoveIndex--)
 	{
-		if (SourcesToRegister[RemoveIndex].Source == &SourceActor &&
-			SourcesToRegister[RemoveIndex].SenseID == UAISense::GetSenseID(Sense))
+		if (SourcesToRegister[RemoveIndex].Source == &SourceActor)
 		{
-			SourcesToRegister.RemoveAt(RemoveIndex, 1, false);
+			// Source actor was pending registration
+			bSourceWasKnown = true;
+
+			// A single sense can be targeted, or Sense == null for all senses
+			if (!Sense || SourcesToRegister[RemoveIndex].SenseID == UAISense::GetSenseID(Sense))
+			{
+				SourcesToRegister.RemoveAt(RemoveIndex, 1, false);
+			}
 		}
 	}
+
+	// Log if SourceActor was not registered or pending registration for any sense
+	if (!bSourceWasKnown)
+	{
+		UE_VLOG(this, LogAIPerception, Log, TEXT("UnregisterSource called for %s but it doesn't seem to be registered as a source"), *SourceActor.GetName());
+	}
 }
 
 void UAIPerceptionSystem::OnListenerRemoved(const FPerceptionListener& NewListener)
