Fix warning poping up often since last modification.

diff --git a/Engine/Source/Runtime/AIModule/Classes/BehaviorTree/BehaviorTreeComponent.h b/Engine/Source/Runtime/AIModule/Classes/BehaviorTree/BehaviorTreeComponent.h
--- a/Engine/Source/Runtime/AIModule/Classes/BehaviorTree/BehaviorTreeComponent.h
+++ b/Engine/Source/Runtime/AIModule/Classes/BehaviorTree/BehaviorTreeComponent.h
@@ -31,7 +31,7 @@ struct FBTNodeExecutionInfo
 	FBTNodeIndex SearchEnd;
 
 	/** node to be executed */
-	UBTCompositeNode* ExecuteNode;
+	const UBTCompositeNode* ExecuteNode;
 
 	/** subtree index */
 	uint16 ExecuteInstanceIdx;
@@ -78,7 +78,7 @@ struct FBTTreeStartInfo
 	bool HasPendingInitialize() const { return bPendingInitialize && IsSet(); }
 };
 
-enum class EBTBranchAction : uint8
+enum class EBTBranchAction : uint16
 {
 	None = 0x0,
 	DecoratorEvaluate = 0x1,
@@ -89,8 +89,10 @@ enum class EBTBranchAction : uint8
 	UnregisterAuxNodes = 0x10,
 	StopTree_Safe = 0x20,
 	StopTree_Forced = 0x40,
+	ActiveNodeEvaluate = 0x80,
+	SubTreeEvaluate = 0x100,
 	StopTree = StopTree_Safe | StopTree_Forced,
-	All = DecoratorEvaluate | DecoratorActivate_IfNotExecuting | DecoratorActivate_EvenIfExecuting | DecoratorDeactivate | UnregisterAuxNodes | StopTree_Safe | StopTree_Forced,
+	All = DecoratorEvaluate | DecoratorActivate_IfNotExecuting | DecoratorActivate_EvenIfExecuting | DecoratorDeactivate | UnregisterAuxNodes | StopTree_Safe | StopTree_Forced | ActiveNodeEvaluate | SubTreeEvaluate,
 };
 ENUM_CLASS_FLAGS(EBTBranchAction);
 
@@ -137,20 +139,23 @@ public:
 	void RestartTree();
 
 	/** request execution change */
-	void RequestExecution(UBTCompositeNode* RequestedOn, int32 InstanceIdx, 
+	void RequestExecution(const UBTCompositeNode* RequestedOn, int32 InstanceIdx, 
 		const UBTNode* RequestedBy, int32 RequestedByChildIndex,
 		EBTNodeResult::Type ContinueWithResult, bool bStoreForDebugger = true);
 
-	/** request execution change: helpers for decorator nodes */
-	void RequestExecution(const UBTDecorator* RequestedBy);
+	/** replaced by the RequestBranchEvaluation from decorator*/
+	void RequestExecution(const UBTDecorator* RequestedBy) { check(RequestedBy); RequestBranchEvaluation(*RequestedBy); }
 
-	/** request execution change: helpers for task nodes */
-	void RequestExecution(EBTNodeResult::Type ContinueWithResult);
+	/** replaced by RequestBranchEvaluation with EBTNodeResult */
+	void RequestExecution(EBTNodeResult::Type ContinueWithResult) { RequestBranchEvaluation(ContinueWithResult); }
 
 	/** request unregistration of aux nodes in the specified branch */
 	UE_DEPRECATED(5.0, "This function is deprecated. Please use RequestBranchDeactivation instead.")
 	void RequestUnregisterAuxNodesInBranch(const UBTCompositeNode* Node);
-		
+
+	/** request branch evaluation: helper for active node (ex: tasks) */
+	void RequestBranchEvaluation(EBTNodeResult::Type ContinueWithResult);
+
 	/** request branch evaluation: helper for decorator */
 	void RequestBranchEvaluation(const UBTDecorator& RequestedBy);
 
@@ -291,15 +296,24 @@ protected:
 
 	struct FBranchActionInfo
 	{
+		FBranchActionInfo(const EBTBranchAction InAction)
+			: Action(InAction)
+		{}
 		FBranchActionInfo(const UBTNode* InNode, const EBTBranchAction InAction)
 			: Node(InNode)
 			, Action(InAction)
 		{}
-		const UBTNode* Node;
+		FBranchActionInfo(const UBTNode* InNode, EBTNodeResult::Type InContinueWithResult, const EBTBranchAction InAction)
+			: Node(InNode)
+			, ContinueWithResult(InContinueWithResult)
+			, Action(InAction)
+		{}
+		const UBTNode* Node = nullptr;
+		EBTNodeResult::Type ContinueWithResult = EBTNodeResult::Succeeded;
 		EBTBranchAction Action;
 	};
 
-	/* Type of supended branch action */
+	/* Type of suspended branch action */
 	EBTBranchAction SuspendedBranchActions;
 
 	/** list of all pending branch action requests */
@@ -394,7 +408,7 @@ protected:
 	void ExecuteTask(UBTTaskNode* TaskNode);
 
 	/** deactivate all nodes up to requested one */
-	bool DeactivateUpTo(UBTCompositeNode* Node, uint16 NodeInstanceIdx, EBTNodeResult::Type& NodeResult, int32& OutLastDeactivatedChildIndex);
+	bool DeactivateUpTo(const UBTCompositeNode* Node, uint16 NodeInstanceIdx, EBTNodeResult::Type& NodeResult, int32& OutLastDeactivatedChildIndex);
 
 	/** returns true if execution was waiting on latent aborts and they are all finished;  */
 	bool TrackPendingLatentAborts();
@@ -443,7 +457,10 @@ protected:
 	/** Return NodeA's relative priority in regards to NodeB */
 	EBTNodeRelativePriority CalculateRelativePriority(const UBTNode* NodeA, const UBTNode* NodeB) const;
 
-	/** Activate a branch as the decorator conditions have changed */
+	/** Evaluate a branch as current active node is finished */
+	void EvaluateBranch(EBTNodeResult::Type LastResult);
+
+	/** Evaluate a branch as the decorator conditions have changed */
 	void EvaluateBranch(const UBTDecorator& RequestedBy);
 
 	/** Activate a branch as the decorator conditions are now passing */
diff --git a/Engine/Source/Runtime/AIModule/Private/BehaviorTree/BehaviorTreeComponent.cpp b/Engine/Source/Runtime/AIModule/Private/BehaviorTree/BehaviorTreeComponent.cpp
--- a/Engine/Source/Runtime/AIModule/Private/BehaviorTree/BehaviorTreeComponent.cpp
+++ b/Engine/Source/Runtime/AIModule/Private/BehaviorTree/BehaviorTreeComponent.cpp
@@ -264,7 +264,7 @@ void UBehaviorTreeComponent::StopTree(EBTStopMode::Type StopMode)
 	if ((SuspendedBranchActions & EBTBranchAction::StopTree) != EBTBranchAction::None)
 	{
 		UE_VLOG(GetOwner(), LogBehaviorTree, Verbose, TEXT("Stop tree(%s) queued up"), bForcedStop ? TEXT("Forced") : TEXT("Safe"));
-		PendingBranchActionRequests.Emplace(nullptr, bForcedStop ? EBTBranchAction::StopTree_Forced : EBTBranchAction::StopTree_Safe);
+		PendingBranchActionRequests.Emplace(bForcedStop ? EBTBranchAction::StopTree_Forced : EBTBranchAction::StopTree_Safe);
 		return;
 	}
 
@@ -726,6 +726,7 @@ void UBehaviorTreeComponent::ResumeBranchActions()
 					// Since we have been queued up, decorator might have been removed from active nodes, need to make sure it is still there.
 					if (!IsAuxNodeActive(RequestedBy))
 					{
+						UE_VLOG(GetOwner(), LogBehaviorTree, Verbose, TEXT("Request deactivation skipped because decorator(%s) is not active anymore"), *UBehaviorTreeTypes::DescribeNodeHelper(RequestedBy));
 						break;
 					}
 
@@ -741,6 +742,7 @@ void UBehaviorTreeComponent::ResumeBranchActions()
 					// Since we have been queued up, decorator might have been removed from active nodes, need to make sure it is still there.
 					if (!IsAuxNodeActive(RequestedBy))
 					{
+						UE_VLOG(GetOwner(), LogBehaviorTree, Verbose, TEXT("Request deactivation skipped because decorator(%s) is not active anymore"), *UBehaviorTreeTypes::DescribeNodeHelper(RequestedBy));
 						break;
 					}
 
@@ -754,6 +756,7 @@ void UBehaviorTreeComponent::ResumeBranchActions()
 					// Since we have been queued up, decorator might have been removed from active nodes, need to make sure it is still there.
 					if (!IsAuxNodeActive(RequestedBy))
 					{
+						UE_VLOG(GetOwner(), LogBehaviorTree, Verbose, TEXT("Request deactivation skipped because decorator(%s) is not active anymore"), *UBehaviorTreeTypes::DescribeNodeHelper(RequestedBy));
 						break;
 					}
 
@@ -770,18 +773,39 @@ void UBehaviorTreeComponent::ResumeBranchActions()
 				case EBTBranchAction::StopTree_Forced:
 				{
 					StopTree(Info.Action == EBTBranchAction::StopTree_Forced ? EBTStopMode::Forced : EBTStopMode::Safe);
+					break;
+				}
+				case EBTBranchAction::ActiveNodeEvaluate:
+				{
+					const UBTNode* ActiveNode = GetActiveNode();
+					if (ActiveNode != Info.Node)
+					{
+						UE_VLOG(GetOwner(), LogBehaviorTree, Verbose, TEXT("Request evaluation skipped because node(%s) is not active anymore"), *UBehaviorTreeTypes::DescribeNodeHelper(ActiveNode), *UEnum::GetValueAsString(Info.ContinueWithResult));
+						break;
+					}
+
+					EvaluateBranch(Info.ContinueWithResult);
+					break;
+				}
+				case EBTBranchAction::SubTreeEvaluate:
+				{
+					const UBTCompositeNode* BranchRoot = CastChecked<UBTCompositeNode>(Info.Node);
+
+					const UBTNode* RootNode = InstanceStack.Num() ? InstanceStack[ActiveInstanceIdx].RootNode : nullptr;
+					if (RootNode != BranchRoot)
+					{
+						UE_VLOG(GetOwner(), LogBehaviorTree, Verbose, TEXT("Sub tree evaluation skipped because node(%s) is not active anymore"), *UBehaviorTreeTypes::DescribeNodeHelper(BranchRoot));
+						break;
+					}
+
+					RequestExecution(BranchRoot, ActiveInstanceIdx, BranchRoot, 0, EBTNodeResult::InProgress);
+					break;
 				}
 			}
 		}
 	}
 }
 
-void UBehaviorTreeComponent::RequestExecution(const UBTDecorator* RequestedBy)
-{
-	check(RequestedBy);
-	RequestBranchEvaluation(*RequestedBy);
-}
-
 void UBehaviorTreeComponent::RequestBranchEvaluation(const UBTDecorator& RequestedBy)
 {
 	if ((SuspendedBranchActions & EBTBranchAction::DecoratorEvaluate) != EBTBranchAction::None)
@@ -917,12 +941,29 @@ EBTNodeRelativePriority UBehaviorTreeComponent::CalculateRelativePriority(const
 	return RelativePriority;
 }
 
-void UBehaviorTreeComponent::RequestExecution(EBTNodeResult::Type LastResult)
+void UBehaviorTreeComponent::RequestBranchEvaluation(EBTNodeResult::Type ContinueWithResult)
 {
 	// task helpers can't continue with InProgress or Aborted result, it should be handled 
 	// either by decorator helper or regular RequestExecution() (6 param version)
+	if (ContinueWithResult == EBTNodeResult::Aborted || ContinueWithResult == EBTNodeResult::InProgress)
+	{
+		return;
+	}
 
-	if (LastResult != EBTNodeResult::Aborted && LastResult != EBTNodeResult::InProgress && InstanceStack.IsValidIndex(ActiveInstanceIdx))
+	if ((SuspendedBranchActions & EBTBranchAction::ActiveNodeEvaluate) != EBTBranchAction::None)
+	{
+		const UBTNode* ActiveNode = GetActiveNode();
+		UE_VLOG(GetOwner(), LogBehaviorTree, Verbose, TEXT("Request evaluation queued up for node(%s) with result(%s)"), *UBehaviorTreeTypes::DescribeNodeHelper(ActiveNode), *UEnum::GetValueAsString(ContinueWithResult));
+		PendingBranchActionRequests.Emplace(ActiveNode, ContinueWithResult, EBTBranchAction::ActiveNodeEvaluate);
+		return;
+	}
+
+	EvaluateBranch(ContinueWithResult);
+}
+
+void UBehaviorTreeComponent::EvaluateBranch(EBTNodeResult::Type ContinueWithResult)
+{
+	if (InstanceStack.IsValidIndex(ActiveInstanceIdx))
 	{
 		const FBehaviorTreeInstance& ActiveInstance = InstanceStack[ActiveInstanceIdx];
 		UBTCompositeNode* ExecuteParent = (ActiveInstance.ActiveNode == NULL) ? ActiveInstance.RootNode :
@@ -931,20 +972,20 @@ void UBehaviorTreeComponent::RequestExecution(EBTNodeResult::Type LastResult)
 
 		RequestExecution(ExecuteParent, InstanceStack.Num() - 1,
 			ActiveInstance.ActiveNode ? ActiveInstance.ActiveNode : ActiveInstance.RootNode, -1,
-			LastResult, false);
+			ContinueWithResult, false);
 	}
 }
 
 static void FindCommonParent(const TArray<FBehaviorTreeInstance>& Instances, const TArray<FBehaviorTreeInstanceId>& KnownInstances,
-							 UBTCompositeNode* InNodeA, uint16 InstanceIdxA,
-							 UBTCompositeNode* InNodeB, uint16 InstanceIdxB,
-							 UBTCompositeNode*& CommonParentNode, uint16& CommonInstanceIdx)
+							 const UBTCompositeNode* InNodeA, uint16 InstanceIdxA,
+							 const UBTCompositeNode* InNodeB, uint16 InstanceIdxB,
+							 const UBTCompositeNode*& CommonParentNode, uint16& CommonInstanceIdx)
 {
 	// find two nodes in the same instance (choose lower index = closer to root)
 	CommonInstanceIdx = (InstanceIdxA <= InstanceIdxB) ? InstanceIdxA : InstanceIdxB;
 
-	UBTCompositeNode* NodeA = (CommonInstanceIdx == InstanceIdxA) ? InNodeA : Instances[CommonInstanceIdx].ActiveNode->GetParentNode();
-	UBTCompositeNode* NodeB = (CommonInstanceIdx == InstanceIdxB) ? InNodeB : Instances[CommonInstanceIdx].ActiveNode->GetParentNode();
+	const UBTCompositeNode* NodeA = (CommonInstanceIdx == InstanceIdxA) ? InNodeA : Instances[CommonInstanceIdx].ActiveNode->GetParentNode();
+	const UBTCompositeNode* NodeB = (CommonInstanceIdx == InstanceIdxB) ? InNodeB : Instances[CommonInstanceIdx].ActiveNode->GetParentNode();
 
 	// special case: node was taken from CommonInstanceIdx, but it had ActiveNode set to root (no parent)
 	if (!NodeA && CommonInstanceIdx != InstanceIdxA)
@@ -1003,7 +1044,7 @@ void UBehaviorTreeComponent::ScheduleExecutionUpdate()
 	bRequestedFlowUpdate = true;
 }
 
-void UBehaviorTreeComponent::RequestExecution(UBTCompositeNode* RequestedOn, int32 InstanceIdx, const UBTNode* RequestedBy,
+void UBehaviorTreeComponent::RequestExecution(const UBTCompositeNode* RequestedOn, int32 InstanceIdx, const UBTNode* RequestedBy,
 											  int32 RequestedByChildIndex, EBTNodeResult::Type ContinueWithResult, bool bStoreForDebugger)
 {
 	SCOPE_CYCLE_COUNTER(STAT_AI_BehaviorTree_SearchTime);
@@ -1043,7 +1084,11 @@ void UBehaviorTreeComponent::RequestExecution(UBTCompositeNode* RequestedOn, int
 	// make sure that the request is not coming from a node that has pending branch actions since it won't be accessible anymore
 	if (SuspendedBranchActions != EBTBranchAction::None)
 	{
-		UE_CVLOG(SuspendedBranchActions != EBTBranchAction::UnregisterAuxNodes, GetOwner(), LogBehaviorTree, Warning, TEXT("Caller should be converted to new Evaluate/Activate/DeactivateBranch API instead of using this RequestExecution directly"));
+		if ((SuspendedBranchActions & ~(EBTBranchAction::UnregisterAuxNodes | EBTBranchAction::StopTree )) != EBTBranchAction::None)
+		{
+			UE_VLOG(GetOwner(), LogBehaviorTree, Warning, TEXT("Caller should be converted to new Evaluate/Activate/DeactivateBranch API instead of using this RequestExecution directly"));
+		}
+
 		for (const FBranchActionInfo& Info : PendingBranchActionRequests)
 		{
 			const UBTCompositeNode* BranchRoot = nullptr;
@@ -1166,7 +1211,7 @@ void UBehaviorTreeComponent::RequestExecution(UBTCompositeNode* RequestedOn, int
 		}
 
 		// update common parent: requesting node with prev common/active node
-		UBTCompositeNode* CurrentNode = ExecutionRequest.ExecuteNode;
+		const UBTCompositeNode* CurrentNode = ExecutionRequest.ExecuteNode;
 		uint16 CurrentInstanceIdx = ExecutionRequest.ExecuteInstanceIdx;
 		if (ExecutionRequest.ExecuteNode == NULL)
 		{
@@ -1180,16 +1225,16 @@ void UBehaviorTreeComponent::RequestExecution(UBTCompositeNode* RequestedOn, int
 
 		if (ExecutionRequest.ExecuteNode != RequestedOn)
 		{
-			UBTCompositeNode* CommonParent = NULL;
+			const UBTCompositeNode* CommonParent = NULL;
 			uint16 CommonInstanceIdx = MAX_uint16;
 
 			FindCommonParent(InstanceStack, KnownInstances, RequestedOn, InstanceIdx, CurrentNode, CurrentInstanceIdx, CommonParent, CommonInstanceIdx);
 
 			// check decorators between common parent and restart parent
 			int32 ItInstanceIdx = InstanceIdx;
-			for (UBTCompositeNode* It = RequestedOn; It && It != CommonParent;)
+			for (const UBTCompositeNode* It = RequestedOn; It && It != CommonParent;)
 			{
-				UBTCompositeNode* ParentNode = It->GetParentNode();
+				const UBTCompositeNode* ParentNode = It->GetParentNode();
 				int32 ChildIdx = INDEX_NONE;
 
 				if (ParentNode == nullptr)
@@ -1803,7 +1848,7 @@ void UBehaviorTreeComponent::ProcessExecutionRequest()
 		}
 
 		FBehaviorTreeInstance& ActiveInstance = InstanceStack[ActiveInstanceIdx];
-		UBTCompositeNode* TestNode = ExecutionRequest.ExecuteNode;
+		const UBTCompositeNode* TestNode = ExecutionRequest.ExecuteNode;
 		SearchData.AssignSearchId();
 		SearchData.bPostponeSearch = false;
 		SearchData.bSearchInProgress = true;
@@ -1859,7 +1904,7 @@ void UBehaviorTreeComponent::ProcessExecutionRequest()
 		{
 			BT_SEARCHLOG(SearchData, Verbose, TEXT("Testing node: %s"), *UBehaviorTreeTypes::DescribeNodeHelper(TestNode));
 			const int32 ChildBranchIdx = TestNode->FindChildToExecute(SearchData, NodeResult);
-			UBTNode* StoreNode = TestNode;
+			const UBTNode* StoreNode = TestNode;
 
 			if (SearchData.bPostponeSearch)
 			{
@@ -1869,7 +1914,7 @@ void UBehaviorTreeComponent::ProcessExecutionRequest()
 			}
 			else if (ChildBranchIdx == BTSpecialChild::ReturnToParent)
 			{
-				UBTCompositeNode* ChildNode = TestNode;
+				const UBTCompositeNode* ChildNode = TestNode;
 				TestNode = TestNode->GetParentNode();
 
 				// does it want to move up the tree?
@@ -2079,9 +2124,9 @@ void UBehaviorTreeComponent::RollbackSearchChanges()
 	}
 }
 
-bool UBehaviorTreeComponent::DeactivateUpTo(UBTCompositeNode* Node, uint16 NodeInstanceIdx, EBTNodeResult::Type& NodeResult, int32& OutLastDeactivatedChildIndex)
+bool UBehaviorTreeComponent::DeactivateUpTo(const UBTCompositeNode* Node, uint16 NodeInstanceIdx, EBTNodeResult::Type& NodeResult, int32& OutLastDeactivatedChildIndex)
 {
-	UBTNode* DeactivatedChild = InstanceStack[ActiveInstanceIdx].ActiveNode;
+	const UBTNode* DeactivatedChild = InstanceStack[ActiveInstanceIdx].ActiveNode;
 	bool bDeactivateRoot = true;
 
 	if (DeactivatedChild == NULL && ActiveInstanceIdx > NodeInstanceIdx)
@@ -2093,7 +2138,7 @@ bool UBehaviorTreeComponent::DeactivateUpTo(UBTCompositeNode* Node, uint16 NodeI
 
 	while (DeactivatedChild)
 	{
-		UBTCompositeNode* NotifyParent = DeactivatedChild->GetParentNode();
+		const UBTCompositeNode* NotifyParent = DeactivatedChild->GetParentNode();
 		if (NotifyParent)
 		{
 			OutLastDeactivatedChildIndex = NotifyParent->GetChildIndex(SearchData, *DeactivatedChild);
@@ -2523,6 +2568,13 @@ bool UBehaviorTreeComponent::PushInstance(UBehaviorTree& TreeAsset)
 
 		FBehaviorTreeDelegates::OnTreeStarted.Broadcast(*this, TreeAsset);
 
+		if ((SuspendedBranchActions & EBTBranchAction::SubTreeEvaluate) != EBTBranchAction::None)
+		{
+			UE_VLOG(GetOwner(), LogBehaviorTree, Verbose, TEXT("Evaluate sub tree(%s) root(%s) queued up"), *TreeAsset.GetName(), *UBehaviorTreeTypes::DescribeNodeHelper(RootNode));
+			PendingBranchActionRequests.Emplace(RootNode, EBTBranchAction::SubTreeEvaluate);
+			return true;
+		}
+
 		// start new task
 		RequestExecution(RootNode, ActiveInstanceIdx, RootNode, 0, EBTNodeResult::InProgress);
 		return true;
@@ -3286,4 +3338,4 @@ bool UBehaviorTreeComponent::IsDebuggerActive()
 #else
 	return false;
 #endif
-}
\ No newline at end of file
+}
