[AITask] added NewAITask and NewBTAITask overloads to be able to create AITask from a UClass
#rnx

diff --git a/Engine/Source/Runtime/AIModule/Classes/BehaviorTree/BTNode.h b/Engine/Source/Runtime/AIModule/Classes/BehaviorTree/BTNode.h
--- a/Engine/Source/Runtime/AIModule/Classes/BehaviorTree/BTNode.h
+++ b/Engine/Source/Runtime/AIModule/Classes/BehaviorTree/BTNode.h
@@ -177,6 +177,14 @@ class AIMODULE_API UBTNode : public UObject, public IGameplayTaskOwnerInterface
 		bOwnsGameplayTasks = true;
 		return UAITask::NewAITask<T>(*BTComponent.GetAIOwner(), *this, TEXT("Behavior"));
 	}
+	
+	template <class T>
+	T* NewBTAITask(const UClass& Class, UBehaviorTreeComponent& BTComponent)
+	{
+		check(BTComponent.GetAIOwner());
+		bOwnsGameplayTasks = true;
+		return UAITask::NewAITask<T>(Class, *BTComponent.GetAIOwner(), *this, TEXT("Behavior"));
+	}
 
 	/** node name */
 	UPROPERTY(Category=Description, EditAnywhere)
diff --git a/Engine/Source/Runtime/AIModule/Classes/Tasks/AITask.h b/Engine/Source/Runtime/AIModule/Classes/Tasks/AITask.h
--- a/Engine/Source/Runtime/AIModule/Classes/Tasks/AITask.h
+++ b/Engine/Source/Runtime/AIModule/Classes/Tasks/AITask.h
@@ -45,36 +45,42 @@ public:
 	template <class T>
 	static T* NewAITask(AAIController& AIOwner, IGameplayTaskOwnerInterface& InTaskOwner, FName InstanceName = FName())
 	{
-		T* TaskInstance = NewObject<T>();
-		TaskInstance->InstanceName = InstanceName;
-		TaskInstance->InitAITask(AIOwner, InTaskOwner);
-		return TaskInstance;
+		return NewAITask<T>(*T::StaticClass(), AIOwner, InTaskOwner, InstanceName);
 	}
 
 	template <class T>
 	static T* NewAITask(AAIController& AIOwner, IGameplayTaskOwnerInterface& InTaskOwner, EAITaskPriority InPriority, FName InstanceName = FName())
 	{
-		T* TaskInstance = NewObject<T>();
-		TaskInstance->InstanceName = InstanceName;
-		TaskInstance->InitAITask(AIOwner, InTaskOwner, (uint8)InPriority);
-		return TaskInstance;
+		return NewAITask<T>(*T::StaticClass(), AIOwner, InTaskOwner, InPriority, InstanceName);
 	}
 
 	template <class T>
 	static T* NewAITask(AAIController& AIOwner, FName InstanceName = FName())
 	{
-		T* TaskInstance = NewObject<T>();
+		return NewAITask<T>(*T::StaticClass(), AIOwner, AIOwner, InstanceName);
+	}
+
+	template <class T>
+	static T* NewAITask(AAIController& AIOwner, EAITaskPriority InPriority, FName InstanceName = FName())
+	{
+		return NewAITask<T>(*T::StaticClass(), AIOwner, AIOwner, InPriority, InstanceName);
+	}
+
+	template <class T>
+	static T* NewAITask(const UClass& Class, AAIController& AIOwner, IGameplayTaskOwnerInterface& InTaskOwner, FName InstanceName = FName())
+	{
+		T* TaskInstance = NewObject<T>(GetTransientPackage(), &Class);
 		TaskInstance->InstanceName = InstanceName;
-		TaskInstance->InitAITask(AIOwner, AIOwner);
+		TaskInstance->InitAITask(AIOwner, InTaskOwner);
 		return TaskInstance;
 	}
 
 	template <class T>
-	static T* NewAITask(AAIController& AIOwner, EAITaskPriority InPriority, FName InstanceName = FName())
+	static T* NewAITask(const UClass& Class, AAIController& AIOwner, IGameplayTaskOwnerInterface& InTaskOwner, EAITaskPriority InPriority, FName InstanceName = FName())
 	{
-		T* TaskInstance = NewObject<T>();
+		T* TaskInstance = NewObject<T>(GetTransientPackage(), &Class);
 		TaskInstance->InstanceName = InstanceName;
-		TaskInstance->InitAITask(AIOwner, AIOwner, (uint8)InPriority);
+		TaskInstance->InitAITask(AIOwner, InTaskOwner, (uint8)InPriority);
 		return TaskInstance;
 	}
 };
