Add the possibility for blueprint EnvQuery generator to generate from actors.
#rnx

diff --git a/Engine/Source/Runtime/AIModule/Classes/EnvironmentQuery/Generators/EnvQueryGenerator_BlueprintBase.h b/Engine/Source/Runtime/AIModule/Classes/EnvironmentQuery/Generators/EnvQueryGenerator_BlueprintBase.h
--- a/Engine/Source/Runtime/AIModule/Classes/EnvironmentQuery/Generators/EnvQueryGenerator_BlueprintBase.h
+++ b/Engine/Source/Runtime/AIModule/Classes/EnvironmentQuery/Generators/EnvQueryGenerator_BlueprintBase.h
@@ -38,6 +38,9 @@ class AIMODULE_API UEnvQueryGenerator_BlueprintBase : public UEnvQueryGenerator
 	UFUNCTION(BlueprintImplementableEvent, Category = Generator)
 	void DoItemGeneration(const TArray<FVector>& ContextLocations) const;
 
+	UFUNCTION(BlueprintImplementableEvent, Category = Generator)
+	void DoItemGenerationFromActors(const TArray<AActor*>& ContextActors) const;
+
 	virtual void GenerateItems(FEnvQueryInstance& QueryInstance) const override;
 
 	virtual FText GetDescriptionTitle() const override;
@@ -53,6 +56,15 @@ class AIMODULE_API UEnvQueryGenerator_BlueprintBase : public UEnvQueryGenerator
 	UObject* GetQuerier() const;
 
 private:
+	enum class ECallMode
+	{
+		Invalid,
+		FromVectors,
+		FromActors
+	};
+
 	/** this is valid and set only within GenerateItems call */
 	mutable FEnvQueryInstance* CachedQueryInstance;
+
+	ECallMode CallMode = ECallMode::Invalid;
 };
diff --git a/Engine/Source/Runtime/AIModule/Private/EnvironmentQuery/Generators/EnvQueryGenerator_BlueprintBase.cpp b/Engine/Source/Runtime/AIModule/Private/EnvironmentQuery/Generators/EnvQueryGenerator_BlueprintBase.cpp
--- a/Engine/Source/Runtime/AIModule/Private/EnvironmentQuery/Generators/EnvQueryGenerator_BlueprintBase.cpp
+++ b/Engine/Source/Runtime/AIModule/Private/EnvironmentQuery/Generators/EnvQueryGenerator_BlueprintBase.cpp
@@ -9,12 +9,40 @@
 
 #define LOCTEXT_NAMESPACE "EnvQueryGenerator"
 
+namespace UE::AI::EQS::Private
+{
+	FORCEINLINE bool DoesImplementBPFunction(const FName FuncName, const UObject& Ob, const UClass* StopAtClass)
+	{
+		const UFunction* Function = Ob.GetClass()->FindFunctionByName(FuncName);
+		checkSlow(Function)
+		return (Function->GetOuter() != StopAtClass);
+	}
+}
+
 UEnvQueryGenerator_BlueprintBase::UEnvQueryGenerator_BlueprintBase(const FObjectInitializer& ObjectInitializer)
 	: Super(ObjectInitializer)
 {
+	using namespace UE::AI::EQS::Private;
 	Context = UEnvQueryContext_Querier::StaticClass();
 	ItemType = UEnvQueryItemType_Actor::StaticClass();
 	GeneratedItemType = UEnvQueryItemType_Actor::StaticClass();
+
+	const UClass* StopAtClass = UEnvQueryGenerator_BlueprintBase::StaticClass();
+	const bool bImplementsGenerateFromVectors = DoesImplementBPFunction(GET_FUNCTION_NAME_CHECKED(UEnvQueryGenerator_BlueprintBase, DoItemGeneration), *this, StopAtClass);
+	const bool bImplementsGenerateFromActors = DoesImplementBPFunction(GET_FUNCTION_NAME_CHECKED(UEnvQueryGenerator_BlueprintBase, DoItemGenerationFromActors), *this, StopAtClass);
+	if (bImplementsGenerateFromVectors)
+	{
+		CallMode = ECallMode::FromVectors;
+	}
+	else if (bImplementsGenerateFromActors)
+	{
+		CallMode = ECallMode::FromActors;
+	}
+
+	if (CallMode == ECallMode::Invalid && GetClass() != StopAtClass)
+	{
+		UE_LOG(LogEQS, Error, TEXT("Blueprint Generator class %s doesn't override DoItemGeneration or DoItemGenerationFromActors and can't generate anything."), *GetNameSafe(GetClass()));
+	}
 }
 
 void UEnvQueryGenerator_BlueprintBase::PostInitProperties()
@@ -30,11 +58,24 @@ UWorld* UEnvQueryGenerator_BlueprintBase::GetWorld() const
 
 void UEnvQueryGenerator_BlueprintBase::GenerateItems(FEnvQueryInstance& QueryInstance) const
 {
-	TArray<FVector> ContextLocations;
-	QueryInstance.PrepareContext(Context, ContextLocations);
-
 	CachedQueryInstance = &QueryInstance;
-	const_cast<UEnvQueryGenerator_BlueprintBase*>(this)->DoItemGeneration(ContextLocations);
+	if (CallMode == ECallMode::FromVectors)
+	{
+		TArray<FVector> ContextLocations;
+		QueryInstance.PrepareContext(Context, ContextLocations);
+		const_cast<UEnvQueryGenerator_BlueprintBase*>(this)->DoItemGeneration(ContextLocations);
+		
+	}
+	else if (CallMode == ECallMode::FromActors)
+	{
+		TArray<AActor*> ContextActors;
+		QueryInstance.PrepareContext(Context, ContextActors);;
+		const_cast<UEnvQueryGenerator_BlueprintBase*>(this)->DoItemGenerationFromActors(ContextActors);
+	}
+	else
+	{
+		UE_LOG(LogEQS, Warning, TEXT("Blueprint generator %s with class %s doesn't override DoItemGeneration or DoItemGenerationFromActors and can't generate anything."), *GetName(), *GetNameSafe(GetClass()));
+	}
 	CachedQueryInstance = NULL;
 }
 
