[UEnvQueryGenerator_ActorsOfClass] allowed abstract classes to be used in the search
#rnx

diff --git a/Engine/Source/Runtime/AIModule/Classes/EnvironmentQuery/Generators/EnvQueryGenerator_ActorsOfClass.h b/Engine/Source/Runtime/AIModule/Classes/EnvironmentQuery/Generators/EnvQueryGenerator_ActorsOfClass.h
--- a/Engine/Source/Runtime/AIModule/Classes/EnvironmentQuery/Generators/EnvQueryGenerator_ActorsOfClass.h
+++ b/Engine/Source/Runtime/AIModule/Classes/EnvironmentQuery/Generators/EnvQueryGenerator_ActorsOfClass.h
@@ -16,7 +16,7 @@ class AIMODULE_API UEnvQueryGenerator_ActorsOfClass : public UEnvQueryGenerator
 {
 	GENERATED_UCLASS_BODY()
 
-	UPROPERTY(EditDefaultsOnly, Category=Generator)
+	UPROPERTY(EditDefaultsOnly, Category=Generator, meta=(AllowAbstract))
 	TSubclassOf<AActor> SearchedActorClass;
 
 	/** If true, this will only returns actors of the specified class within the SearchRadius of the SearchCenter context.  If false, it will return ALL actors of the specified class in the world. */
