Inclusive language change for Runtime/AIModule

diff --git a/Engine/Source/Runtime/AIModule/Classes/Navigation/CrowdFollowingComponent.h b/Engine/Source/Runtime/AIModule/Classes/Navigation/CrowdFollowingComponent.h
--- a/Engine/Source/Runtime/AIModule/Classes/Navigation/CrowdFollowingComponent.h
+++ b/Engine/Source/Runtime/AIModule/Classes/Navigation/CrowdFollowingComponent.h
@@ -76,7 +76,7 @@ class AIMODULE_API UCrowdFollowingComponent : public UPathFollowingComponent, pu
 	/** pass desired position to movement component (after resolving collisions between crowd agents) */
 	virtual void ApplyCrowdAgentPosition(const FVector& NewPosition);
 
-	/** master switch for crowd steering & avoidance */
+	/** main switch for crowd steering & avoidance */
 	UFUNCTION(BlueprintCallable, Category = "Crowd")
 	virtual void SuspendCrowdSteering(bool bSuspend);
 
diff --git a/Engine/Source/Runtime/AIModule/Private/BehaviorTree/BehaviorTreeManager.cpp b/Engine/Source/Runtime/AIModule/Private/BehaviorTree/BehaviorTreeManager.cpp
--- a/Engine/Source/Runtime/AIModule/Private/BehaviorTree/BehaviorTreeManager.cpp
+++ b/Engine/Source/Runtime/AIModule/Private/BehaviorTree/BehaviorTreeManager.cpp
@@ -97,8 +97,8 @@ static void MergeDecoratorOpsHelper(TArray<FBTDecoratorLogic>& LinkOps, const TA
 	if (NumOriginalDecorators > 0)
 	{
 		// and operator for two groups of composites: original and new one
-		FBTDecoratorLogic MasterAndOp(EBTDecoratorLogic::And, LinkOps.Num() ? 2 : (NumOriginalDecorators + 1));
-		LinkOps.Insert(MasterAndOp, 0);
+		FBTDecoratorLogic MainAndOp(EBTDecoratorLogic::And, LinkOps.Num() ? 2 : (NumOriginalDecorators + 1));
+		LinkOps.Insert(MainAndOp, 0);
 
 		if (NumOriginalOps == 0)
 		{
