[BTDecorator_Blackboard]
- now using 'StringValue' property instead of 'IntValue' to store value for enum based keys to better support enumeration modifications (e.g. adding/removing values). 'IntValue' is kept in sync since still used for runtime evaluation.
- now using 'IntValue' to store enum value instead of index since runtime arithtmetic operation expects to compare values (bugfix). Also added validation for enumerations with values that could not be properly stored in enum based keys at runtime.
- improved Editor refresh on live UserDefinedEnum modifications
- current approach is to refresh only nodes displayed in the editor and not all loaded nodes. For the other we rely on the initialization when loaded to refresh all descriptions (i.e. InitializeFromAsset). Note that custom version serialization was not possible since we need related types to be loaded and keys to be resolved.
- fixed uninitialized properties in case an enum key is not fully setup (i.e. not pointing to a valid enumeration)

diff --git a/Engine/Source/Editor/BehaviorTreeEditor/Private/BehaviorTreeGraphNode_CompositeDecorator.cpp b/Engine/Source/Editor/BehaviorTreeEditor/Private/BehaviorTreeGraphNode_CompositeDecorator.cpp
--- a/Engine/Source/Editor/BehaviorTreeEditor/Private/BehaviorTreeGraphNode_CompositeDecorator.cpp
+++ b/Engine/Source/Editor/BehaviorTreeEditor/Private/BehaviorTreeGraphNode_CompositeDecorator.cpp
@@ -9,6 +9,7 @@
 #include "BehaviorTreeDecoratorGraphNode_Decorator.h"
 #include "EdGraphSchema_BehaviorTreeDecorator.h"
 #include "BehaviorTree/BehaviorTree.h"
+#include "BehaviorTree/Decorators/BTDecorator_Blackboard.h"
 
 #define LOCTEXT_NAMESPACE "BehaviorTreeEditor"
 
@@ -368,6 +369,11 @@ void UBehaviorTreeGraphNode_CompositeDecorator::BuildDescription()
 				bPendingNotOp = false;
 			}
 
+			// Composite decorator based on blackboard might need to rebuild inner decorators description before aggregating it
+			if (UBTDecorator_Blackboard* Decorator = Cast<UBTDecorator_Blackboard>(NodeInstances[TestOp.Number]))
+			{
+				Decorator->BuildDescription();
+			}
 			Description += NodeInstances[TestOp.Number]->GetStaticDescription();
 			UpdateLogicOpStack(OpStack, Description, Indent);
 		}
diff --git a/Engine/Source/Editor/BehaviorTreeEditor/Private/DetailCustomizations/BlackboardDecoratorDetails.cpp b/Engine/Source/Editor/BehaviorTreeEditor/Private/DetailCustomizations/BlackboardDecoratorDetails.cpp
--- a/Engine/Source/Editor/BehaviorTreeEditor/Private/DetailCustomizations/BlackboardDecoratorDetails.cpp
+++ b/Engine/Source/Editor/BehaviorTreeEditor/Private/DetailCustomizations/BlackboardDecoratorDetails.cpp
@@ -73,8 +73,7 @@ void FBlackboardDecoratorDetails::CustomizeDetails( IDetailLayoutBuilder& Detail
 
 #endif // WITH_EDITORONLY_DATA
 
-	IntValueProperty = DetailLayout.GetProperty(GET_MEMBER_NAME_CHECKED(UBTDecorator_Blackboard, IntValue));
-	IDetailPropertyRow& IntValueRow = BBCategory.AddProperty(IntValueProperty);
+	IDetailPropertyRow& IntValueRow = BBCategory.AddProperty(DetailLayout.GetProperty(GET_MEMBER_NAME_CHECKED(UBTDecorator_Blackboard, IntValue)));
 	IntValueRow.Visibility(TAttribute<EVisibility>::Create(TAttribute<EVisibility>::FGetter::CreateSP(this, &FBlackboardDecoratorDetails::GetIntValueVisibility)));
 	IntValueRow.IsEnabled(PropertyEditCheck);
 
@@ -82,17 +81,18 @@ void FBlackboardDecoratorDetails::CustomizeDetails( IDetailLayoutBuilder& Detail
 	FloatValueRow.Visibility(TAttribute<EVisibility>::Create(TAttribute<EVisibility>::FGetter::CreateSP(this, &FBlackboardDecoratorDetails::GetFloatValueVisibility)));
 	FloatValueRow.IsEnabled(PropertyEditCheck);
 
-	IDetailPropertyRow& StringValueRow = BBCategory.AddProperty(DetailLayout.GetProperty(GET_MEMBER_NAME_CHECKED(UBTDecorator_Blackboard, StringValue)));
+	StringValueProperty = DetailLayout.GetProperty(GET_MEMBER_NAME_CHECKED(UBTDecorator_Blackboard, StringValue));
+	IDetailPropertyRow& StringValueRow = BBCategory.AddProperty(StringValueProperty);
 	StringValueRow.Visibility(TAttribute<EVisibility>::Create(TAttribute<EVisibility>::FGetter::CreateSP(this, &FBlackboardDecoratorDetails::GetStringValueVisibility)));
 	StringValueRow.IsEnabled(PropertyEditCheck);
 
-	IDetailPropertyRow& EnumValueRow = BBCategory.AddProperty(IntValueProperty);
+	IDetailPropertyRow& EnumValueRow = BBCategory.AddProperty(StringValueProperty);
 	EnumValueRow.Visibility(TAttribute<EVisibility>::Create(TAttribute<EVisibility>::FGetter::CreateSP(this, &FBlackboardDecoratorDetails::GetEnumValueVisibility)));
 	EnumValueRow.IsEnabled(PropertyEditCheck);
 	EnumValueRow.CustomWidget()
 		.NameContent()
 		[
-			IntValueProperty->CreatePropertyNameWidget()
+			StringValueProperty->CreatePropertyNameWidget()
 		]
 		.ValueContent()
 		[
@@ -129,10 +129,11 @@ void FBlackboardDecoratorDetails::CacheBlackboardData(IDetailLayoutBuilder& Deta
 void FBlackboardDecoratorDetails::OnKeyIDChanged()
 {
 	CachedOperationType = EBlackboardKeyOperation::Basic;
-	CachedKeyType = NULL;
+	CachedCustomObjectType = nullptr;
+	CachedKeyType = nullptr;
 
 	UBlackboardData* Blackboard = CachedBlackboardAsset.Get();
-	if (Blackboard == NULL)
+	if (Blackboard == nullptr)
 	{
 		return;
 	}
@@ -147,11 +148,6 @@ void FBlackboardDecoratorDetails::OnKeyIDChanged()
 			CachedKeyType = KeyEntry->KeyType->GetClass();
 			CachedOperationType = KeyEntry->KeyType->GetTestOperation();
 		}
-		else
-		{
-			CachedKeyType = nullptr;
-			CachedOperationType = 0;
-		}
 	}
 
 	// special handling of enum type: cache all names for combo box (display names)
@@ -170,19 +166,37 @@ void FBlackboardDecoratorDetails::OnKeyIDChanged()
 	if (SelectedEnumType)
 	{
 		CachedCustomObjectType = SelectedEnumType;
-		EnumPropValues.Reset();
+		RefreshEnumPropertyValues();
+	}
+}
+
+void FBlackboardDecoratorDetails::RefreshEnumPropertyValues()
+{
+	EnumPropValues.Reset();
 
-		if (CachedCustomObjectType)
+	if (CachedCustomObjectType)
+	{
+		for (int32 i = 0; i < CachedCustomObjectType->NumEnums() - 1; i++)
 		{
-			for (int32 i = 0; i < CachedCustomObjectType->NumEnums() - 1; i++)
-			{
-				FString DisplayedName = CachedCustomObjectType->GetDisplayNameTextByIndex(i).ToString();
-				EnumPropValues.Add(DisplayedName);
-			}
+			FString DisplayedName = CachedCustomObjectType->GetDisplayNameTextByIndex(i).ToString();
+			EnumPropValues.Add(DisplayedName);
 		}
 	}
 }
 
+void FBlackboardDecoratorDetails::PreChange(const UUserDefinedEnum* Changed, FEnumEditorUtils::EEnumEditorChangeInfo ChangedType)
+{
+	// Implementing interface pure virtual method but nothing to do here
+}
+
+void FBlackboardDecoratorDetails::PostChange(const UUserDefinedEnum* Changed, FEnumEditorUtils::EEnumEditorChangeInfo ChangedType)
+{
+	if (Changed != nullptr && CachedCustomObjectType == Changed)
+	{
+		RefreshEnumPropertyValues();
+	}
+}
+
 TSharedRef<SWidget> FBlackboardDecoratorDetails::OnGetEnumValueContent() const
 {
 	FMenuBuilder MenuBuilder(true, NULL);
@@ -199,21 +213,29 @@ TSharedRef<SWidget> FBlackboardDecoratorDetails::OnGetEnumValueContent() const
 FText FBlackboardDecoratorDetails::GetCurrentEnumValueDesc() const
 {
 	FPropertyAccess::Result Result = FPropertyAccess::Fail;
-	int32 CurrentIntValue = INDEX_NONE;
+	int32 EnumIndex = INDEX_NONE;
 
 	if (CachedCustomObjectType)
 	{	
-		Result = IntValueProperty->GetValue(CurrentIntValue);
+		// Always use string value to recompute matching index since enumeration
+		// can be modified while Editing (i.e. UserDefinedEnumeration)
+		FString CurrentStringValue;
+		Result = StringValueProperty->GetValue(CurrentStringValue);
+		EnumIndex = CachedCustomObjectType->GetIndexByNameString(CurrentStringValue);
 	}
 
-	return (Result == FPropertyAccess::Success && EnumPropValues.IsValidIndex(CurrentIntValue))
-		? FText::FromString(EnumPropValues[CurrentIntValue])
+	return (Result == FPropertyAccess::Success && EnumPropValues.IsValidIndex(EnumIndex))
+		? FText::FromString(EnumPropValues[EnumIndex])
 		: FText::GetEmpty();
 }
 
 void FBlackboardDecoratorDetails::OnEnumValueComboChange(int32 Index)
 {
-	IntValueProperty->SetValue(Index);
+	if (CachedCustomObjectType)
+	{
+		const FString NewStringValue = CachedCustomObjectType->GetNameStringByIndex(Index);
+		StringValueProperty->SetValue(NewStringValue);
+	}
 }
 
 EVisibility FBlackboardDecoratorDetails::GetIntValueVisibility() const
diff --git a/Engine/Source/Editor/BehaviorTreeEditor/Private/DetailCustomizations/BlackboardDecoratorDetails.h b/Engine/Source/Editor/BehaviorTreeEditor/Private/DetailCustomizations/BlackboardDecoratorDetails.h
--- a/Engine/Source/Editor/BehaviorTreeEditor/Private/DetailCustomizations/BlackboardDecoratorDetails.h
+++ b/Engine/Source/Editor/BehaviorTreeEditor/Private/DetailCustomizations/BlackboardDecoratorDetails.h
@@ -8,24 +8,31 @@
 #include "Layout/Visibility.h"
 #include "IDetailCustomization.h"
 #include "DetailCustomizations/BehaviorDecoratorDetails.h"
+#include "Kismet2/EnumEditorUtils.h"
 
 class IDetailLayoutBuilder;
 class IPropertyHandle;
 class SWidget;
 class UBlackboardData;
 
-class FBlackboardDecoratorDetails : public FBehaviorDecoratorDetails
+class FBlackboardDecoratorDetails : public FBehaviorDecoratorDetails, public FEnumEditorUtils::INotifyOnEnumChanged
 {
 public:
 	/** Makes a new instance of this detail layout class for a specific detail view requesting it */
 	static TSharedRef<IDetailCustomization> MakeInstance();
 
+protected:
 	/** IDetailCustomization interface */
 	virtual void CustomizeDetails( IDetailLayoutBuilder& DetailLayout ) override;
 
+	/** INotifyOnEnumChanged interface */
+	virtual void PreChange(const UUserDefinedEnum* Changed, FEnumEditorUtils::EEnumEditorChangeInfo ChangedType) override;
+	virtual void PostChange(const UUserDefinedEnum* Changed, FEnumEditorUtils::EEnumEditorChangeInfo ChangedType) override;
+
 private:
 
 	void CacheBlackboardData(IDetailLayoutBuilder& DetailLayout);
+	void RefreshEnumPropertyValues();
 
 	EVisibility GetIntValueVisibility() const;
 	EVisibility GetFloatValueVisibility() const;
@@ -41,7 +48,7 @@ private:
 	TSharedRef<SWidget> OnGetEnumValueContent() const;
 	FText GetCurrentEnumValueDesc() const;
 
-	TSharedPtr<IPropertyHandle> IntValueProperty;
+	TSharedPtr<IPropertyHandle> StringValueProperty;
 	TSharedPtr<IPropertyHandle> KeyIDProperty;
 	TSharedPtr<IPropertyHandle> NotifyObserverProperty;
 	
diff --git a/Engine/Source/Editor/BehaviorTreeEditor/Private/SGraphNode_BehaviorTree.cpp b/Engine/Source/Editor/BehaviorTreeEditor/Private/SGraphNode_BehaviorTree.cpp
--- a/Engine/Source/Editor/BehaviorTreeEditor/Private/SGraphNode_BehaviorTree.cpp
+++ b/Engine/Source/Editor/BehaviorTreeEditor/Private/SGraphNode_BehaviorTree.cpp
@@ -30,6 +30,7 @@
 #include "IDocumentation.h"
 #include "Widgets/Text/SInlineEditableTextBlock.h"
 #include "SLevelOfDetailBranchNode.h"
+#include "BehaviorTree/Decorators/BTDecorator_Blackboard.h"
 
 #define LOCTEXT_NAMESPACE "BehaviorTreeEditor"
 
@@ -607,6 +608,27 @@ EVisibility SGraphNode_BehaviorTree::GetDebuggerSearchFailedMarkerVisibility() c
 	return MyNode && MyNode->bDebuggerMarkSearchFailed ? EVisibility::HitTestInvisible : EVisibility::Collapsed;
 }
 
+void SGraphNode_BehaviorTree::PreChange(const UUserDefinedEnum* Changed, FEnumEditorUtils::EEnumEditorChangeInfo ChangedType)
+{
+	// Implementing interface pure virtual method but nothing to do here
+}
+
+void SGraphNode_BehaviorTree::PostChange(const UUserDefinedEnum* Changed, FEnumEditorUtils::EEnumEditorChangeInfo ChangedType)
+{
+	if (const UBehaviorTreeGraphNode_Decorator* DecoratorNode = Cast<UBehaviorTreeGraphNode_Decorator>(GraphNode))
+	{
+		if (UBTDecorator_Blackboard* Decorator = Cast<UBTDecorator_Blackboard>(DecoratorNode->NodeInstance))
+		{
+			Decorator->BuildDescription();
+		}
+	}
+	
+	if (UBehaviorTreeGraphNode_CompositeDecorator* DecoratorNode = Cast<UBehaviorTreeGraphNode_CompositeDecorator>(GraphNode))
+	{
+		DecoratorNode->BuildDescription();
+	}
+}
+
 void SGraphNode_BehaviorTree::Tick( const FGeometry& AllottedGeometry, const double InCurrentTime, const float InDeltaTime )
 {
 	SGraphNode::Tick( AllottedGeometry, InCurrentTime, InDeltaTime );
diff --git a/Engine/Source/Editor/BehaviorTreeEditor/Private/SGraphNode_BehaviorTree.h b/Engine/Source/Editor/BehaviorTreeEditor/Private/SGraphNode_BehaviorTree.h
--- a/Engine/Source/Editor/BehaviorTreeEditor/Private/SGraphNode_BehaviorTree.h
+++ b/Engine/Source/Editor/BehaviorTreeEditor/Private/SGraphNode_BehaviorTree.h
@@ -12,13 +12,14 @@
 #include "SGraphNode.h"
 #include "SGraphPin.h"
 #include "SGraphNodeAI.h"
+#include "Kismet2/EnumEditorUtils.h"
 
 class SHorizontalBox;
 class SToolTip;
 class SVerticalBox;
 class UBehaviorTreeGraphNode;
 
-class SGraphNode_BehaviorTree : public SGraphNodeAI
+class SGraphNode_BehaviorTree : public SGraphNodeAI, public FEnumEditorUtils::INotifyOnEnumChanged
 {
 public:
 	SLATE_BEGIN_ARGS(SGraphNode_BehaviorTree){}
@@ -61,6 +62,10 @@ public:
 	FVector2D GetCachedPosition() const { return CachedPosition; }
 
 protected:
+	/** INotifyOnEnumChanged interface */
+	virtual void PreChange(const UUserDefinedEnum* Changed, FEnumEditorUtils::EEnumEditorChangeInfo ChangedType) override;
+	virtual void PostChange(const UUserDefinedEnum* Changed, FEnumEditorUtils::EEnumEditorChangeInfo ChangedType) override;
+
 	uint32 bSuppressDebuggerColor : 1;
 	uint32 bSuppressDebuggerTriggers : 1;
 	
diff --git a/Engine/Source/Runtime/AIModule/Classes/BehaviorTree/Blackboard/BlackboardKeyType_Enum.h b/Engine/Source/Runtime/AIModule/Classes/BehaviorTree/Blackboard/BlackboardKeyType_Enum.h
--- a/Engine/Source/Runtime/AIModule/Classes/BehaviorTree/Blackboard/BlackboardKeyType_Enum.h
+++ b/Engine/Source/Runtime/AIModule/Classes/BehaviorTree/Blackboard/BlackboardKeyType_Enum.h
@@ -27,8 +27,8 @@ class AIMODULE_API UBlackboardKeyType_Enum : public UBlackboardKeyType
 	UPROPERTY(Category = Blackboard, VisibleDefaultsOnly)
 	uint32 bIsEnumNameValid : 1;
 
-	static uint8 GetValue(const UBlackboardKeyType_Enum* KeyOb, const uint8* RawData);
-	static bool SetValue(UBlackboardKeyType_Enum* KeyOb, uint8* RawData, uint8 Value);
+	static FDataType GetValue(const UBlackboardKeyType_Enum* KeyOb, const uint8* RawData);
+	static bool SetValue(UBlackboardKeyType_Enum* KeyOb, uint8* RawData, FDataType Value);
 
 	virtual EBlackboardCompare::Type CompareValues(const UBlackboardComponent& OwnerComp, const uint8* MemoryBlock,
 		const UBlackboardKeyType* OtherKeyOb, const uint8* OtherMemoryBlock) const override;
@@ -37,11 +37,12 @@ class AIMODULE_API UBlackboardKeyType_Enum : public UBlackboardKeyType
 	virtual FString DescribeArithmeticParam(int32 IntValue, float FloatValue) const override;
 	virtual bool IsAllowedByFilter(UBlackboardKeyType* FilterOb) const override;
 
+protected:
 #if WITH_EDITOR
 	virtual void PostEditChangeProperty(struct FPropertyChangedEvent& PropertyChangedEvent) override;
+	static bool ValidateEnum(const UEnum& EnumType);
 #endif
 
-protected:
 	virtual FString DescribeValue(const UBlackboardComponent& OwnerComp, const uint8* RawData) const override;
 	virtual bool TestArithmeticOperation(const UBlackboardComponent& OwnerComp, const uint8* MemoryBlock, EArithmeticKeyOperation::Type Op, int32 OtherIntValue, float OtherFloatValue) const override;
 };
diff --git a/Engine/Source/Runtime/AIModule/Classes/BehaviorTree/Decorators/BTDecorator_Blackboard.h b/Engine/Source/Runtime/AIModule/Classes/BehaviorTree/Decorators/BTDecorator_Blackboard.h
--- a/Engine/Source/Runtime/AIModule/Classes/BehaviorTree/Decorators/BTDecorator_Blackboard.h
+++ b/Engine/Source/Runtime/AIModule/Classes/BehaviorTree/Decorators/BTDecorator_Blackboard.h
@@ -12,6 +12,7 @@
 class FBlackboardDecoratorDetails;
 class UBehaviorTree;
 class UBlackboardComponent;
+struct FBlackboardEntry;
 
 /**
  *  Decorator for accessing blackboard values
@@ -81,14 +82,21 @@ protected:
 #endif
 
 #if WITH_EDITOR
-
+public:
 	/** describe decorator and cache it */
 	virtual void BuildDescription();
-
+protected:
 	virtual void PostEditChangeProperty(FPropertyChangedEvent& PropertyChangedEvent) override;
+#endif
+
 	virtual void InitializeFromAsset(UBehaviorTree& Asset) override;
 
-#endif
+	/**
+	 * decorator using enum based key requires to synchronize 'StringValue' used for storage with 'IntValue' used for runtime evaluation.
+	 * @note using custom serialization version was not possible since we require blackboard and used types (e.g. UserDefinedEnumerations)
+	 * to be loaded and keys to be resolved (i.e. InitializeFromAsset).
+	 */
+	void RefreshEnumBasedDecorator(const FBlackboardEntry& Entry);
 
 	/** take blackboard value and evaluate decorator's condition */
 	bool EvaluateOnBlackboard(const UBlackboardComponent& BlackboardComp) const;
diff --git a/Engine/Source/Runtime/AIModule/Private/BehaviorTree/Blackboard/BlackboardKeyType_Enum.cpp b/Engine/Source/Runtime/AIModule/Private/BehaviorTree/Blackboard/BlackboardKeyType_Enum.cpp
--- a/Engine/Source/Runtime/AIModule/Private/BehaviorTree/Blackboard/BlackboardKeyType_Enum.cpp
+++ b/Engine/Source/Runtime/AIModule/Private/BehaviorTree/Blackboard/BlackboardKeyType_Enum.cpp
@@ -1,32 +1,37 @@
 // Copyright Epic Games, Inc. All Rights Reserved.
 
 #include "BehaviorTree/Blackboard/BlackboardKeyType_Enum.h"
-#include "UObject/Package.h"
+#include "BehaviorTree/BTNode.h"
 #include "UObject/UnrealType.h"
 
+#if WITH_EDITOR
+#include "Misc/MessageDialog.h"
+#include <limits>
+#endif // WITH_EDITOR
+
 const UBlackboardKeyType_Enum::FDataType UBlackboardKeyType_Enum::InvalidValue = UBlackboardKeyType_Enum::FDataType(0);
 
 UBlackboardKeyType_Enum::UBlackboardKeyType_Enum(const FObjectInitializer& ObjectInitializer) : Super(ObjectInitializer)
 {
-	ValueSize = sizeof(uint8);
+	ValueSize = sizeof(FDataType);
 	SupportedOp = EBlackboardKeyOperation::Arithmetic;
 }
 
-uint8 UBlackboardKeyType_Enum::GetValue(const UBlackboardKeyType_Enum* KeyOb, const uint8* RawData)
+UBlackboardKeyType_Enum::FDataType UBlackboardKeyType_Enum::GetValue(const UBlackboardKeyType_Enum* KeyOb, const uint8* RawData)
 {
-	return GetValueFromMemory<uint8>(RawData);
+	return GetValueFromMemory<FDataType>(RawData);
 }
 
-bool UBlackboardKeyType_Enum::SetValue(UBlackboardKeyType_Enum* KeyOb, uint8* RawData, uint8 Value)
+bool UBlackboardKeyType_Enum::SetValue(UBlackboardKeyType_Enum* KeyOb, uint8* RawData, const FDataType Value)
 {
-	return SetValueInMemory<uint8>(RawData, Value);
+	return SetValueInMemory<FDataType>(RawData, Value);
 }
 
 EBlackboardCompare::Type UBlackboardKeyType_Enum::CompareValues(const UBlackboardComponent& OwnerComp, const uint8* MemoryBlock,
 	const UBlackboardKeyType* OtherKeyOb, const uint8* OtherMemoryBlock) const
 {
-	const uint8 MyValue = GetValue(this, MemoryBlock);
-	const uint8 OtherValue = GetValue((UBlackboardKeyType_Enum*)OtherKeyOb, OtherMemoryBlock);
+	const FDataType MyValue = GetValue(this, MemoryBlock);
+	const FDataType OtherValue = GetValue((UBlackboardKeyType_Enum*)OtherKeyOb, OtherMemoryBlock);
 
 	return (MyValue > OtherValue) ? EBlackboardCompare::Greater :
 		(MyValue < OtherValue) ? EBlackboardCompare::Less :
@@ -51,7 +56,7 @@ bool UBlackboardKeyType_Enum::IsAllowedByFilter(UBlackboardKeyType* FilterOb) co
 
 bool UBlackboardKeyType_Enum::TestArithmeticOperation(const UBlackboardComponent& OwnerComp, const uint8* MemoryBlock, EArithmeticKeyOperation::Type Op, int32 OtherIntValue, float OtherFloatValue) const
 {
-	const uint8 Value = GetValue(this, MemoryBlock);
+	const FDataType Value = GetValue(this, MemoryBlock);
 	switch (Op)
 	{
 	case EArithmeticKeyOperation::Equal:			return (Value == OtherIntValue);
@@ -80,8 +85,42 @@ void UBlackboardKeyType_Enum::PostEditChangeProperty(struct FPropertyChangedEven
 		PropertyChangedEvent.Property->GetFName() == GET_MEMBER_NAME_CHECKED(UBlackboardKeyType_Enum, EnumName))
 	{
 		EnumType = UClass::TryFindTypeSlow<UEnum>(EnumName, EFindFirstObjectOptions::ExactClass);
+
+		if (EnumType != nullptr && !ValidateEnum(*EnumType))
+		{
+			EnumType = nullptr;
+		}
 	}
 
 	bIsEnumNameValid = EnumType && !EnumName.IsEmpty();
 }
-#endif
+
+bool UBlackboardKeyType_Enum::ValidateEnum(const UEnum& EnumType)
+{
+	bool bAllValid = true;
+
+	// Do not test the max value (if present) since it is an internal value and users don't have access to it
+	const int32 NumEnums = EnumType.ContainsExistingMax() ? EnumType.NumEnums() - 1 : EnumType.NumEnums();
+	for (int32 i = 0; i < NumEnums; i++)
+	{
+		const int64 Value = EnumType.GetValueByIndex(i);
+		if (Value < std::numeric_limits<FDataType>::min() || Value > std::numeric_limits<FDataType>::max())
+		{
+			UE_LOG(LogBehaviorTree, Error, TEXT("'%s' value %d is outside the range of supported key values for enum [%d, %d].")
+				, *EnumType.GenerateFullEnumName(*EnumType.GetDisplayNameTextByIndex(i).ToString())
+				, Value, std::numeric_limits<FDataType>::min(), std::numeric_limits<FDataType>::max());
+			bAllValid = false;
+		}
+	}
+
+	if (!bAllValid)
+	{
+		FMessageDialog::Open(EAppMsgType::Ok,
+			NSLOCTEXT("BehaviorTree"
+				, "Unsupported enumeration"
+				, "Specified enumeration contains one or more values outside supported value range for enum keys and can not be used in the Blackboard. See log for details."));
+	}
+
+	return bAllValid;
+}
+#endif // WITH_EDITOR
diff --git a/Engine/Source/Runtime/AIModule/Private/BehaviorTree/Decorators/BTDecorator_Blackboard.cpp b/Engine/Source/Runtime/AIModule/Private/BehaviorTree/Decorators/BTDecorator_Blackboard.cpp
--- a/Engine/Source/Runtime/AIModule/Private/BehaviorTree/Decorators/BTDecorator_Blackboard.cpp
+++ b/Engine/Source/Runtime/AIModule/Private/BehaviorTree/Decorators/BTDecorator_Blackboard.cpp
@@ -1,7 +1,6 @@
 // Copyright Epic Games, Inc. All Rights Reserved.
 
 #include "BehaviorTree/Decorators/BTDecorator_Blackboard.h"
-#include "UObject/Package.h"
 #include "BehaviorTree/Blackboard/BlackboardKeyType_Enum.h"
 #include "BehaviorTree/Blackboard/BlackboardKeyType_NativeEnum.h"
 #include "BehaviorTree/BTCompositeNode.h"
@@ -124,7 +123,7 @@ static void CacheOperationEnums()
 
 void UBTDecorator_Blackboard::BuildDescription()
 {
-	UBlackboardData* BlackboardAsset = GetBlackboardAsset();
+	const UBlackboardData* BlackboardAsset = GetBlackboardAsset();
 	const FBlackboardEntry* EntryInfo = BlackboardAsset ? BlackboardAsset->GetKey(BlackboardKey.GetSelectedKeyID()) : NULL;
 
 	FString BlackboardDesc = "invalid";
@@ -134,6 +133,8 @@ void UBTDecorator_Blackboard::BuildDescription()
 		// while "post edit property" triggers for every each of them
 		if (EntryInfo->KeyType && EntryInfo->KeyType->GetClass() == BlackboardKey.SelectedKeyType)
 		{
+			RefreshEnumBasedDecorator(*EntryInfo);
+
 			const FString KeyName = EntryInfo->EntryName.ToString();
 			CacheOperationEnums();		
 
@@ -170,18 +171,48 @@ void UBTDecorator_Blackboard::PostEditChangeProperty(FPropertyChangedEvent& Prop
 	}
 
 	const FName ChangedPropName = PropertyChangedEvent.Property->GetFName();
-	if (ChangedPropName == GET_MEMBER_NAME_CHECKED(UBTDecorator_Blackboard, BlackboardKey.SelectedKeyName))
+
+	const FName SelectedKeyName = GET_MEMBER_NAME_CHECKED(FBlackboardKeySelector, SelectedKeyName);
+	const FName IntValueName = GET_MEMBER_NAME_CHECKED(UBTDecorator_Blackboard, IntValue);
+	const FName StringValueName = GET_MEMBER_NAME_CHECKED(UBTDecorator_Blackboard, StringValue);
+
+	const bool bIsEnumKey = BlackboardKey.SelectedKeyType == UBlackboardKeyType_Enum::StaticClass()
+		|| BlackboardKey.SelectedKeyType == UBlackboardKeyType_NativeEnum::StaticClass();
+	
+	if (bIsEnumKey && (ChangedPropName == SelectedKeyName || ChangedPropName == IntValueName || ChangedPropName == StringValueName))
 	{
-		if (BlackboardKey.SelectedKeyType == UBlackboardKeyType_Enum::StaticClass() ||
-			BlackboardKey.SelectedKeyType == UBlackboardKeyType_NativeEnum::StaticClass())
+		const UBlackboardData* BlackboardAsset = GetBlackboardAsset();
+		const FBlackboardEntry* EntryInfo = BlackboardAsset ? BlackboardAsset->GetKey(BlackboardKey.GetSelectedKeyID()) : nullptr;
+
+		const UEnum* EnumType = (BlackboardKey.SelectedKeyType == UBlackboardKeyType_Enum::StaticClass())
+									? ((UBlackboardKeyType_Enum*)(EntryInfo->KeyType))->EnumType
+									: ((UBlackboardKeyType_NativeEnum*)(EntryInfo->KeyType))->EnumType;
+
+		if (ChangedPropName == SelectedKeyName)
+		{
+			// Set both properties to match the first valid value or invalidate them
+			if (EnumType != nullptr && EnumType->NumEnums())
+			{
+				IntValue = EnumType->GetValueByIndex(0);
+				StringValue = EnumType->GetNameStringByIndex(0);
+			}
+			else
+			{
+				IntValue = INDEX_NONE;
+				StringValue.Reset();
+			}
+		}
+		else if (ChangedPropName == IntValueName)
 		{
-			IntValue = 0;
+			StringValue = (EnumType != nullptr) ? EnumType->GetNameStringByValue(IntValue) : TEXT("");
+		}
+		else if (ChangedPropName == StringValueName)
+		{
+			IntValue = (EnumType != nullptr) ? EnumType->GetValueByNameString(StringValue) : INDEX_NONE;
 		}
 	}
 
-#if WITH_EDITORONLY_DATA
-
-	UBlackboardKeyType* KeyCDO = BlackboardKey.SelectedKeyType ? BlackboardKey.SelectedKeyType->GetDefaultObject<UBlackboardKeyType>() : NULL;
+	const UBlackboardKeyType* KeyCDO = BlackboardKey.SelectedKeyType ? BlackboardKey.SelectedKeyType->GetDefaultObject<UBlackboardKeyType>() : NULL;
 	if (ChangedPropName == GET_MEMBER_NAME_CHECKED(UBTDecorator_Blackboard, BasicOperation))
 	{
 		if (KeyCDO && KeyCDO->GetTestOperation() == EBlackboardKeyOperation::Basic)
@@ -204,15 +235,56 @@ void UBTDecorator_Blackboard::PostEditChangeProperty(FPropertyChangedEvent& Prop
 		}
 	}
 
-#endif // WITH_EDITORONLY_DATA
-
 	BuildDescription();
 }
 
+#endif	// WITH_EDITOR
+
 void UBTDecorator_Blackboard::InitializeFromAsset(UBehaviorTree& Asset)
 {
 	Super::InitializeFromAsset(Asset);
+
+#if WITH_EDITOR
+	// BuildDescription will take care of refreshing enum based decorator
 	BuildDescription();
+#else
+	const UBlackboardData* BlackboardAsset = GetBlackboardAsset();
+	if (const FBlackboardEntry* EntryInfo = BlackboardAsset ? BlackboardAsset->GetKey(BlackboardKey.GetSelectedKeyID()) : nullptr)
+	{
+		RefreshEnumBasedDecorator(*EntryInfo);
+	}
+#endif // WITH_EDITOR
 }
 
-#endif	// WITH_EDITOR
+void UBTDecorator_Blackboard::RefreshEnumBasedDecorator(const FBlackboardEntry& Entry)
+{
+	if (BlackboardKey.SelectedKeyType != UBlackboardKeyType_Enum::StaticClass() &&
+		BlackboardKey.SelectedKeyType != UBlackboardKeyType_NativeEnum::StaticClass())
+	{
+		return;
+	}
+
+	const UEnum* Enum = (BlackboardKey.SelectedKeyType == UBlackboardKeyType_Enum::StaticClass())
+							? ((UBlackboardKeyType_Enum*)(Entry.KeyType))->EnumType
+							: ((UBlackboardKeyType_NativeEnum*)(Entry.KeyType))->EnumType;
+	if (Enum == nullptr)
+	{
+		return;
+	}
+
+	// Enum implementation relies on 'StringValue' to store enum name but also
+	// synchronizes 'IntValue' as the enumeration value for runtime evaluation
+	if (!StringValue.IsEmpty())
+	{
+		// Pass string as FName to use fast search path whenever possible 
+		const FName Name = *StringValue;
+		IntValue = Enum->GetValueByName(Name);
+	}
+	else if (IntValue >= 0 && IntValue < Enum->NumEnums())
+	{
+		// Previous implementation was using only 'IntValue' to store the enumeration
+		// index so we need to extract associated name and update `IntValue` to store the value.
+		StringValue = Enum->GetNameByIndex(IntValue).ToString();
+		IntValue = Enum->GetValueByIndex(IntValue);
+	}
+}
