Static Analysis fix
#rnx

diff --git a/Engine/Source/Runtime/AIModule/Private/BehaviorTree/Decorators/BTDecorator_Blackboard.cpp b/Engine/Source/Runtime/AIModule/Private/BehaviorTree/Decorators/BTDecorator_Blackboard.cpp
--- a/Engine/Source/Runtime/AIModule/Private/BehaviorTree/Decorators/BTDecorator_Blackboard.cpp
+++ b/Engine/Source/Runtime/AIModule/Private/BehaviorTree/Decorators/BTDecorator_Blackboard.cpp
@@ -183,32 +183,34 @@ void UBTDecorator_Blackboard::PostEditChangeProperty(FPropertyChangedEvent& Prop
 	{
 		const UBlackboardData* BlackboardAsset = GetBlackboardAsset();
 		const FBlackboardEntry* EntryInfo = BlackboardAsset ? BlackboardAsset->GetKey(BlackboardKey.GetSelectedKeyID()) : nullptr;
-
-		const UEnum* EnumType = (BlackboardKey.SelectedKeyType == UBlackboardKeyType_Enum::StaticClass())
-									? ((UBlackboardKeyType_Enum*)(EntryInfo->KeyType))->EnumType
-									: ((UBlackboardKeyType_NativeEnum*)(EntryInfo->KeyType))->EnumType;
-
-		if (ChangedPropName == SelectedKeyName)
+		if (EntryInfo != nullptr)
 		{
-			// Set both properties to match the first valid value or invalidate them
-			if (EnumType != nullptr && EnumType->NumEnums())
+			const UEnum* EnumType = (BlackboardKey.SelectedKeyType == UBlackboardKeyType_Enum::StaticClass())
+									   ? ((UBlackboardKeyType_Enum*)(EntryInfo->KeyType))->EnumType
+									   : ((UBlackboardKeyType_NativeEnum*)(EntryInfo->KeyType))->EnumType;
+
+			if (ChangedPropName == SelectedKeyName)
 			{
-				IntValue = EnumType->GetValueByIndex(0);
-				StringValue = EnumType->GetNameStringByIndex(0);
+				// Set both properties to match the first valid value or invalidate them
+				if (EnumType != nullptr && EnumType->NumEnums())
+				{
+					IntValue = EnumType->GetValueByIndex(0);
+					StringValue = EnumType->GetNameStringByIndex(0);
+				}
+				else
+				{
+					IntValue = INDEX_NONE;
+					StringValue.Reset();
+				}
 			}
-			else
+			else if (ChangedPropName == IntValueName)
 			{
-				IntValue = INDEX_NONE;
-				StringValue.Reset();
+				StringValue = (EnumType != nullptr) ? EnumType->GetNameStringByValue(IntValue) : TEXT("");
+			}
+			else if (ChangedPropName == StringValueName)
+			{
+				IntValue = (EnumType != nullptr) ? EnumType->GetValueByNameString(StringValue) : INDEX_NONE;
 			}
-		}
-		else if (ChangedPropName == IntValueName)
-		{
-			StringValue = (EnumType != nullptr) ? EnumType->GetNameStringByValue(IntValue) : TEXT("");
-		}
-		else if (ChangedPropName == StringValueName)
-		{
-			IntValue = (EnumType != nullptr) ? EnumType->GetValueByNameString(StringValue) : INDEX_NONE;
 		}
 	}
 
