[GameplayTask] added MarkOwnerFinished to inform the task that owner is done and don't want to get notified through OnGameplayTaskDeactivated callback.
This is for cases where an owner doesn't control the lifetime of the task (i.e. will not call ExternalCancel or EndTask)

diff --git a/Engine/Source/Runtime/AIModule/Private/BehaviorTree/Tasks/BTTask_GameplayTaskBase.cpp b/Engine/Source/Runtime/AIModule/Private/BehaviorTree/Tasks/BTTask_GameplayTaskBase.cpp
--- a/Engine/Source/Runtime/AIModule/Private/BehaviorTree/Tasks/BTTask_GameplayTaskBase.cpp
+++ b/Engine/Source/Runtime/AIModule/Private/BehaviorTree/Tasks/BTTask_GameplayTaskBase.cpp
@@ -33,16 +33,21 @@ EBTNodeResult::Type UBTTask_GameplayTaskBase::StartGameplayTask(UBehaviorTreeCom
 
 EBTNodeResult::Type UBTTask_GameplayTaskBase::AbortTask(UBehaviorTreeComponent& OwnerComp, uint8* NodeMemory)
 {
-	if (bWaitForGameplayTask)
-	{
-		FBTGameplayTaskMemory* MyMemory = CastInstanceNodeMemory<FBTGameplayTaskMemory>(NodeMemory);
-		MyMemory->bObserverCanFinishTask = false;
+	FBTGameplayTaskMemory* MyMemory = CastInstanceNodeMemory<FBTGameplayTaskMemory>(NodeMemory);
+	MyMemory->bObserverCanFinishTask = false;
 
-		UAITask* TaskOb = MyMemory->Task.Get();
-		if (TaskOb && !TaskOb->IsFinished())
+	UAITask* TaskOb = MyMemory->Task.Get();
+	if (TaskOb && !TaskOb->IsFinished())
+	{
+		if (bWaitForGameplayTask)
 		{
 			TaskOb->ExternalCancel();
 		}
+		else
+		{
+			// Inform the task that we're done and no longer want to receive OnGameplayTaskDeactivated callback
+			TaskOb->MarkOwnerFinished();
+		}
 	}
 
 	return EBTNodeResult::Aborted;
@@ -71,11 +76,14 @@ void UBTTask_GameplayTaskBase::OnGameplayTaskDeactivated(UGameplayTask& Task)
 		{
 			uint8* RawMemory = BehaviorComp->GetNodeMemory(this, BehaviorComp->FindInstanceContainingNode(this));
 			const FBTGameplayTaskMemory* MyMemory = CastInstanceNodeMemory<FBTGameplayTaskMemory>(RawMemory);
-
-			if (MyMemory->bObserverCanFinishTask && (AITask == MyMemory->Task))
+			if (ensureMsgf(MyMemory, TEXT("OnGameplayTaskDeactivated called after behavior tree instance was shutdown."
+				" Make sure that AITask was cancelled ('ExternalCancel'), ended ('TaskOwnerEnded') or informed that owner is done ('MarkOwnerFinished')")))
 			{
-				const EBTNodeResult::Type FinishResult = DetermineGameplayTaskResult(*AITask);
-				FinishLatentTask(*BehaviorComp, FinishResult);
+				if (MyMemory->bObserverCanFinishTask && (AITask == MyMemory->Task))
+				{
+					const EBTNodeResult::Type FinishResult = DetermineGameplayTaskResult(*AITask);
+					FinishLatentTask(*BehaviorComp, FinishResult);
+				}
 			}
 		}
 	}
diff --git a/Engine/Source/Runtime/GameplayTasks/Classes/GameplayTask.h b/Engine/Source/Runtime/GameplayTasks/Classes/GameplayTask.h
--- a/Engine/Source/Runtime/GameplayTasks/Classes/GameplayTask.h
+++ b/Engine/Source/Runtime/GameplayTasks/Classes/GameplayTask.h
@@ -161,7 +161,7 @@ protected:
 	 *	Note that the default implementation does nothing and you don't have to call it */
 	virtual void Activate();
 
-	/** Initailizes the task with the task owner interface instance but does not actviate until Activate() is called */
+	/** Initializes the task with the task owner interface instance but does not activate until Activate() is called */
 	void InitTask(IGameplayTaskOwnerInterface& InTaskOwner, uint8 InPriority);
 
 public:
@@ -176,6 +176,9 @@ public:
 	/** Called when the task is asked to cancel from an outside node. What this means depends on the individual task. By default, this does nothing other than ending the task. */
 	virtual void ExternalCancel();
 
+	/** Called when a task owner wants to inform the task that it's done and don't want to get notified on task deactivation (i.e. OnGameplayTaskDeactivated). */
+	void MarkOwnerFinished() { bOwnerFinished = true; }
+
 	/** Return debug string describing task */
 	virtual FString GetDebugString() const;
 
@@ -321,7 +324,7 @@ protected:
 	UPROPERTY()
 	FName InstanceName;
 
-	/** This controls how this task will be treaded in relation to other, already running tasks, 
+	/** This controls how this task will be treated in relation to other, already running tasks, 
 	 *	provided GameplayTasksComponent is configured to care about priorities (the default behavior)*/
 	uint8 Priority;
 
