Renamed awkwardly named UPawnAction_Move function

diff --git a/Engine/Source/Runtime/AIModule/Classes/Actions/PawnAction_Move.h b/Engine/Source/Runtime/AIModule/Classes/Actions/PawnAction_Move.h
--- a/Engine/Source/Runtime/AIModule/Classes/Actions/PawnAction_Move.h
+++ b/Engine/Source/Runtime/AIModule/Classes/Actions/PawnAction_Move.h
@@ -66,7 +66,7 @@ protected:
 
 	/** if set, other actions with the same priority will be aborted when path is changed */
 	UPROPERTY()
-	uint32 bAbortChildActionOnPathChange : 1;
+	uint32 bAbortSubActionOnPathChange : 1;
 
 public:
 	virtual void BeginDestroy() override;
@@ -87,10 +87,13 @@ public:
 	void EnableStrafing(bool bNewStrafing) { bAllowStrafe = bNewStrafing; }
 	void EnablePathUpdateOnMoveGoalLocationChange(bool bEnable) { bUpdatePathToGoal = bEnable; }
 	void EnableGoalLocationProjectionToNavigation(bool bEnable) { bProjectGoalToNavigation = bEnable; }
-	void EnableChildAbortionOnPathUpdate(bool bEnable) { bAbortChildActionOnPathChange = bEnable; }
+	void SetAbortSubActionOnPathUpdate(bool bEnable) { bAbortSubActionOnPathChange = bEnable; }
 	void SetFilterClass(TSubclassOf<UNavigationQueryFilter> NewFilterClass) { FilterClass = NewFilterClass; }
 	void SetAllowPartialPath(bool bEnable) { bAllowPartialPath = bEnable; }
 
+	UE_DEPRECATED(5.1, "Use SetAbortSubActionOnPathUpdate instead.")
+	void EnableChildAbortionOnPathUpdate(bool bEnable) { SetAbortSubActionOnPathUpdate(bEnable); }
+
 protected:
 	/** currently followed path */
 	FNavPathSharedPtr Path;
diff --git a/Engine/Source/Runtime/AIModule/Private/Actions/PawnAction_Move.cpp b/Engine/Source/Runtime/AIModule/Private/Actions/PawnAction_Move.cpp
--- a/Engine/Source/Runtime/AIModule/Private/Actions/PawnAction_Move.cpp
+++ b/Engine/Source/Runtime/AIModule/Private/Actions/PawnAction_Move.cpp
@@ -16,7 +16,7 @@ UPawnAction_Move::UPawnAction_Move(const FObjectInitializer& ObjectInitializer)
 	, bAllowPartialPath(true)
 	, bProjectGoalToNavigation(false)
 	, bUpdatePathToGoal(true)
-	, bAbortChildActionOnPathChange(false)
+	, bAbortSubActionOnPathChange(false)
 {
 	bShouldPauseMovement = true;
 
@@ -281,7 +281,7 @@ void UPawnAction_Move::OnPathUpdated(FNavigationPath* UpdatedPath, ENavPathEvent
 
 	UE_VLOG(MyOwner, LogPawnAction, Log, TEXT("%s> Path updated!"), *GetName());
 	
-	if (bAbortChildActionOnPathChange && GetChildAction())
+	if (bAbortSubActionOnPathChange && GetChildAction())
 	{
 		UE_VLOG(MyOwner, LogPawnAction, Log, TEXT(">> aborting child action: %s"), *GetNameSafe(GetChildAction()));
 		
