Fixing crash that can occur in AAIController:MoveTo if the goal actor is marked for garbage.

#tests Tested in PIE to ensure that AI were still able to path ok with and without goal actors

diff --git a/Engine/Source/Runtime/AIModule/Private/AIController.cpp b/Engine/Source/Runtime/AIModule/Private/AIController.cpp
--- a/Engine/Source/Runtime/AIModule/Private/AIController.cpp
+++ b/Engine/Source/Runtime/AIModule/Private/AIController.cpp
@@ -622,7 +622,7 @@ FPathFollowingRequestResult AAIController::MoveTo(const FAIMoveRequest& MoveRequ
 
 	ensure(MoveRequest.GetNavigationFilter() || !DefaultNavigationFilterClass);
 
-	bool bCanRequestMove = true;
+	bool bCanRequestMove = !MoveRequest.IsMoveToActorRequest() || IsValid(MoveRequest.GetGoalActor());
 	bool bAlreadyAtGoal = false;
 	
 	if (!MoveRequest.IsMoveToActorRequest())
