[Navmesh] Allow FGameplayDebuggerCategory_Navmesh to display navdata name even if not a ARecastNavMesh

diff --git a/Engine/Source/Runtime/AIModule/Private/GameplayDebugger/GameplayDebuggerCategory_Navmesh.cpp b/Engine/Source/Runtime/AIModule/Private/GameplayDebugger/GameplayDebuggerCategory_Navmesh.cpp
--- a/Engine/Source/Runtime/AIModule/Private/GameplayDebugger/GameplayDebuggerCategory_Navmesh.cpp
+++ b/Engine/Source/Runtime/AIModule/Private/GameplayDebugger/GameplayDebuggerCategory_Navmesh.cpp
@@ -198,8 +198,7 @@ void FGameplayDebuggerCategory_Navmesh::CollectData(APlayerController* OwnerPC,
 		}
 	}
 
-	const ARecastNavMesh* RecastNavMesh = Cast<const ARecastNavMesh>(NavData);
-	if (RecastNavMesh && RefPawn)
+	if (NavData)
 	{
 		DataPack.bIsUsingPlayerActor = (ActorReferenceMode != EActorReferenceMode::DebugActor);
 		DataPack.bCanChangeReference = (ActorReferenceMode != EActorReferenceMode::PlayerActorOnly);
@@ -207,13 +206,17 @@ void FGameplayDebuggerCategory_Navmesh::CollectData(APlayerController* OwnerPC,
 
 		if (NumNavData > 1)
 		{
-			DataPack.NavDataName = FString::Printf(TEXT("[%d/%d] %s"), NavDataIndexToDisplay + 1, NumNavData, *RecastNavMesh->GetFName().ToString());
+			DataPack.NavDataName = FString::Printf(TEXT("[%d/%d] %s"), NavDataIndexToDisplay + 1, NumNavData, *NavData->GetFName().ToString());
 		}
 		else
 		{
-			DataPack.NavDataName = RecastNavMesh->GetFName().ToString();
+			DataPack.NavDataName = NavData->GetFName().ToString();
 		}
+	}
 
+	const ARecastNavMesh* RecastNavMesh = Cast<const ARecastNavMesh>(NavData);
+	if (RecastNavMesh && RefPawn)
+	{
 		// add NxN neighborhood of target (where N is the number of tiles)
 		// Note that we round up to the next odd number to keep the reference position in the middle tile
 		const FVector TargetLocation = RefPawn->GetActorLocation();
