Changing the GoalActor for AI move requests to weakpointer. This resolves some crashes seen when the goal actor of AITask_MoveTo is garbage
Also backing out earlier change to AIController.cpp since its not necessary with this fix

#tests Loaded up a zone in PIE to ensure there was no issues with AI pathing to targets

diff --git a/Engine/Source/Runtime/AIModule/Classes/AITypes.h b/Engine/Source/Runtime/AIModule/Classes/AITypes.h
--- a/Engine/Source/Runtime/AIModule/Classes/AITypes.h
+++ b/Engine/Source/Runtime/AIModule/Classes/AITypes.h
@@ -506,14 +506,14 @@ struct AIMODULE_API FAIMoveRequest
 	FAIMoveRequest& SetUserFlags(int32 InUserFlags) { UserFlags = InUserFlags; return *this; }
 
 	/** the request should be either set up to move to a location, of go to a valid actor */
-	bool IsValid() const { return bInitialized && (!bMoveToActor || GoalActor); }
+	bool IsValid() const { return bInitialized && (!bMoveToActor || GoalActor.IsValid()); }
 
 	bool IsMoveToActorRequest() const { return bMoveToActor; }
-	AActor* GetGoalActor() const { return bMoveToActor ? GoalActor : nullptr; }
+	AActor* GetGoalActor() const { return bMoveToActor ? GoalActor.Get() : nullptr; }
 	FVector GetGoalLocation() const { return GoalLocation; }
 	/** retrieves request's requested destination location, GoalActor's location 
 	 *	or GoalLocation, depending on the request itself */
-	FVector GetDestination() const { return bMoveToActor ? (GoalActor ? GoalActor->GetActorLocation() : FAISystem::InvalidLocation) : GoalLocation; }
+	FVector GetDestination() const { return bMoveToActor ? (GoalActor.IsValid() ? GoalActor->GetActorLocation() : FAISystem::InvalidLocation) : GoalLocation; }
 
 	bool IsUsingPathfinding() const { return bUsePathfinding; }
 	bool IsUsingPartialPaths() const { return bAllowPartialPath; }
@@ -537,7 +537,7 @@ protected:
 
 	/** move goal: actor */
 	UPROPERTY()
-	TObjectPtr<AActor> GoalActor;
+	TWeakObjectPtr<AActor> GoalActor;
 
 	/** move goal: location */
 	mutable FVector GoalLocation;
diff --git a/Engine/Source/Runtime/AIModule/Private/AIController.cpp b/Engine/Source/Runtime/AIModule/Private/AIController.cpp
--- a/Engine/Source/Runtime/AIModule/Private/AIController.cpp
+++ b/Engine/Source/Runtime/AIModule/Private/AIController.cpp
@@ -622,7 +622,7 @@ FPathFollowingRequestResult AAIController::MoveTo(const FAIMoveRequest& MoveRequ
 
 	ensure(MoveRequest.GetNavigationFilter() || !DefaultNavigationFilterClass);
 
-	bool bCanRequestMove = !MoveRequest.IsMoveToActorRequest() || IsValid(MoveRequest.GetGoalActor());
+	bool bCanRequestMove = true;
 	bool bAlreadyAtGoal = false;
 	
 	if (!MoveRequest.IsMoveToActorRequest())
diff --git a/Engine/Source/Runtime/AIModule/Private/AITypes.cpp b/Engine/Source/Runtime/AIModule/Private/AITypes.cpp
--- a/Engine/Source/Runtime/AIModule/Private/AITypes.cpp
+++ b/Engine/Source/Runtime/AIModule/Private/AITypes.cpp
@@ -220,7 +220,7 @@ bool FAIMoveRequest::UpdateGoalLocation(const FVector& NewLocation) const
 FString FAIMoveRequest::ToString() const
 {
 	return FString::Printf(TEXT("%s(%s) Mode(%s) Filter(%s) AcceptanceRadius(%.1f%s)"),
-		bMoveToActor ? TEXT("Actor") : TEXT("Location"), bMoveToActor ? *GetNameSafe(GoalActor) : *GoalLocation.ToString(),
+		bMoveToActor ? TEXT("Actor") : TEXT("Location"), bMoveToActor ? *GetNameSafe(GoalActor.Get()) : *GoalLocation.ToString(),
 		bUsePathfinding ? (bAllowPartialPath ? TEXT("partial path") : TEXT("complete path")) : TEXT("direct"),
 		*GetNameSafe(FilterClass),
 		AcceptanceRadius, bReachTestIncludesAgentRadius || (bMoveToActor && bReachTestIncludesGoalRadius) ? TEXT(" + overlap") : TEXT("")
