Added functions to FEnvQueryResult for convenience (GetItemRawMemory) and flexibility (GetItemAsTypeChecked).

GetItemAsTypeChecked lets programmers to convert query results to the type used by query's generator to extract  type-specific data from items.

diff --git a/Engine/Source/Runtime/AIModule/Classes/EnvironmentQuery/EnvQueryTypes.h b/Engine/Source/Runtime/AIModule/Classes/EnvironmentQuery/EnvQueryTypes.h
--- a/Engine/Source/Runtime/AIModule/Classes/EnvironmentQuery/EnvQueryTypes.h
+++ b/Engine/Source/Runtime/AIModule/Classes/EnvironmentQuery/EnvQueryTypes.h
@@ -546,10 +546,26 @@ public:
 	AActor* GetItemAsActor(int32 Index) const;
 	FVector GetItemAsLocation(int32 Index) const;
 
+	template<typename TItemType>
+	const typename TItemType::FValueType& GetItemAsTypeChecked(const int32 Index) const
+	{
+		check(Items.IsValidIndex(Index) 
+			&& ItemType.Get()
+			&& ItemType->IsChildOf(TItemType::StaticClass()));
+
+		return TItemType::template GetValueFromMemory<typename TItemType::FValueType>(GetItemRawMemory(Index));
+	}
+
 	/** note that this function does not strip out the null-actors to not mess up results of GetItemScore(Index) calls*/
 	void GetAllAsActors(TArray<AActor*>& OutActors) const;
 	void GetAllAsLocations(TArray<FVector>& OutLocations) const;
 
+	const uint8* GetItemRawMemory(const int32 Index) const
+	{
+		check(Items.IsValidIndex(Index));
+		return RawData.GetData() + Items[Index].DataOffset;
+	}
+
 	FEnvQueryResult() : ItemType(NULL), Status(EEnvQueryStatus::Processing), OptionIndex(0), QueryID(0) {}
 	FEnvQueryResult(const EEnvQueryStatus::Type& InStatus) : ItemType(NULL), Status(InStatus), OptionIndex(0), QueryID(0) {}
 
diff --git a/Engine/Source/Runtime/AIModule/Classes/EnvironmentQuery/Items/EnvQueryItemType.h b/Engine/Source/Runtime/AIModule/Classes/EnvironmentQuery/Items/EnvQueryItemType.h
--- a/Engine/Source/Runtime/AIModule/Classes/EnvironmentQuery/Items/EnvQueryItemType.h
+++ b/Engine/Source/Runtime/AIModule/Classes/EnvironmentQuery/Items/EnvQueryItemType.h
@@ -39,18 +39,18 @@ public:
 	/** describe item */
 	virtual FString GetDescription(const uint8* RawData) const;
 
-protected:
-
-	/** size of value for this type */
-	uint16 ValueSize;
-
 	/** helper function for reading typed data from memory block */
 	template<typename T>
-	static T GetValueFromMemory(const uint8* MemoryBlock)
+	static const T& GetValueFromMemory(const uint8* MemoryBlock)
 	{
 		return *((T*)MemoryBlock);
 	}
 
+protected:
+
+	/** size of value for this type */
+	uint16 ValueSize;
+
 	/** helper function for writing typed data to memory block */
 	template<typename T>
 	static void SetValueInMemory(uint8* MemoryBlock, const T& Value)
