Fixed BT component not redirecting its vislog to owning AIController is the component is dynamically added as part of AAIController::RunBehaviorTree

diff --git a/Engine/Source/Runtime/AIModule/Private/AIController.cpp b/Engine/Source/Runtime/AIModule/Private/AIController.cpp
--- a/Engine/Source/Runtime/AIModule/Private/AIController.cpp
+++ b/Engine/Source/Runtime/AIModule/Private/AIController.cpp
@@ -921,6 +921,7 @@ bool AAIController::RunBehaviorTree(UBehaviorTree* BTAsset)
 
 			BTComp = NewObject<UBehaviorTreeComponent>(this, TEXT("BTComponent"));
 			BTComp->RegisterComponent();
+			REDIRECT_OBJECT_TO_VLOG(BTComp, this);
 		}
 		
 		// make sure BrainComponent points at the newly created BT component
