Fix for OnNewPawn not getting called for loaded pawns. This was causing an issue with pawns placed in data layers auto registering as perception sources.

diff --git a/Engine/Source/Runtime/AIModule/Classes/AISystem.h b/Engine/Source/Runtime/AIModule/Classes/AISystem.h
--- a/Engine/Source/Runtime/AIModule/Classes/AISystem.h
+++ b/Engine/Source/Runtime/AIModule/Classes/AISystem.h
@@ -129,6 +129,8 @@ protected:
 
 	FDelegateHandle ActorSpawnedDelegateHandle;
 
+	FDelegateHandle PawnBeginPlayDelegateHandle;
+
 	/** random number stream to be used by all things AI. WIP */
 	static FRandomStream RandomStream;
 	
@@ -276,5 +278,7 @@ public:
 
 protected:
 	virtual void OnActorSpawned(AActor* SpawnedActor);
+	virtual void OnPawnBeginPlay(APawn* Pawn);
+
 	void LoadDebuggerPlugin();
 };
diff --git a/Engine/Source/Runtime/AIModule/Private/AISystem.cpp b/Engine/Source/Runtime/AIModule/Private/AISystem.cpp
--- a/Engine/Source/Runtime/AIModule/Private/AISystem.cpp
+++ b/Engine/Source/Runtime/AIModule/Private/AISystem.cpp
@@ -83,10 +83,12 @@ void UAISystem::PostInitProperties()
 
 		if (WorldOuter)
 		{
-			FOnActorSpawned::FDelegate ActorSpawnedDelegate = FOnActorSpawned::FDelegate::CreateUObject(this, &UAISystem::OnActorSpawned);
+			const FOnActorSpawned::FDelegate ActorSpawnedDelegate = FOnActorSpawned::FDelegate::CreateUObject(this, &UAISystem::OnActorSpawned);
 			ActorSpawnedDelegateHandle = WorldOuter->AddOnActorSpawnedHandler(ActorSpawnedDelegate);
 		}
 
+		PawnBeginPlayDelegateHandle = APawn::OnPawnBeginPlay.AddUObject(this, &UAISystem::OnPawnBeginPlay);
+
 		ConditionalLoadDebuggerPlugin();
 	}
 }
@@ -103,15 +105,23 @@ void UAISystem::StartPlay()
 
 void UAISystem::OnActorSpawned(AActor* SpawnedActor)
 {
+}
+
+void UAISystem::OnPawnBeginPlay(APawn* Pawn)
+{
+	check(Pawn);
+
 	if (PerceptionSystem == nullptr || PerceptionSystem->bHandlePawnNotification == false)
 	{
 		return;
 	}
 
-	APawn* AsPawn = Cast<APawn>(SpawnedActor);
-	if (AsPawn)
+	const UWorld* const PawnWorld = Pawn->GetWorld();
+	check(PawnWorld);
+
+	if (PawnWorld == GetWorld())
 	{
-		PerceptionSystem->OnNewPawn(*AsPawn);
+		PerceptionSystem->OnNewPawn(*Pawn);
 	}
 }
 
@@ -142,6 +152,13 @@ void UAISystem::CleanupWorld(bool bSessionEnded, bool bCleanupResources)
 			EnvironmentQueryManager = nullptr;
 		}
 	}
+
+	const UWorld* const WorldOuter = GetOuterWorld();
+	if (WorldOuter)
+	{
+		WorldOuter->RemoveOnActorSpawnedHandler(ActorSpawnedDelegateHandle);
+	}
+	APawn::OnPawnBeginPlay.Remove(PawnBeginPlayDelegateHandle);
 }
 
 void UAISystem::AIIgnorePlayers()
