Upgraded to use FVector::FReal rather than floats.

diff --git a/Engine/Plugins/Runtime/ZoneGraph/Source/ZoneGraph/Public/ZoneGraphAStar.h b/Engine/Plugins/Runtime/ZoneGraph/Source/ZoneGraph/Public/ZoneGraphAStar.h
--- a/Engine/Plugins/Runtime/ZoneGraph/Source/ZoneGraph/Public/ZoneGraphAStar.h
+++ b/Engine/Plugins/Runtime/ZoneGraph/Source/ZoneGraph/Public/ZoneGraphAStar.h
@@ -39,7 +39,7 @@ struct ZONEGRAPH_API FZoneGraphAStarNode : public FGraphAStarDefaultNode<FZoneGr
 	typedef FGraphAStarDefaultNode<FZoneGraphAStarWrapper> Super;
 	typedef int32/*lane index type*/ FNodeRef;
 
-	FORCEINLINE FZoneGraphAStarNode(const FNodeRef InNodeRef = INDEX_NONE, const FVector InPosition = FVector(MAX_flt))
+	FORCEINLINE FZoneGraphAStarNode(const FNodeRef InNodeRef = INDEX_NONE, const FVector InPosition = FVector(TNumericLimits<FVector::FReal>::Max()))
 		: Super(InNodeRef)
 		, Position(InPosition)
 	{}
@@ -47,7 +47,7 @@ struct ZONEGRAPH_API FZoneGraphAStarNode : public FGraphAStarDefaultNode<FZoneGr
 	FZoneGraphAStarNode(const FZoneGraphAStarNode& Other) = default;
 	FGraphAStarDefaultNode& operator=(const FGraphAStarDefaultNode& Other) = delete;
 
-	bool IsStartOrIsEnd() const { return Position != FVector(MAX_flt); };
+	bool IsStartOrIsEnd() const { return Position != FVector(TNumericLimits<FVector::FReal>::Max()); };
 
 	FVector Position; //@todo: this will likely change to be a position along the lane
 };
@@ -68,11 +68,11 @@ struct ZONEGRAPH_API FZoneGraphPathFilter
 	bool IsStart(const FZoneGraphAStarNode& Node) const { return Node.NodeRef == StartLocation.LaneHandle.Index; }
 	bool IsEnd(const FZoneGraphAStarNode& Node) const { return Node.NodeRef == EndLocation.LaneHandle.Index; }
 
-	FORCEINLINE float GetHeuristicScale() const { return 1.f; }
+	FORCEINLINE FVector::FReal GetHeuristicScale() const { return 1.; }
 
-	float GetHeuristicCost(const FZoneGraphAStarNode& NeighbourNode, const FZoneGraphAStarNode& EndNode) const
+	FVector::FReal GetHeuristicCost(const FZoneGraphAStarNode& NeighbourNode, const FZoneGraphAStarNode& EndNode) const
 	{
-		static const FVector InvalidPosition(MAX_flt);
+		static const FVector InvalidPosition(TNumericLimits<FVector::FReal>::Max());
 		
 		// Except for start to end nodes, compute heurisitic from the last point of lanes
 		if (NeighbourNode.Position == InvalidPosition)
@@ -92,7 +92,7 @@ struct ZONEGRAPH_API FZoneGraphPathFilter
 		}
 	}
 
-	float GetTraversalCost(const FZoneGraphAStarNode& CurNode, const FZoneGraphAStarNode& NeighbourNode) const
+	FVector::FReal GetTraversalCost(const FZoneGraphAStarNode& CurNode, const FZoneGraphAStarNode& NeighbourNode) const
 	{
 		const FZoneLaneData& CurLane = ZoneStorage.Lanes[CurNode.NodeRef];
 		const FZoneLaneData& NeighbourLane = ZoneStorage.Lanes[NeighbourNode.NodeRef];
@@ -159,7 +159,7 @@ struct ZONEGRAPH_API FZoneGraphPathFilter
 			else
 			{
 				UE_LOG(LogAStar, Warning, TEXT("Not handled"));
-				return MAX_flt;
+				return TNumericLimits<FVector::FReal>::Max();
 			}
 		}
 	}
diff --git a/Engine/Source/Runtime/AIModule/Classes/Navigation/RecastGraphAStar.h b/Engine/Source/Runtime/AIModule/Classes/Navigation/RecastGraphAStar.h
--- a/Engine/Source/Runtime/AIModule/Classes/Navigation/RecastGraphAStar.h
+++ b/Engine/Source/Runtime/AIModule/Classes/Navigation/RecastGraphAStar.h
@@ -167,29 +167,29 @@ struct AIMODULE_API FRecastAStar : public FGraphAStar<FRecastGraphWrapper, FReca
 
 struct AIMODULE_API FRecastGraphAStarFilter
 {
-	FRecastGraphAStarFilter(FRecastGraphWrapper& InRecastGraphWrapper, const FRecastQueryFilter& InFilter, uint32 InMaxSearchNodes, const float InCostLimit, const UObject* Owner);
+	FRecastGraphAStarFilter(FRecastGraphWrapper& InRecastGraphWrapper, const FRecastQueryFilter& InFilter, uint32 InMaxSearchNodes, const FVector::FReal InCostLimit, const UObject* Owner);
 
 	FORCEINLINE bool WantsPartialSolution() const
 	{ 
 		return true; 
 	}
-	FORCEINLINE float GetHeuristicScale() const
+	FORCEINLINE FVector::FReal GetHeuristicScale() const
 	{ 
 		// LWC_TODO_AI: Refactor all of these related classes in multiple files to use doubles! Probably not until after 5.0!
-		return UE_REAL_TO_FLOAT(Filter.getHeuristicScale());
+		return Filter.getHeuristicScale();
 	}
 
-	FORCEINLINE float GetHeuristicCost(const FRecastAStarSearchNode& StartNode, const FRecastAStarSearchNode& EndNode) const
+	FORCEINLINE FVector::FReal GetHeuristicCost(const FRecastAStarSearchNode& StartNode, const FRecastAStarSearchNode& EndNode) const
 	{
 		check(EndNode.HasValidCacheInfo());
 
 		const FVector::FReal Cost = dtVdist(StartNode.Position, EndNode.Position);
 
 		// LWC_TODO_AI: Refactor all of these related classes in multiple files to use doubles! Probably not until after 5.0!
-		return UE_REAL_TO_FLOAT_CLAMPED_MAX(Cost);
+		return Cost;
 	}
 
-	FORCEINLINE float GetTraversalCost(const FRecastAStarSearchNode& StartNode, const FRecastAStarSearchNode& EndNode) const
+	FORCEINLINE FVector::FReal GetTraversalCost(const FRecastAStarSearchNode& StartNode, const FRecastAStarSearchNode& EndNode) const
 	{
 		EndNode.CacheInfo(RecastGraphWrapper, StartNode);
 		const FVector::FReal Cost = Filter.getCost(
@@ -199,7 +199,7 @@ struct AIMODULE_API FRecastGraphAStarFilter
 			EndNode.NodeRef, EndNode.Tile, EndNode.Poly);
 
 		// LWC_TODO_AI: Refactor all of these related classes in multiple files to use doubles! Probably not until after 5.0!
-		return UE_REAL_TO_FLOAT_CLAMPED_MAX(Cost);
+		return Cost;
 	}
 
 	FORCEINLINE bool IsTraversalAllowed(const dtPolyRef& NodeA, const FRecastNeighbour& NodeB) const
@@ -229,7 +229,7 @@ struct AIMODULE_API FRecastGraphAStarFilter
 		return MaxSearchNodes;
 	}
 
-	FORCEINLINE float GetCostLimit() const
+	FORCEINLINE FVector::FReal GetCostLimit() const
 	{
 		return CostLimit;
 	}
@@ -246,5 +246,5 @@ private:
 	FRecastSpeciaLinkFilter LinkFilter;
 	const FRecastGraphWrapper& RecastGraphWrapper;
 	uint32 MaxSearchNodes;
-	float CostLimit;
+	FVector::FReal CostLimit;
 };
\ No newline at end of file
diff --git a/Engine/Source/Runtime/AIModule/Private/Navigation/NavLocalGridData.cpp b/Engine/Source/Runtime/AIModule/Private/Navigation/NavLocalGridData.cpp
--- a/Engine/Source/Runtime/AIModule/Private/Navigation/NavLocalGridData.cpp
+++ b/Engine/Source/Runtime/AIModule/Private/Navigation/NavLocalGridData.cpp
@@ -240,17 +240,17 @@ namespace FGridHelpers
 	{
 		FGridPathFilter(const FNavLocalGridData& InGridRef) : GridRef(InGridRef) {}
 
-		float GetHeuristicScale() const
+		FVector::FReal GetHeuristicScale() const
 		{
-			return 1.0f;
+			return 1.;
 		}
 
-		float GetHeuristicCost(const int32 StartNodeRef, const int32 EndNodeRef) const
+		FVector::FReal GetHeuristicCost(const int32 StartNodeRef, const int32 EndNodeRef) const
 		{
 			return GetTraversalCost(StartNodeRef, EndNodeRef);
 		}
 
-		float GetTraversalCost(const int32 StartNodeRef, const int32 EndNodeRef) const
+		FVector::FReal GetTraversalCost(const int32 StartNodeRef, const int32 EndNodeRef) const
 		{
 			// manhattan distance
 			return FMath::Abs(GridRef.GetCellCoordX(StartNodeRef) - GridRef.GetCellCoordX(EndNodeRef)) + FMath::Abs(GridRef.GetCellCoordY(StartNodeRef) - GridRef.GetCellCoordY(EndNodeRef));
diff --git a/Engine/Source/Runtime/AIModule/Private/Navigation/RecastGraphAStar.cpp b/Engine/Source/Runtime/AIModule/Private/Navigation/RecastGraphAStar.cpp
--- a/Engine/Source/Runtime/AIModule/Private/Navigation/RecastGraphAStar.cpp
+++ b/Engine/Source/Runtime/AIModule/Private/Navigation/RecastGraphAStar.cpp
@@ -79,7 +79,7 @@ dtPolyRef FRecastAStarResult::SetPathInfo(const int32 Index, const FRecastAStarS
 	return SearchNode.NodeRef;
 }
 
-FRecastGraphAStarFilter::FRecastGraphAStarFilter(FRecastGraphWrapper& InRecastGraphWrapper, const FRecastQueryFilter& InFilter, const uint32 InMaxSearchNodes, const float InCostLimit, const UObject* Owner)
+FRecastGraphAStarFilter::FRecastGraphAStarFilter(FRecastGraphWrapper& InRecastGraphWrapper, const FRecastQueryFilter& InFilter, const uint32 InMaxSearchNodes, const FVector::FReal InCostLimit, const UObject* Owner)
 	: Filter{ InFilter }
 	, LinkFilter{ FNavigationSystem::GetCurrent<UNavigationSystemV1>(InRecastGraphWrapper.GetRecastNavMeshActor()->GetWorld()), Owner }
 	, RecastGraphWrapper{ InRecastGraphWrapper }
diff --git a/Engine/Source/Runtime/AIModule/Public/GraphAStar.h b/Engine/Source/Runtime/AIModule/Public/GraphAStar.h
--- a/Engine/Source/Runtime/AIModule/Public/GraphAStar.h
+++ b/Engine/Source/Runtime/AIModule/Public/GraphAStar.h
@@ -41,8 +41,8 @@ struct FGraphAStarDefaultNode
 
 	const FGraphNodeRef NodeRef;
 	FGraphNodeRef ParentRef;
-	float TraversalCost;
-	float TotalCost;
+	FVector::FReal TraversalCost;
+	FVector::FReal TotalCost;
 	int32 SearchNodeIndex;
 	int32 ParentNodeIndex;
 	uint8 bIsOpened : 1;
@@ -51,8 +51,8 @@ struct FGraphAStarDefaultNode
 	FORCEINLINE FGraphAStarDefaultNode(const FGraphNodeRef& InNodeRef)
 		: NodeRef(InNodeRef)
 		, ParentRef(TIsPointer<FGraphNodeRef>::Value ? (FGraphNodeRef)0 : (FGraphNodeRef)INDEX_NONE)
-		, TraversalCost(FLT_MAX)
-		, TotalCost(FLT_MAX)
+		, TraversalCost(TNumericLimits<FVector::FReal>::Max())
+		, TotalCost(TNumericLimits<FVector::FReal>::Max())
 		, SearchNodeIndex(INDEX_NONE)
 		, ParentNodeIndex(INDEX_NONE)
 		, bIsOpened(false)
@@ -123,30 +123,30 @@ template <> struct TAllocatorTraits<FRangeChecklessAllocator<false>> : TAllocato
  *	Generic graph A* implementation
  *
  *	TGraph holds graph representation. Needs to implement functions:
- *		bool IsValidRef(FNodeRef NodeRef) const;												- returns whether given node identyfication is correct
- *		FNodeRef GetNeighbour(const FSearchNode& NodeRef, const int32 NeighbourIndex) const;	- returns neighbour ref
+ *		bool IsValidRef(FNodeRef NodeRef) const;															- returns whether given node identyfication is correct
+ *		FNodeRef GetNeighbour(const FSearchNode& NodeRef, const int32 NeighbourIndex) const;				- returns neighbour ref
  *
  *	it also needs to specify node type
  *		FNodeRef		- type used as identification of nodes in the graph
  *
  *	TQueryFilter (FindPath's parameter) filter class is what decides which graph edges can be used and at what cost. It needs to implement following functions:
- *		float GetHeuristicScale() const;														- used as GetHeuristicCost's multiplier
- *		float GetHeuristicCost(const FSearchNode& StartNode, const FSearchNode& EndNode) const;	- estimate of cost from StartNode to EndNode from a search node
- *		float GetTraversalCost(const FSearchNode& StartNode, const FSearchNode& EndNode) const;	- real cost of traveling from StartNode directly to EndNode from a search node
- *		bool IsTraversalAllowed(const FNodeRef NodeA, const FNodeRef NodeB) const;				- whether traversing given edge is allowed from a NodeRef
- *		bool WantsPartialSolution() const;														- whether to accept solutions that do not reach the goal
+ *		FVector::FReal GetHeuristicScale() const;															- used as GetHeuristicCost's multiplier
+ *		FVector::FReal GetHeuristicCost(const FSearchNode& StartNode, const FSearchNode& EndNode) const;	- estimate of cost from StartNode to EndNode from a search node
+ *		FVector::FReal GetTraversalCost(const FSearchNode& StartNode, const FSearchNode& EndNode) const;	- real cost of traveling from StartNode directly to EndNode from a search node
+ *		bool IsTraversalAllowed(const FNodeRef NodeA, const FNodeRef NodeB) const;							- whether traversing given edge is allowed from a NodeRef
+ *		bool WantsPartialSolution() const;																	- whether to accept solutions that do not reach the goal
  *
  *		// Backward compatible methods, please use the FSearchNode version. If the FSearchNode version are implemented, these methods will not be called at all.
- *		FNodeRef GetNeighbour(const FNodeRef NodeRef, const int32 NeighbourIndex) const;		- returns neighbour ref
- *		float GetHeuristicCost(const FNodeRef StartNodeRef, const FNodeRef EndNodeRef) const;	- estimate of cost from StartNode to EndNode from a NodeRef
- *		float GetTraversalCost(const FNodeRef StartNodeRef, const FNodeRef EndNodeRef) const;	- real cost of traveling from StartNode directly to EndNode from a NodeRef
+ *		FNodeRef GetNeighbour(const FNodeRef NodeRef, const int32 NeighbourIndex) const;					- returns neighbour ref
+ *		FVector::FReal GetHeuristicCost(const FNodeRef StartNodeRef, const FNodeRef EndNodeRef) const;		- estimate of cost from StartNode to EndNode from a NodeRef
+ *		FVector::FReal GetTraversalCost(const FNodeRef StartNodeRef, const FNodeRef EndNodeRef) const;		- real cost of traveling from StartNode directly to EndNode from a NodeRef
  *
  *		// Optionally implemented methods to parameterize the search
- *		int32 GetNeighbourCount(FNodeRef NodeRef) const;										- returns number of neighbours that the graph node identified with NodeRef has, it is ok if not implemented, the logic will stop calling GetNeighbour once it received an invalid noderef
- *		bool ShouldIgnoreClosedNodes() const;													- whether to revisit closed node or not
- *		bool ShouldIncludeStartNodeInPath() const;												- whether to put the start node in the resulting path
- *		int32 GetMaxSearchNodes() const;														- whether to limit the number of search nodes to a maximum
- *		float GetCostLimit() const																- whether to limit the search to a maximum cost
+ *		int32 GetNeighbourCount(FNodeRef NodeRef) const;													- returns number of neighbours that the graph node identified with NodeRef has, it is ok if not implemented, the logic will stop calling GetNeighbour once it received an invalid noderef
+ *		bool ShouldIgnoreClosedNodes() const;																- whether to revisit closed node or not
+ *		bool ShouldIncludeStartNodeInPath() const;															- whether to put the start node in the resulting path
+ *		int32 GetMaxSearchNodes() const;																	- whether to limit the number of search nodes to a maximum
+ *		FVector::FReal GetCostLimit() const																	- whether to limit the search to a maximum cost
  */
 template<typename TGraph, typename Policy = FGraphAStarDefaultPolicy, typename TSearchNode = FGraphAStarDefaultNode<TGraph>, bool DoRangeCheck = false >
 struct FGraphAStar
@@ -283,14 +283,14 @@ struct FGraphAStar
 	DECLARE_OPTIONALLY_IMPLEMENTED_TEMPLATE_CLASS_FUNCTION_2PARAMS_CONST(FGraphNodeRef, GetNeighbour, const FSearchNode&, const int32, Obj.GetNeighbour(Param1.NodeRef,Param2));
 
 	// TQueryFilter optionally implemented wrapper methods
-	DECLARE_OPTIONALLY_IMPLEMENTED_TEMPLATE_CLASS_FUNCTION_2PARAMS_CONST(float, GetTraversalCost, const FSearchNode&, const FSearchNode&, Obj.GetTraversalCost(Param1.NodeRef, Param2.NodeRef))
-	DECLARE_OPTIONALLY_IMPLEMENTED_TEMPLATE_CLASS_FUNCTION_2PARAMS_CONST(float, GetHeuristicCost, const FSearchNode&, const FSearchNode&, Obj.GetHeuristicCost(Param1.NodeRef, Param2.NodeRef))
+	DECLARE_OPTIONALLY_IMPLEMENTED_TEMPLATE_CLASS_FUNCTION_2PARAMS_CONST(FVector::FReal, GetTraversalCost, const FSearchNode&, const FSearchNode&, Obj.GetTraversalCost(Param1.NodeRef, Param2.NodeRef))
+	DECLARE_OPTIONALLY_IMPLEMENTED_TEMPLATE_CLASS_FUNCTION_2PARAMS_CONST(FVector::FReal, GetHeuristicCost, const FSearchNode&, const FSearchNode&, Obj.GetHeuristicCost(Param1.NodeRef, Param2.NodeRef))
 	DECLARE_OPTIONALLY_IMPLEMENTED_TEMPLATE_CLASS_FUNCTION_CONST(bool, ShouldIgnoreClosedNodes, false);
 	DECLARE_OPTIONALLY_IMPLEMENTED_TEMPLATE_CLASS_FUNCTION_CONST(bool, ShouldIncludeStartNodeInPath, false);
-	DECLARE_OPTIONALLY_IMPLEMENTED_TEMPLATE_CLASS_FUNCTION_CONST(float, GetCostLimit, MAX_FLT);
+	DECLARE_OPTIONALLY_IMPLEMENTED_TEMPLATE_CLASS_FUNCTION_CONST(FVector::FReal, GetCostLimit, TNumericLimits<FVector::FReal>::Max());
 	// Custom methods implemented over TQueryFilter optionally implemented methods
 	DECLARE_OPTIONALLY_IMPLEMENTED_TEMPLATE_CLASS_FUNCTION_CUSTOM(GetMaxSearchNodes, bool, HasReachMaxSearchNodes, uint32 NodeCount, false, NodeCount >= Obj.GetMaxSearchNodes());
-	DECLARE_OPTIONALLY_IMPLEMENTED_TEMPLATE_CLASS_FUNCTION_CUSTOM(GetCostLimit, bool, HasExceededCostLimit, float Cost, false, Cost > Obj.GetCostLimit());
+	DECLARE_OPTIONALLY_IMPLEMENTED_TEMPLATE_CLASS_FUNCTION_CUSTOM(GetCostLimit, bool, HasExceededCostLimit, FVector::FReal Cost, false, Cost > Obj.GetCostLimit());
 
 	// TResultPathInfo optionally implemented wrapper methods
 	DECLARE_OPTIONALLY_IMPLEMENTED_TEMPLATE_CLASS_FUNCTION_2PARAMS(FGraphNodeRef, SetPathInfo, const int32, const FSearchNode&, Obj[Param1] = Param2.NodeRef);
@@ -301,12 +301,27 @@ struct FGraphAStar
 		NodePool.Reserve(Policy::NodePoolSize);
 	}
 
+	/**
+	 * Single run of A* loop: get node from open set and process neighbors
+	 * returns true if loop should be continued
+	 */
+	template<typename TQueryFilter>
+	UE_DEPRECATED(5.1, "Please use ProcessSingleNode() taking FVector::FReal& OutBestNodeCost instead!")
+	FORCEINLINE_DEBUGGABLE bool ProcessSingleNode(const FSearchNode& EndNode, const bool bIsBound, const TQueryFilter& Filter, int32& OutBestNodeIndex, float& OutBestNodeCost)
+	{
+		double BestNodeCost;
+
+		const bool bSucess = ProcessSingleNode(EndNode, bIsBound, Filter, OutBestNodeIndex, BestNodeCost);
+		OutBestNodeCost = UE_REAL_TO_FLOAT_CLAMPED_MAX(BestNodeCost);
+		return bSucess;
+	}
+
 	/** 
 	 * Single run of A* loop: get node from open set and process neighbors 
 	 * returns true if loop should be continued
 	 */
 	template<typename TQueryFilter>
-	FORCEINLINE_DEBUGGABLE bool ProcessSingleNode(const FSearchNode& EndNode, const bool bIsBound, const TQueryFilter& Filter, int32& OutBestNodeIndex, float& OutBestNodeCost)
+	FORCEINLINE_DEBUGGABLE bool ProcessSingleNode(const FSearchNode& EndNode, const bool bIsBound, const TQueryFilter& Filter, int32& OutBestNodeIndex, FVector::FReal& OutBestNodeCost)
 	{
 		// Pop next best node and put it on closed list
 		const int32 ConsideredNodeIndex = OpenList.PopIndex();
@@ -318,11 +333,11 @@ struct FGraphAStar
 		if (bIsBound && (ConsideredNodeUnsafe.NodeRef == EndNode.NodeRef))
 		{
 			OutBestNodeIndex = ConsideredNodeUnsafe.SearchNodeIndex;
-			OutBestNodeCost = 0.f;
+			OutBestNodeCost = 0.;
 			return false;
 		}
 
-		const float HeuristicScale = Filter.GetHeuristicScale();
+		const FVector::FReal HeuristicScale = Filter.GetHeuristicScale();
 
 		// consider every neighbor of BestNode
 		const int32 NeighbourCount = GetNeighbourCountV2(Graph, ConsideredNodeUnsafe);
@@ -378,11 +393,11 @@ struct FGraphAStar
 			}
 
 			// calculate cost and heuristic.
-			const float NewTraversalCost = GetTraversalCost(Filter, NodePool[ConsideredNodeIndex], NeighbourNode) + NodePool[ConsideredNodeIndex].TraversalCost;
-			const float NewHeuristicCost = bIsBound && (NeighbourNode.NodeRef != EndNode.NodeRef)
+			const FVector::FReal NewTraversalCost = GetTraversalCost(Filter, NodePool[ConsideredNodeIndex], NeighbourNode) + NodePool[ConsideredNodeIndex].TraversalCost;
+			const FVector::FReal NewHeuristicCost = bIsBound && (NeighbourNode.NodeRef != EndNode.NodeRef)
 				? (GetHeuristicCost(Filter, NeighbourNode, EndNode) * HeuristicScale)
-				: 0.f;
-			const float NewTotalCost = NewTraversalCost + NewHeuristicCost;
+				: 0.;
+			const FVector::FReal NewTotalCost = NewTraversalCost + NewHeuristicCost;
 
 			// check against cost limit
 			if (HasExceededCostLimit(Filter, NewTotalCost))
@@ -472,7 +487,7 @@ struct FGraphAStar
 		OpenList.Push(StartPoolNode);
 
 		int32 BestNodeIndex = StartPoolNode.SearchNodeIndex;
-		float BestNodeCost = StartPoolNode.TotalCost;
+		FVector::FReal BestNodeCost = StartPoolNode.TotalCost;
 
 		EGraphAStarResult Result = EGraphAStarResult::SearchSuccess;
 		const bool bIsBound = true;
@@ -484,7 +499,7 @@ struct FGraphAStar
 		}
 
 		// check if we've reached the goal
-		if (BestNodeCost != 0.f)
+		if (BestNodeCost != 0.)
 		{
 			Result = EGraphAStarResult::GoalUnreachable;
 		}
@@ -544,7 +559,7 @@ struct FGraphAStar
 		OpenList.Push(StartPoolNode);
 
 		int32 BestNodeIndex = StartPoolNode.SearchNodeIndex;
-		float BestNodeCost = StartPoolNode.TotalCost;
+		FVector::FReal BestNodeCost = StartPoolNode.TotalCost;
 		
 		const FSearchNode FakeEndNode = StartNode;
 		const bool bIsBound = false;
