Removed stale comments

diff --git a/Engine/Source/Runtime/AIModule/Classes/Navigation/RecastGraphAStar.h b/Engine/Source/Runtime/AIModule/Classes/Navigation/RecastGraphAStar.h
--- a/Engine/Source/Runtime/AIModule/Classes/Navigation/RecastGraphAStar.h
+++ b/Engine/Source/Runtime/AIModule/Classes/Navigation/RecastGraphAStar.h
@@ -175,7 +175,6 @@ struct AIMODULE_API FRecastGraphAStarFilter
 	}
 	FORCEINLINE FVector::FReal GetHeuristicScale() const
 	{ 
-		// LWC_TODO_AI: Refactor all of these related classes in multiple files to use doubles! Probably not until after 5.0!
 		return Filter.getHeuristicScale();
 	}
 
@@ -185,7 +184,6 @@ struct AIMODULE_API FRecastGraphAStarFilter
 
 		const FVector::FReal Cost = dtVdist(StartNode.Position, EndNode.Position);
 
-		// LWC_TODO_AI: Refactor all of these related classes in multiple files to use doubles! Probably not until after 5.0!
 		return Cost;
 	}
 
@@ -198,7 +196,6 @@ struct AIMODULE_API FRecastGraphAStarFilter
 			StartNode.NodeRef, StartNode.Tile, StartNode.Poly,
 			EndNode.NodeRef, EndNode.Tile, EndNode.Poly);
 
-		// LWC_TODO_AI: Refactor all of these related classes in multiple files to use doubles! Probably not until after 5.0!
 		return Cost;
 	}
 
