Optimization for BTTask_Wait
Added the same mechanic that already exists for AuxiliaryNodes to tell the behavior tree component their next desired tick time
Implemented it on BTTaskNode so all nodes can enable it via bTickIntervals
Used it in BTTask_Wait so wait task doesn't have to tick anymore

diff --git a/Engine/Source/Runtime/AIModule/Classes/BehaviorTree/BTTaskNode.h b/Engine/Source/Runtime/AIModule/Classes/BehaviorTree/BTTaskNode.h
--- a/Engine/Source/Runtime/AIModule/Classes/BehaviorTree/BTTaskNode.h
+++ b/Engine/Source/Runtime/AIModule/Classes/BehaviorTree/BTTaskNode.h
@@ -8,6 +8,12 @@
 
 class UBTService;
 
+struct FBTTaskMemory : public FBTInstancedNodeMemory
+{
+	float NextTickRemainingTime = 0.f;
+	float AccumulatedDeltaTime = 0.f;
+};
+
 /** 
  * Task are leaf nodes of behavior tree, which perform actual actions
  *
@@ -33,12 +39,17 @@ class AIMODULE_API UBTTaskNode : public UBTNode
 	 * this function should be considered as const (don't modify state of object) if node is not instanced! */
 	virtual EBTNodeResult::Type ExecuteTask(UBehaviorTreeComponent& OwnerComp, uint8* NodeMemory);
 
+	virtual uint16 GetSpecialMemorySize() const override;
+
 protected:
 	/** aborts this task, should return Aborted or InProgress
 	 *  (use FinishLatentAbort() when returning InProgress)
 	 * this function should be considered as const (don't modify state of object) if node is not instanced! */
 	virtual EBTNodeResult::Type AbortTask(UBehaviorTreeComponent& OwnerComp, uint8* NodeMemory);
 
+	/** sets next tick time */
+	void SetNextTickTime(uint8* NodeMemory, float RemainingTime) const;
+
 public:
 #if WITH_EDITOR
 	virtual FName GetNodeIconName() const override;
@@ -119,6 +130,9 @@ protected:
 		bNotifyTick = !TIsSame<decltype(&UBTTaskNode::TickTask), TickTask>::Value;
 		bNotifyTaskFinished = !TIsSame<decltype(&UBTTaskNode::OnTaskFinished), OnTaskFinished>::Value;
 	}
+
+	/** if set, conditional tick will use remaining time from node's memory */
+	uint8 bTickIntervals : 1;
 };
 
 #define INIT_TASK_NODE_NOTIFY_FLAGS() \
diff --git a/Engine/Source/Runtime/AIModule/Classes/BehaviorTree/Tasks/BTTask_Wait.h b/Engine/Source/Runtime/AIModule/Classes/BehaviorTree/Tasks/BTTask_Wait.h
--- a/Engine/Source/Runtime/AIModule/Classes/BehaviorTree/Tasks/BTTask_Wait.h
+++ b/Engine/Source/Runtime/AIModule/Classes/BehaviorTree/Tasks/BTTask_Wait.h
@@ -8,8 +8,6 @@
 
 struct FBTWaitTaskMemory
 {
-	/** time left */
-	float RemainingWaitTime;
 };
 
 /**
@@ -30,7 +28,6 @@ class AIMODULE_API UBTTask_Wait : public UBTTaskNode
 	float RandomDeviation;
 
 	virtual EBTNodeResult::Type ExecuteTask(UBehaviorTreeComponent& OwnerComp, uint8* NodeMemory) override;
-	virtual uint16 GetInstanceMemorySize() const override;
 	virtual void DescribeRuntimeValues(const UBehaviorTreeComponent& OwnerComp, uint8* NodeMemory, EBTDescriptionVerbosity::Type Verbosity, TArray<FString>& Values) const override;
 	virtual FString GetStaticDescription() const override;
 
diff --git a/Engine/Source/Runtime/AIModule/Private/BehaviorTree/BTTaskNode.cpp b/Engine/Source/Runtime/AIModule/Private/BehaviorTree/BTTaskNode.cpp
--- a/Engine/Source/Runtime/AIModule/Private/BehaviorTree/BTTaskNode.cpp
+++ b/Engine/Source/Runtime/AIModule/Private/BehaviorTree/BTTaskNode.cpp
@@ -10,6 +10,7 @@ UBTTaskNode::UBTTaskNode(const FObjectInitializer& ObjectInitializer) : Super(Ob
 	bNotifyTick = false;
 	bNotifyTaskFinished = false;
 	bIgnoreRestartSelf = false;
+	bTickIntervals = false;
 }
 
 EBTNodeResult::Type UBTTaskNode::ExecuteTask(UBehaviorTreeComponent& OwnerComp, uint8* NodeMemory)
@@ -22,6 +23,15 @@ EBTNodeResult::Type UBTTaskNode::AbortTask(UBehaviorTreeComponent& OwnerComp, ui
 	return EBTNodeResult::Aborted;
 }
 
+void UBTTaskNode::SetNextTickTime(uint8* NodeMemory, float RemainingTime) const
+{
+	if (bTickIntervals)
+	{
+		FBTTaskMemory* TaskMemory = GetSpecialNodeMemory<FBTTaskMemory>(NodeMemory);
+		TaskMemory->NextTickRemainingTime = RemainingTime;
+	}
+}
+
 EBTNodeResult::Type UBTTaskNode::WrappedExecuteTask(UBehaviorTreeComponent& OwnerComp, uint8* NodeMemory) const
 {
 	const UBTNode* NodeOb = bCreateNodeInstance ? GetNodeInstance(OwnerComp, NodeMemory) : this;
@@ -41,12 +51,40 @@ bool UBTTaskNode::WrappedTickTask(UBehaviorTreeComponent& OwnerComp, uint8* Node
 {
 	if (bNotifyTick)
 	{
-		const UBTNode* NodeOb = bCreateNodeInstance ? GetNodeInstance(OwnerComp, NodeMemory) : this;
+		const UBTTaskNode* NodeOb = bCreateNodeInstance ? static_cast<UBTTaskNode*>(GetNodeInstance(OwnerComp, NodeMemory)) : this;
 		if (NodeOb)
 		{
-			((UBTTaskNode*)NodeOb)->TickTask(OwnerComp, NodeMemory, DeltaSeconds);
-			NextNeededDeltaTime = 0.0f;
-			return true;
+			if (NodeOb->bTickIntervals)
+			{
+				FBTTaskMemory* TaskMemory = GetSpecialNodeMemory<FBTTaskMemory>(NodeMemory);
+				TaskMemory->NextTickRemainingTime -= DeltaSeconds;
+				TaskMemory->AccumulatedDeltaTime += DeltaSeconds;
+
+				const bool bTick = TaskMemory->NextTickRemainingTime <= 0.0f;
+				if (bTick)
+				{
+				    const float UseDeltaTime = TaskMemory->AccumulatedDeltaTime;
+				    TaskMemory->AccumulatedDeltaTime = 0.0f;
+					TaskMemory->NextTickRemainingTime = 0.0f;
+    
+					UE_VLOG(OwnerComp.GetOwner(), LogBehaviorTree, Verbose, TEXT("Ticking aux node: %s"), *UBehaviorTreeTypes::DescribeNodeHelper(this));
+
+					const_cast<UBTTaskNode*>(NodeOb)->TickTask(OwnerComp, NodeMemory, UseDeltaTime);
+				}
+
+				if (TaskMemory->NextTickRemainingTime < NextNeededDeltaTime)
+				{
+					NextNeededDeltaTime = TaskMemory->NextTickRemainingTime;
+				}
+
+				return bTick;
+			}
+			else
+			{
+				const_cast<UBTTaskNode*>(NodeOb)->TickTask(OwnerComp, NodeMemory, DeltaSeconds);
+				NextNeededDeltaTime = 0.0f;
+				return true;
+			}
 		}
 	}
 	return false;
@@ -148,6 +186,11 @@ void UBTTaskNode::StopWaitingForMessages(UBehaviorTreeComponent& OwnerComp) cons
 	OwnerComp.UnregisterMessageObserversFrom(this);
 }
 
+uint16 UBTTaskNode::GetSpecialMemorySize() const
+{
+	return bTickIntervals ? sizeof(FBTTaskMemory) : Super::GetSpecialMemorySize();
+}
+
 #if WITH_EDITOR
 
 FName UBTTaskNode::GetNodeIconName() const
diff --git a/Engine/Source/Runtime/AIModule/Private/BehaviorTree/Tasks/BTTask_Wait.cpp b/Engine/Source/Runtime/AIModule/Private/BehaviorTree/Tasks/BTTask_Wait.cpp
--- a/Engine/Source/Runtime/AIModule/Private/BehaviorTree/Tasks/BTTask_Wait.cpp
+++ b/Engine/Source/Runtime/AIModule/Private/BehaviorTree/Tasks/BTTask_Wait.cpp
@@ -6,27 +6,25 @@ UBTTask_Wait::UBTTask_Wait(const FObjectInitializer& ObjectInitializer) : Super(
 {
 	NodeName = "Wait";
 	WaitTime = 5.0f;
+	bTickIntervals = true;
 	INIT_TASK_NODE_NOTIFY_FLAGS();
 }
 
 EBTNodeResult::Type UBTTask_Wait::ExecuteTask(UBehaviorTreeComponent& OwnerComp, uint8* NodeMemory)
 {
-	FBTWaitTaskMemory* MyMemory = CastInstanceNodeMemory<FBTWaitTaskMemory>(NodeMemory);
-	MyMemory->RemainingWaitTime = FMath::FRandRange(FMath::Max(0.0f, WaitTime - RandomDeviation), (WaitTime + RandomDeviation));
+	const float RemainingWaitTime = FMath::FRandRange(FMath::Max(0.0f, WaitTime - RandomDeviation), (WaitTime + RandomDeviation));
+	SetNextTickTime(NodeMemory, RemainingWaitTime);
 	
 	return EBTNodeResult::InProgress;
 }
 
 void UBTTask_Wait::TickTask(UBehaviorTreeComponent& OwnerComp, uint8* NodeMemory, float DeltaSeconds)
 {
-	FBTWaitTaskMemory* MyMemory = CastInstanceNodeMemory<FBTWaitTaskMemory>(NodeMemory);
-	MyMemory->RemainingWaitTime -= DeltaSeconds;
+	// Using the SetNextTickTime in ExecuteTask we are certain we are only getting ticked when the wait is finished
+	ensure(GetSpecialNodeMemory<FBTTaskMemory>(NodeMemory)->NextTickRemainingTime <= 0.f);
 
-	if (MyMemory->RemainingWaitTime <= 0.0f)
-	{
-		// continue execution from this node
-		FinishLatentTask(OwnerComp, EBTNodeResult::Succeeded);
-	}
+	// continue execution from this node
+	FinishLatentTask(OwnerComp, EBTNodeResult::Succeeded);
 }
 
 FString UBTTask_Wait::GetStaticDescription() const
@@ -45,18 +43,13 @@ void UBTTask_Wait::DescribeRuntimeValues(const UBehaviorTreeComponent& OwnerComp
 {
 	Super::DescribeRuntimeValues(OwnerComp, NodeMemory, Verbosity, Values);
 
-	FBTWaitTaskMemory* MyMemory = CastInstanceNodeMemory<FBTWaitTaskMemory>(NodeMemory);
-	if (MyMemory->RemainingWaitTime)
+	FBTTaskMemory* TaskMemory = GetSpecialNodeMemory<FBTTaskMemory>(NodeMemory);
+	if (TaskMemory->NextTickRemainingTime)
 	{
-		Values.Add(FString::Printf(TEXT("remaining: %ss"), *FString::SanitizeFloat(MyMemory->RemainingWaitTime)));
+		Values.Add(FString::Printf(TEXT("remaining: %ss"), *FString::SanitizeFloat(TaskMemory->NextTickRemainingTime)));
 	}
 }
 
-uint16 UBTTask_Wait::GetInstanceMemorySize() const
-{
-	return sizeof(FBTWaitTaskMemory);
-}
-
 #if WITH_EDITOR
 
 FName UBTTask_Wait::GetNodeIconName() const
