Fix for EQS Test's "ReferenceValue" having the effect opposite to its desscription and designed function.
The fix does ensure proper fix up of existing assets (in PostLoad)

diff --git a/Engine/Source/Runtime/AIModule/Classes/EnvironmentQuery/EnvQueryTest.h b/Engine/Source/Runtime/AIModule/Classes/EnvironmentQuery/EnvQueryTest.h
--- a/Engine/Source/Runtime/AIModule/Classes/EnvironmentQuery/EnvQueryTest.h
+++ b/Engine/Source/Runtime/AIModule/Classes/EnvironmentQuery/EnvQueryTest.h
@@ -31,8 +31,9 @@ namespace EnvQueryTestVersion
 {
 	const int32 Initial = 0;
 	const int32 DataProviders = 1;
+	const int32 ReferenceValueFix = 2;
 
-	const int32 Latest = DataProviders;
+	const int32 Latest = ReferenceValueFix;
 }
 
 #if WITH_EDITORONLY_DATA
diff --git a/Engine/Source/Runtime/AIModule/Private/EnvironmentQuery/EnvQueryTest.cpp b/Engine/Source/Runtime/AIModule/Private/EnvironmentQuery/EnvQueryTest.cpp
--- a/Engine/Source/Runtime/AIModule/Private/EnvironmentQuery/EnvQueryTest.cpp
+++ b/Engine/Source/Runtime/AIModule/Private/EnvironmentQuery/EnvQueryTest.cpp
@@ -32,7 +32,7 @@ void UEnvQueryTest::NormalizeItemScores(FEnvQueryInstance& QueryInstance)
 	}
 
 	ScoringFactor.BindData(QueryOwner, QueryInstance.QueryID);
-	float ScoringFactorValue = ScoringFactor.GetValue();
+	const float ScoringFactorValue = ScoringFactor.GetValue();
 
 	float MinScore = (NormalizationType == EEQSNormalizationType::Absolute) ? 0 : BIG_NUMBER;
 	float MaxScore = -BIG_NUMBER;
@@ -95,8 +95,9 @@ void UEnvQueryTest::NormalizeItemScores(FEnvQueryInstance& QueryInstance)
 		{
 			ReferenceValue.BindData(QueryOwner, QueryInstance.QueryID);
 		}
-		const float LocalReferenceValue = bDefineReferenceValue ? ReferenceValue.GetValue() : MinScore;
-		const float ValueSpan = FMath::Max(FMath::Abs(LocalReferenceValue - MinScore), FMath::Abs(LocalReferenceValue - MaxScore));
+		const float TargetScore = bDefineReferenceValue ? ReferenceValue.GetValue() : MaxScore;
+		const float ValueSpan = FMath::Max(FMath::Abs(TargetScore - MinScore), FMath::Abs(TargetScore - MaxScore));
+		const float AbsoluteWeight = FMath::Abs(ScoringFactorValue);
 
 		for (int32 ItemIndex = 0; ItemIndex < QueryInstance.ItemDetails.Num(); ItemIndex++, DetailInfo++)
 		{
@@ -112,9 +113,8 @@ void UEnvQueryTest::NormalizeItemScores(FEnvQueryInstance& QueryInstance)
 			{
 				const float ClampedScore = FMath::Clamp(TestValue, MinScore, MaxScore);
 				const float NormalizedScore = (ScoringFactorValue >= 0) 
-					? (FMath::Abs(LocalReferenceValue - ClampedScore) / ValueSpan)
-					: (1.f - (FMath::Abs(LocalReferenceValue - ClampedScore) / ValueSpan));
-				const float AbsoluteWeight = (ScoringFactorValue >= 0) ? ScoringFactorValue : -ScoringFactorValue;
+					? (1.f - (FMath::Abs(TargetScore - ClampedScore) / ValueSpan))
+					: (FMath::Abs(TargetScore - ClampedScore) / ValueSpan);
 
 				switch (ScoringEquation)
 				{
@@ -194,6 +194,52 @@ AActor* UEnvQueryTest::GetItemActor(FEnvQueryInstance& QueryInstance, int32 Item
 void UEnvQueryTest::PostLoad()
 {
 	Super::PostLoad();
+
+	if (VerNum < EnvQueryTestVersion::ReferenceValueFix && bDefineReferenceValue)
+	{
+		// if someone was using the ReferenceValue the fix will reverse their score. We can fix it by flipping the 
+		// ScoringFactor. Unfortunately we won't be able to to that if the ScoringFactor is actually bound to some 
+		// provider - in that case all we can do is to spit out a warning. 
+
+		UWorld* World = GetWorld();
+		if (World == nullptr)
+		{
+			// we want this logic to run only for assets, not the instantiated queries that should pull the following change in automatically
+			if (ScoringFactor.IsDynamic() == false)
+			{
+				ScoringFactor.DefaultValue *= -1;
+				UE_LOG(LogEQS, Verbose, TEXT("%s using a ScoringFactor. Due to a bug-fix the effect will be flipped so the\
+					value\'s sign has been flipped for this EQS Test asset instance.")
+					, *GetPathName());
+			}
+			else if (ScoringEquation == EEnvTestScoreEquation::Linear)
+			{
+				// we can still address this by flipping the ScoringEquation to 
+				ScoringEquation = EEnvTestScoreEquation::InverseLinear;
+				UE_LOG(LogEQS, Verbose, TEXT("%s using a AIDataProvider for ScoringFactor. Due to a bug-fix the effect will be flipped so ScoringEquation has been automatically flipped from Linear to InverseLinear for this EQS Test instance.")
+					, *GetPathName());
+			}
+			else if (ScoringEquation == EEnvTestScoreEquation::InverseLinear)
+			{
+				// we can still address this by flipping the ScoringEquation to 
+				ScoringEquation = EEnvTestScoreEquation::Linear;
+				UE_LOG(LogEQS, Verbose, TEXT("%s using a AIDataProvider for ScoringFactor. Due to a bug-fix the effect will be flipped so ScoringEquation has been automatically flipped from InverseLinear to Linear for this EQS Test instance.")
+					, *GetPathName());
+			}
+			else
+			{
+				UE_LOG(LogEQS, Warning, TEXT("%s using a AIDataProvider for ScoringFactor. Due to a bug-fix the effect will be flipped and an automated change could not be performed. Please address manually.")
+					, *GetPathName());
+			}
+
+			if (MarkPackageDirty() == false)
+			{
+				UE_LOG(LogEQS, Warning, TEXT("%s unable to mark package dirty - please resave manually or use a commandlet.")
+					, *GetPathName());
+			}
+		}
+	}
+
 	UpdateNodeVersion();
 #if WITH_EDITOR
 	UpdatePreviewData();
