PR #7909: Make OnTreeFinished virtual (Contributed by [REDACTED_PERSON])

diff --git a/Engine/Source/Runtime/AIModule/Classes/BehaviorTree/BehaviorTreeComponent.h b/Engine/Source/Runtime/AIModule/Classes/BehaviorTree/BehaviorTreeComponent.h
--- a/Engine/Source/Runtime/AIModule/Classes/BehaviorTree/BehaviorTreeComponent.h
+++ b/Engine/Source/Runtime/AIModule/Classes/BehaviorTree/BehaviorTreeComponent.h
@@ -394,7 +394,7 @@ protected:
 	UBTTaskNode* FindNextTask(UBTCompositeNode* ParentNode, uint16 ParentInstanceIdx, EBTNodeResult::Type LastResult);
 
 	/** called when tree runs out of nodes to execute */
-	void OnTreeFinished();
+	virtual void OnTreeFinished();
 
 	/** apply pending node updates from SearchData */
 	void ApplySearchData(UBTNode* NewActiveNode);
