Fixed Blackboard editor not clearing out display of all persistent keys when parented with an asset duplicating the same keys. Note that neigher set of persistent keys was visible until editor restart or asset unloading and reloading.

diff --git a/Engine/Source/Runtime/AIModule/Private/BehaviorTree/BlackboardData.cpp b/Engine/Source/Runtime/AIModule/Private/BehaviorTree/BlackboardData.cpp
--- a/Engine/Source/Runtime/AIModule/Private/BehaviorTree/BlackboardData.cpp
+++ b/Engine/Source/Runtime/AIModule/Private/BehaviorTree/BlackboardData.cpp
@@ -254,7 +254,12 @@ void UBlackboardData::UpdateParentKeys()
 		Parent = NULL;
 	}
 
+	UpdateKeyIDs();
+	UpdatePersistentKeys(*this);
+
 #if WITH_EDITORONLY_DATA
+	// note that we need to gather ParentKeys only once UpdatePersistentKeys was called since that will remove 
+	// this BB asset's persistent keys that double the ones already present in the parent asset.
 	ParentKeys.Reset();
 
 	for (UBlackboardData* It = Parent; It; It = It->Parent)
@@ -270,8 +275,6 @@ void UBlackboardData::UpdateParentKeys()
 	}
 #endif // WITH_EDITORONLY_DATA
 
-	UpdateKeyIDs();
-	UpdatePersistentKeys(*this);
 	OnUpdateKeys.Broadcast(this);
 }
 
