Fixed an EQSTestingPawn crash when undoing editor-time actor movement from "no eqs query results" location to a "some results" location.

diff --git a/Engine/Source/Runtime/AIModule/Private/EnvironmentQuery/EQSTestingPawn.cpp b/Engine/Source/Runtime/AIModule/Private/EnvironmentQuery/EQSTestingPawn.cpp
--- a/Engine/Source/Runtime/AIModule/Private/EnvironmentQuery/EQSTestingPawn.cpp
+++ b/Engine/Source/Runtime/AIModule/Private/EnvironmentQuery/EQSTestingPawn.cpp
@@ -305,12 +305,12 @@ void AEQSTestingPawn::UpdateDrawing()
 
 const FEnvQueryResult* AEQSTestingPawn::GetQueryResult() const 
 {
-	return StepResults.Num() > 0 ? &StepResults[StepToDebugDraw] : NULL;
+	return StepResults.IsValidIndex(StepToDebugDraw) ? &StepResults[StepToDebugDraw] : nullptr;
 }
 
 const FEnvQueryInstance* AEQSTestingPawn::GetQueryInstance() const 
 {
-	return StepResults.Num() > 0 ? &StepResults[StepToDebugDraw] : NULL;
+	return StepResults.IsValidIndex(StepToDebugDraw) ? &StepResults[StepToDebugDraw] : nullptr;
 }
 
 #if WITH_EDITOR
@@ -380,11 +380,20 @@ void AEQSTestingPawn::PostTransacted(const FTransactionObjectEvent& TransactionE
 {
 	Super::PostTransacted(TransactionEvent);
 
-	if (TransactionEvent.GetEventType() == ETransactionObjectEventType::UndoRedo)		
+	if (TransactionEvent.GetEventType() == ETransactionObjectEventType::UndoRedo)
 	{
-		for (const FName PropertyName : TransactionEvent.GetChangedProperties())
+		if (TransactionEvent.GetChangedProperties().Num() > 0)
 		{
-			OnPropertyChanged(PropertyName);
+			// targeted update
+			for (const FName PropertyName : TransactionEvent.GetChangedProperties())
+			{
+				OnPropertyChanged(PropertyName);
+			}
+		}
+		else
+		{
+			// fallback - make sure the results are up to date
+			RunEQSQuery();
 		}
 	}
 }
