Fixed UEnvQueryTest_Project's nav projection not applying ProjectionData.PostProjectionVerticalOffset

diff --git a/Engine/Source/Runtime/AIModule/Private/EnvironmentQuery/Tests/EnvQueryTest_Project.cpp b/Engine/Source/Runtime/AIModule/Private/EnvironmentQuery/Tests/EnvQueryTest_Project.cpp
--- a/Engine/Source/Runtime/AIModule/Private/EnvironmentQuery/Tests/EnvQueryTest_Project.cpp
+++ b/Engine/Source/Runtime/AIModule/Private/EnvironmentQuery/Tests/EnvQueryTest_Project.cpp
@@ -63,15 +63,30 @@ void UEnvQueryTest_Project::RunTest(FEnvQueryInstance& QueryInstance) const
 			NavData->BatchProjectPoints(Workload, ProjectionExtent, NavigationFilter);
 
 			int32 Idx = 0;
-			for (FEnvQueryInstance::ItemIterator It(this, QueryInstance); It; ++It, Idx++)
+			if (ProjectionData.PostProjectionVerticalOffset == 0.f)
 			{
-				const bool bProjected = Workload[Idx].bResult;
-				if (bProjected && ItemTypeCDO)
+				for (FEnvQueryInstance::ItemIterator It(this, QueryInstance); It; ++It, Idx++)
 				{
-					ItemTypeCDO->SetItemNavLocation(It.GetItemData(), Workload[Idx].OutLocation);
+					const bool bProjected = Workload[Idx].bResult;
+					if (bProjected && ItemTypeCDO)
+					{
+						ItemTypeCDO->SetItemNavLocation(It.GetItemData(), Workload[Idx].OutLocation);
+					}
+					It.SetScore(TestPurpose, FilterType, bProjected, bWantsProjected);
+				}
+			}
+			else
+			{
+				const FVector PostProjectionVerticalOffset(0.f, 0.f, ProjectionData.PostProjectionVerticalOffset);
+				for (FEnvQueryInstance::ItemIterator It(this, QueryInstance); It; ++It, Idx++)
+				{
+					const bool bProjected = Workload[Idx].bResult;
+					if (bProjected && ItemTypeCDO)
+					{
+						ItemTypeCDO->SetItemNavLocation(It.GetItemData(), FNavLocation(Workload[Idx].OutLocation.Location + PostProjectionVerticalOffset, Workload[Idx].OutLocation.NodeRef));
+					}
+					It.SetScore(TestPurpose, FilterType, bProjected, bWantsProjected);
 				}
-
-				It.SetScore(TestPurpose, FilterType, bProjected, bWantsProjected);
 			}
 		}
 	}
