PR #8303: DoNotCreateDefaultSubobject for AIController (Contributed by [REDACTED_PERSON])

Massaged a tiny bit to handle null-PathFollowingComponent case, and used a dedicated function to create the subobjects.

diff --git a/Engine/Source/Runtime/AIModule/Private/AIController.cpp b/Engine/Source/Runtime/AIModule/Private/AIController.cpp
--- a/Engine/Source/Runtime/AIModule/Private/AIController.cpp
+++ b/Engine/Source/Runtime/AIModule/Private/AIController.cpp
@@ -40,10 +40,13 @@ AAIController::AAIController(const FObjectInitializer& ObjectInitializer)
 	: Super(ObjectInitializer)
 {
 	bSetControlRotationFromPawnOrientation = true;
-	PathFollowingComponent = CreateDefaultSubobject<UPathFollowingComponent>(TEXT("PathFollowingComponent"));
-	PathFollowingComponent->OnRequestFinished.AddUObject(this, &AAIController::OnMoveCompleted);
+	PathFollowingComponent = CreateOptionalDefaultSubobject<UPathFollowingComponent>(TEXT("PathFollowingComponent"));
+	if (PathFollowingComponent)
+	{
+		PathFollowingComponent->OnRequestFinished.AddUObject(this, &AAIController::OnMoveCompleted);
+	}
 
-	ActionsComp = CreateDefaultSubobject<UPawnActionsComponent>("ActionsComp");
+	ActionsComp = CreateOptionalDefaultSubobject<UPawnActionsComponent>(TEXT("ActionsComp"));
 
 	bSkipExtraLOSChecks = true;
 	bWantsPlayerState = false;
