Improve comments only

diff --git a/Engine/Source/Runtime/AIModule/Classes/Navigation/NavLinkProxy.h b/Engine/Source/Runtime/AIModule/Classes/Navigation/NavLinkProxy.h
--- a/Engine/Source/Runtime/AIModule/Classes/Navigation/NavLinkProxy.h
+++ b/Engine/Source/Runtime/AIModule/Classes/Navigation/NavLinkProxy.h
@@ -20,6 +20,16 @@ struct FNavigationRelevantData;
 
 DECLARE_DYNAMIC_MULTICAST_DELEGATE_TwoParams( FSmartLinkReachedSignature, AActor*, MovingActor, const FVector&, DestinationPoint );
 
+/** 
+ *  ANavLinkProxy connects areas of Navmesh that don't have a direct navigation path.
+ *  It directly supports Simple Links (see PointLinks array). There can be multiple Simple links per ANavLinkProxy instance.
+ *  Simple links are designed to statically link areas of Navmesh and are associated with a particular area class that the link provides.
+ *  Smart Link functionality is provided via UNavLinkCustomComponent, see SmartLinkComp. They are designed to be able to be dynamically toggled
+ *  between enabled and disabled and provide different area classes for both cases. The area classes can be dynamically modified 
+ *  without navmesh rebuilds.
+ *  There can only be at most one smart link per ANavLinkProxy instance.
+ *  Both simple and smart links on a single ANavLinkProxy instance, can be set / enabled at once, as well as either or neither of them.
+ */
 UCLASS(Blueprintable, autoCollapseCategories=(SmartLink, Actor), hideCategories=(Input))
 class AIMODULE_API ANavLinkProxy : public AActor, public INavLinkHostInterface, public INavRelevantInterface
 {
