Hide deprecated and no-op settings from Editor Preferences and Project Settings.

diff --git a/Engine/Source/Developer/DeveloperToolSettings/Classes/Settings/ProjectPackagingSettings.h b/Engine/Source/Developer/DeveloperToolSettings/Classes/Settings/ProjectPackagingSettings.h
--- a/Engine/Source/Developer/DeveloperToolSettings/Classes/Settings/ProjectPackagingSettings.h
+++ b/Engine/Source/Developer/DeveloperToolSettings/Classes/Settings/ProjectPackagingSettings.h
@@ -544,7 +544,7 @@ public:
 	TArray<FString> CompressedChunkWildcard;
 
 	UE_DEPRECATED(5.1, "This property is no longer supported. Use IniKeyDenylist.")
-	UPROPERTY(config, EditAnywhere, Category = Packaging)
+	UPROPERTY(config /*, EditAnywhere, Category = Packaging */, meta = (DeprecatedProperty, DeprecationMessage = "This property is no longer supported. Use IniKeyDenylist."))
 	TArray<FString> IniKeyBlacklist;
 
 	/** List of ini file keys to strip when packaging */
@@ -552,7 +552,7 @@ public:
 	TArray<FString> IniKeyDenylist;
 
 	UE_DEPRECATED(5.1, "This property is no longer supported. Use IniSectionDenylist.")
-	UPROPERTY(config, EditAnywhere, Category = Packaging)
+	UPROPERTY(config /*, EditAnywhere, Category = Packaging */, meta = (DeprecatedProperty, DeprecationMessage = "This property is no longer supported. Use IniSectionDenylist."))
 	TArray<FString> IniSectionBlacklist;
 
 	/** List of ini file sections to strip when packaging */
diff --git a/Engine/Source/Runtime/AIModule/Classes/AISystem.h b/Engine/Source/Runtime/AIModule/Classes/AISystem.h
--- a/Engine/Source/Runtime/AIModule/Classes/AISystem.h
+++ b/Engine/Source/Runtime/AIModule/Classes/AISystem.h
@@ -67,12 +67,6 @@ public:
 	UPROPERTY(globalconfig, EditDefaultsOnly, Category = "Movement")
 	bool bAllowStrafing;
 
-	/** 
-	 * Deprecated: Whether or not to enable Gameplay Tasks for move tasks (always enabled now)
-	 */
-	UPROPERTY(globalconfig, EditDefaultsOnly, Category = "Gameplay Tasks", meta=(DisplayName="DEPRECATED Enable BT AITasks"))
-	bool bEnableBTAITasks;
-
 	/** if enable will make EQS not complaint about using Controllers as queriers. Default behavior (false) will 
 	 *	in places automatically convert controllers to pawns, and complain if code user bypasses the conversion or uses
 	 *	pawn-less controller */
diff --git a/Engine/Source/Runtime/AIModule/Private/AISystem.cpp b/Engine/Source/Runtime/AIModule/Private/AISystem.cpp
--- a/Engine/Source/Runtime/AIModule/Private/AISystem.cpp
+++ b/Engine/Source/Runtime/AIModule/Private/AISystem.cpp
@@ -29,8 +29,6 @@ UAISystem::UAISystem(const FObjectInitializer& ObjectInitializer)
 	bAllowStrafing = false;
 	DefaultSightCollisionChannel = ECC_Visibility;
 
-	bEnableBTAITasks = true;
-
 	if (HasAnyFlags(RF_ClassDefaultObject))
 	{
 		// game-wise config
diff --git a/Engine/Source/Runtime/AIModule/Private/BehaviorTree/Tasks/BTTask_MoveTo.cpp b/Engine/Source/Runtime/AIModule/Private/BehaviorTree/Tasks/BTTask_MoveTo.cpp
--- a/Engine/Source/Runtime/AIModule/Private/BehaviorTree/Tasks/BTTask_MoveTo.cpp
+++ b/Engine/Source/Runtime/AIModule/Private/BehaviorTree/Tasks/BTTask_MoveTo.cpp
@@ -120,56 +120,37 @@ EBTNodeResult::Type UBTTask_MoveTo::PerformMoveTask(UBehaviorTreeComponent& Owne
 
 		if (MoveReq.IsValid())
 		{
-			if (GET_AI_CONFIG_VAR(bEnableBTAITasks))
+			UAITask_MoveTo* MoveTask = MyMemory->Task.Get();
+			const bool bReuseExistingTask = (MoveTask != nullptr);
+
+			MoveTask = PrepareMoveTask(OwnerComp, MoveTask, MoveReq);
+			if (MoveTask)
 			{
-				UAITask_MoveTo* MoveTask = MyMemory->Task.Get();
-				const bool bReuseExistingTask = (MoveTask != nullptr);
+				MyMemory->bObserverCanFinishTask = false;
 
-				MoveTask = PrepareMoveTask(OwnerComp, MoveTask, MoveReq);
-				if (MoveTask)
+				if (bReuseExistingTask)
 				{
-					MyMemory->bObserverCanFinishTask = false;
-
-					if (bReuseExistingTask)
+					if (MoveTask->IsActive())
 					{
-						if (MoveTask->IsActive())
-						{
-							UE_VLOG(MyController, LogBehaviorTree, Verbose, TEXT("\'%s\' reusing AITask %s"), *GetNodeName(), *MoveTask->GetName());
-							MoveTask->ConditionalPerformMove();
-						}
-						else
-						{
-							UE_VLOG(MyController, LogBehaviorTree, Verbose, TEXT("\'%s\' reusing AITask %s, but task is not active - handing over move performing to task mechanics"), *GetNodeName(), *MoveTask->GetName());
-						}
+						UE_VLOG(MyController, LogBehaviorTree, Verbose, TEXT("\'%s\' reusing AITask %s"), *GetNodeName(), *MoveTask->GetName());
+						MoveTask->ConditionalPerformMove();
 					}
 					else
 					{
-						MyMemory->Task = MoveTask;
-						UE_VLOG(MyController, LogBehaviorTree, Verbose, TEXT("\'%s\' task implementing move with task %s"), *GetNodeName(), *MoveTask->GetName());
-						MoveTask->ReadyForActivation();
+						UE_VLOG(MyController, LogBehaviorTree, Verbose, TEXT("\'%s\' reusing AITask %s, but task is not active - handing over move performing to task mechanics"), *GetNodeName(), *MoveTask->GetName());
 					}
-
-					MyMemory->bObserverCanFinishTask = true;
-					NodeResult = (MoveTask->GetState() != EGameplayTaskState::Finished) ? EBTNodeResult::InProgress :
-						MoveTask->WasMoveSuccessful() ? EBTNodeResult::Succeeded :
-						EBTNodeResult::Failed;
 				}
-			}
-			else
-			{
-				const FPathFollowingRequestResult RequestResult = MyController->MoveTo(MoveReq);
-				if (RequestResult.Code == EPathFollowingRequestResult::RequestSuccessful)
-				{
-					MyMemory->MoveRequestID = RequestResult.MoveId;
-					WaitForMessage(OwnerComp, UBrainComponent::AIMessage_MoveFinished, RequestResult.MoveId);
-					WaitForMessage(OwnerComp, UBrainComponent::AIMessage_RepathFailed);
-
-					NodeResult = EBTNodeResult::InProgress;
-				}
-				else if (RequestResult.Code == EPathFollowingRequestResult::AlreadyAtGoal)
+				else
 				{
-					NodeResult = EBTNodeResult::Succeeded;
+					MyMemory->Task = MoveTask;
+					UE_VLOG(MyController, LogBehaviorTree, Verbose, TEXT("\'%s\' task implementing move with task %s"), *GetNodeName(), *MoveTask->GetName());
+					MoveTask->ReadyForActivation();
 				}
+
+				MyMemory->bObserverCanFinishTask = true;
+				NodeResult = (MoveTask->GetState() != EGameplayTaskState::Finished) ? EBTNodeResult::InProgress :
+					MoveTask->WasMoveSuccessful() ? EBTNodeResult::Succeeded :
+					EBTNodeResult::Failed;
 			}
 		}
 	}
diff --git a/Engine/Source/Runtime/Engine/Classes/Animation/AnimationSettings.h b/Engine/Source/Runtime/Engine/Classes/Animation/AnimationSettings.h
--- a/Engine/Source/Runtime/Engine/Classes/Animation/AnimationSettings.h
+++ b/Engine/Source/Runtime/Engine/Classes/Animation/AnimationSettings.h
@@ -48,7 +48,7 @@ class ENGINE_API UAnimationSettings : public UDeveloperSettings
 
 	/** If true and the existing compression error is greater than Alternative Compression Threshold, then Alternative Compression Threshold will be effectively raised to the existing error level */
 	UE_DEPRECATED(5.1, "This is being removed because it is unused")
-	UPROPERTY(config, EditAnywhere, Category = Compression)
+	UPROPERTY(config /*, EditAnywhere, Category = Compression*/, meta = (DeprecatedProperty, DeprecationMessage = "No longer used."))
 	bool bRaiseMaxErrorToExisting;
 
 	/** If true, recompression will log performance information */
