Refactor AI Sense Sight :
- separate visibiliity computation and query selection logic
- extract stimulus registration code
- improve slicing debug logging to show average query age
- add more profile tags

This is done to prepare some code that allow to use asynchronous queries in the Sight sense

diff --git a/Engine/Source/Runtime/AIModule/Classes/Perception/AISense_Sight.h b/Engine/Source/Runtime/AIModule/Classes/Perception/AISense_Sight.h
--- a/Engine/Source/Runtime/AIModule/Classes/Perception/AISense_Sight.h
+++ b/Engine/Source/Runtime/AIModule/Classes/Perception/AISense_Sight.h
@@ -77,7 +77,7 @@ struct FAISightQuery
 
 	FVector LastSeenLocation;
 
-	/** User data that can be used inside the IAISightTargetInterface::TestVisibilityFrom method to store a persistence state */ 
+	/** User data that can be used inside the IAISightTargetInterface::CanBeSeenFrom method to store a persistence state */ 
 	mutable int32 UserData; 
 
 	uint64 bLastResult:1;
@@ -192,7 +192,9 @@ public:
 protected:
 	virtual float Update() override;
 
+	bool ComputeVisibility(const UWorld* World, FAISightQuery& SightQuery, FPerceptionListener& Listener, const AActor* ListenerActor, FAISightTarget& Target, AActor* TargetActor, const FDigestedSightProperties& PropDigest, float& OutStimulusStrength, FVector& OutSeenLocation, int32& OutNumberOfLoSChecksPerformed) const;
 	virtual bool ShouldAutomaticallySeeTarget(const FDigestedSightProperties& PropDigest, FAISightQuery* SightQuery, FPerceptionListener& Listener, AActor* TargetActor, float& OutStimulusStrength) const;
+	void UpdateQueryVisibilityStatus(FAISightQuery& SightQuery, FPerceptionListener& Listener, const bool bIsVisible, const FVector& SeenLocation, const float StimulusStrength, AActor* TargetActor, const FVector& TargetLocation) const;
 
 	void OnNewListenerImpl(const FPerceptionListener& NewListener);
 	void OnListenerUpdateImpl(const FPerceptionListener& UpdatedListener);
diff --git a/Engine/Source/Runtime/AIModule/Private/Perception/AIPerceptionSystem.cpp b/Engine/Source/Runtime/AIModule/Private/Perception/AIPerceptionSystem.cpp
--- a/Engine/Source/Runtime/AIModule/Private/Perception/AIPerceptionSystem.cpp
+++ b/Engine/Source/Runtime/AIModule/Private/Perception/AIPerceptionSystem.cpp
@@ -16,6 +16,7 @@
 #include UE_INLINE_GENERATED_CPP_BY_NAME(AIPerceptionSystem)
 
 DECLARE_CYCLE_STAT(TEXT("Perception System"),STAT_AI_PerceptionSys,STATGROUP_AI);
+DECLARE_CYCLE_STAT(TEXT("Perception System - Process Stim"),STAT_AI_Perception_ProcessStim,STATGROUP_AI);
 
 DEFINE_LOG_CATEGORY(LogAIPerception);
 
@@ -215,19 +216,21 @@ void UAIPerceptionSystem::Tick(float DeltaSeconds)
 				}
 			}
 		}
-
-		/** no point in sorting if no new stimuli was processed */
-		bool bStimuliDelivered = DeliverDelayedStimuli(bNeedsUpdate ? RequiresSorting : NoNeedToSort);
-
-		if (bNeedsUpdate || bStimuliDelivered || bSomeListenersNeedUpdateDueToStimuliAging)
 		{
-			for (AIPerception::FListenerMap::TIterator ListenerIt(ListenerContainer); ListenerIt; ++ListenerIt)
-			{
-				check(ListenerIt->Value.Listener.IsValid());
+			SCOPE_CYCLE_COUNTER(STAT_AI_Perception_ProcessStim);
+			/** no point in sorting if no new stimuli was processed */
+			const bool bStimuliDelivered = DeliverDelayedStimuli(bNeedsUpdate ? RequiresSorting : NoNeedToSort);
 
-				if (ListenerIt->Value.HasAnyNewStimuli())
+			if (bNeedsUpdate || bStimuliDelivered || bSomeListenersNeedUpdateDueToStimuliAging)
+			{
+				for (AIPerception::FListenerMap::TIterator ListenerIt(ListenerContainer); ListenerIt; ++ListenerIt)
 				{
-					ListenerIt->Value.ProcessStimuli();
+					check(ListenerIt->Value.Listener.IsValid());
+
+					if (ListenerIt->Value.HasAnyNewStimuli())
+					{
+						ListenerIt->Value.ProcessStimuli();
+					}
 				}
 			}
 		}
diff --git a/Engine/Source/Runtime/AIModule/Private/Perception/AISense_Sight.cpp b/Engine/Source/Runtime/AIModule/Private/Perception/AISense_Sight.cpp
--- a/Engine/Source/Runtime/AIModule/Private/Perception/AISense_Sight.cpp
+++ b/Engine/Source/Runtime/AIModule/Private/Perception/AISense_Sight.cpp
@@ -14,6 +14,7 @@
 
 #include UE_INLINE_GENERATED_CPP_BY_NAME(AISense_Sight)
 
+#define AISENSE_SIGHT_TIMESLICING_DEBUG 0
 #define DO_SIGHT_VLOGGING (0 && ENABLE_VISUAL_LOG)
 
 #if DO_SIGHT_VLOGGING
@@ -26,10 +27,13 @@
 
 DECLARE_CYCLE_STAT(TEXT("Perception Sense: Sight"),STAT_AI_Sense_Sight,STATGROUP_AI);
 DECLARE_CYCLE_STAT(TEXT("Perception Sense: Sight, Update Sort"),STAT_AI_Sense_Sight_UpdateSort,STATGROUP_AI);
+DECLARE_CYCLE_STAT(TEXT("Perception Sense: Sight, Compute visibility"),STAT_AI_Sense_Sight_ComputeVisibility,STATGROUP_AI);
+DECLARE_CYCLE_STAT(TEXT("Perception Sense: Sight, Query operations"),STAT_AI_Sense_Sight_QueryOperations,STATGROUP_AI);
 DECLARE_CYCLE_STAT(TEXT("Perception Sense: Sight, Listener Update"), STAT_AI_Sense_Sight_ListenerUpdate, STATGROUP_AI);
 DECLARE_CYCLE_STAT(TEXT("Perception Sense: Sight, Register Target"), STAT_AI_Sense_Sight_RegisterTarget, STATGROUP_AI);
 DECLARE_CYCLE_STAT(TEXT("Perception Sense: Sight, Remove By Listener"), STAT_AI_Sense_Sight_RemoveByListener, STATGROUP_AI);
 DECLARE_CYCLE_STAT(TEXT("Perception Sense: Sight, Remove To Target"), STAT_AI_Sense_Sight_RemoveToTarget, STATGROUP_AI);
+DECLARE_CYCLE_STAT(TEXT("Perception Sense: Sight, Process pending result"), STAT_AI_Sense_Sight_ProcessPendingQuery, STATGROUP_AI);
 
 
 static const int32 DefaultMaxTracesPerTick = 6;
@@ -167,13 +171,66 @@ bool UAISense_Sight::ShouldAutomaticallySeeTarget(const FDigestedSightProperties
 	return false;
 }
 
+#if AISENSE_SIGHT_TIMESLICING_DEBUG
+namespace 
+{
+	struct FTimingSlicingInfo
+	{
+		FTimingSlicingInfo() 
+		{
+			Start(); 
+		}
+
+		double StartTime = 0.;
+		double EndTime = 0.;
+
+		int32 InRangeCount = 0;
+		int32 OutOfRangeCount = 0;
+
+		float InRangeAgeSum = 0.f;
+		float OutOfRangeAgeSum = 0.f;
+
+		void Start() { StartTime = FPlatformTime::Seconds();}
+		void Stop() { EndTime = FPlatformTime::Seconds();}
+
+		void PushQueryInfo(const bool bIsInRange, const float Age)
+		{
+			if (bIsInRange)
+			{
+				++InRangeCount;
+				InRangeAgeSum += Age;
+			}
+			else
+			{
+				++OutOfRangeCount;
+				OutOfRangeAgeSum += Age;
+			}
+		}
+
+		FString ToString() const
+		{
+			FString Info = FString::Format(TEXT("in {0} seconds"), {EndTime - StartTime});
+			if (InRangeCount > 0)
+			{
+				Info.Append(FString::Format(TEXT("[{0} InRange Age:{1}]"), {InRangeCount, InRangeAgeSum/InRangeCount}));
+			}
+			if (OutOfRangeCount > 0)
+			{
+				Info.Append(FString::Format(TEXT("[{0} OutOfRange Age:{1}]"), {OutOfRangeCount, OutOfRangeAgeSum/OutOfRangeCount}));
+			}
+			return Info;
+		}
+	};
+}
+#endif // AISENSE_SIGHT_TIMESLICING_DEBUG
+
 float UAISense_Sight::Update()
 {
 	SCOPE_CYCLE_COUNTER(STAT_AI_Sense_Sight);
 
 	const UWorld* World = GEngine->GetWorldFromContextObject(GetPerceptionSystem()->GetOuter(), EGetWorldErrorMode::LogAndReturnNull);
 
-	if (World == NULL)
+	if (World == nullptr)
 	{
 		return SuspendNextUpdate;
 	}
@@ -205,10 +262,8 @@ float UAISense_Sight::Update()
 	int32 NumQueriesProcessed = 0;
 	double TimeSliceEnd = FPlatformTime::Seconds() + MaxTimeSlicePerTick;
 	bool bHitTimeSliceLimit = false;
-//#define AISENSE_SIGHT_TIMESLICING_DEBUG
-#ifdef AISENSE_SIGHT_TIMESLICING_DEBUG
-	double TimeSpent = 0.0;
-	double LastTime = FPlatformTime::Seconds();
+#if AISENSE_SIGHT_TIMESLICING_DEBUG
+	FTimingSlicingInfo SlicingInfo;
 #endif // AISENSE_SIGHT_TIMESLICING_DEBUG
 	static const int32 InitialInvalidItemsSize = 16;
 	enum class EOperationType : uint8
@@ -234,6 +289,18 @@ float UAISense_Sight::Update()
 	int32 OutOfRangeItr = 0;
 	for (int32 QueryIndex = 0; QueryIndex < SightQueriesInRange.Num() + SightQueriesOutOfRange.Num(); ++QueryIndex)
 	{
+		// Time slice limit check - spread out checks to every N queries so we don't spend more time checking timer than doing work
+		NumQueriesProcessed++;
+		if ((NumQueriesProcessed % MinQueriesPerTimeSliceCheck) == 0 && FPlatformTime::Seconds() > TimeSliceEnd)
+		{
+			bHitTimeSliceLimit = true;
+		}
+
+		if (bHitTimeSliceLimit || TracesCount >= MaxTracesPerTick)
+		{
+			break;
+		}
+
 		// Calculate next in range query
 		int32 InRangeIndex = SightQueriesInRange.IsValidIndex(InRangeItr) ? InRangeItr : INDEX_NONE;
 		FAISightQuery* InRangeQuery = InRangeIndex != INDEX_NONE ? &SightQueriesInRange[InRangeIndex] : nullptr;
@@ -249,158 +316,72 @@ float UAISense_Sight::Update()
 		// Compare to real find next query
 		const bool bIsInRangeQuery = (InRangeQuery && OutOfRangeQuery) ? FAISightQuery::FSortPredicate()(*InRangeQuery,*OutOfRangeQuery) : !OutOfRangeQuery;
 		FAISightQuery* SightQuery = bIsInRangeQuery ? InRangeQuery : OutOfRangeQuery;
+		ensure(SightQuery);
 
-		// Time slice limit check - spread out checks to every N queries so we don't spend more time checking timer than doing work
-		NumQueriesProcessed++;
-#ifdef AISENSE_SIGHT_TIMESLICING_DEBUG
-		TimeSpent += (FPlatformTime::Seconds() - LastTime);
-		LastTime = FPlatformTime::Seconds();
-#endif // AISENSE_SIGHT_TIMESLICING_DEBUG
-		if (bHitTimeSliceLimit == false && (NumQueriesProcessed % MinQueriesPerTimeSliceCheck) == 0 && FPlatformTime::Seconds() > TimeSliceEnd)
-		{
-			bHitTimeSliceLimit = true;
-			// do not break here since that would bypass queue aging
-		}
+#if AISENSE_SIGHT_TIMESLICING_DEBUG
+		SlicingInfo.PushQueryInfo(bIsInRangeQuery, SightQuery->GetAge());
+#endif //AISENSE_SIGHT_TIMESLICING_DEBUG
 
-		if (TracesCount < MaxTracesPerTick && bHitTimeSliceLimit == false)
-		{
-			bIsInRangeQuery ? ++InRangeItr : ++OutOfRangeItr;
+		bIsInRangeQuery ? ++InRangeItr : ++OutOfRangeItr;
 
-			FPerceptionListener& Listener = ListenersMap[SightQuery->ObserverId];
-			FAISightTarget& Target = ObservedTargets[SightQuery->TargetId];
+		FPerceptionListener& Listener = ListenersMap[SightQuery->ObserverId];
+		FAISightTarget& Target = ObservedTargets[SightQuery->TargetId];
 
-			AActor* TargetActor = Target.Target.Get();
-			UAIPerceptionComponent* ListenerPtr = Listener.Listener.Get();
-			ensure(ListenerPtr);
+		AActor* TargetActor = Target.Target.Get();
+		UAIPerceptionComponent* ListenerPtr = Listener.Listener.Get();
+		ensure(ListenerPtr);
 
-			// @todo figure out what should we do if not valid
-			if (TargetActor && ListenerPtr)
-			{
-				const FVector TargetLocation = TargetActor->GetActorLocation();
-				const FDigestedSightProperties& PropDigest = DigestedProperties[SightQuery->ObserverId];
-				const float SightRadiusSq = SightQuery->bLastResult ? PropDigest.LoseSightRadiusSq : PropDigest.SightRadiusSq;
-				
-				float StimulusStrength = 1.f;
-				
-				// @Note that automagical "seeing" does not care about sight range nor vision cone
-				const bool bShouldAutomatically = ShouldAutomaticallySeeTarget(PropDigest, SightQuery, Listener, TargetActor, StimulusStrength);
-				if (bShouldAutomatically)
-				{
-					// Pretend like we've seen this target where we last saw them
-					Listener.RegisterStimulus(TargetActor, FAIStimulus(*this, StimulusStrength, SightQuery->LastSeenLocation, Listener.CachedLocation));
-					SightQuery->bLastResult = true;
-				}
-				else if (FAISystem::CheckIsTargetInSightCone(Listener.CachedLocation, Listener.CachedDirection, PropDigest.PeripheralVisionAngleCos, PropDigest.PointOfViewBackwardOffset, PropDigest.NearClippingRadiusSq, SightRadiusSq, TargetLocation))
-				{
-					SIGHT_LOG_SEGMENT(ListenerPtr->GetOwner(), Listener.CachedLocation, TargetLocation, FColor::Green, TEXT("TargetID %d"), Target.TargetId);
+		// @todo figure out what should we do if not valid
+		if (TargetActor && ListenerPtr)
+		{
+			const FDigestedSightProperties& PropDigest = DigestedProperties[SightQuery->ObserverId];
+			const AActor* ListenerBodyActor = ListenerPtr->GetBodyActor();
+			float StimulusStrength = 1.f;
+			FVector SeenLocation(0.f);
+			int32 NumberOfLoSChecksPerformed = 0;
 
-					FVector OutSeenLocation(0.f);
-					// do line checks
-					if (Target.SightTargetInterface != NULL)
-					{
-						int32 NumberOfLoSChecksPerformed = 0;
-						// defaulting to 1 to have "full strength" by default instead of "no strength"
-						const bool bWasVisible = SightQuery->bLastResult;
-						if (Target.SightTargetInterface->CanBeSeenFrom(Listener.CachedLocation, OutSeenLocation, NumberOfLoSChecksPerformed, StimulusStrength, ListenerPtr->GetBodyActor(), &bWasVisible, &SightQuery->UserData) == true)
-						{
-							Listener.RegisterStimulus(TargetActor, FAIStimulus(*this, StimulusStrength, OutSeenLocation, Listener.CachedLocation));
-							SightQuery->bLastResult = true;
-							SightQuery->LastSeenLocation = OutSeenLocation;
-						}
-						// communicate failure only if we've seen give actor before
-						else if (SightQuery->bLastResult == true)
-						{
-							Listener.RegisterStimulus(TargetActor, FAIStimulus(*this, 0.f, TargetLocation, Listener.CachedLocation, FAIStimulus::SensingFailed));
-							SightQuery->bLastResult = false;
-							SightQuery->LastSeenLocation = FAISystem::InvalidLocation;
-						}
-
-						if (SightQuery->bLastResult == false)
-						{
-							SIGHT_LOG_LOCATION(ListenerPtr->GetOwner(), TargetLocation, 25.f, FColor::Red, TEXT(""));
-						}
-
-						TracesCount += NumberOfLoSChecksPerformed;
-					}
-					else
-					{
-						// we need to do tests ourselves
-						FHitResult HitResult;
-						const bool bHit = World->LineTraceSingleByChannel(HitResult, Listener.CachedLocation, TargetLocation
-							, DefaultSightCollisionChannel
-							, FCollisionQueryParams(SCENE_QUERY_STAT(AILineOfSight), true, ListenerPtr->GetBodyActor()));
-
-						++TracesCount;
-
-						auto HitResultActorIsOwnedByTargetActor = [&HitResult, TargetActor]()
-						{
-							AActor* HitResultActor = HitResult.HitObjectHandle.FetchActor();
-							return (HitResultActor ? HitResultActor->IsOwnedBy(TargetActor) : false);
-						};
-
-						if (bHit == false || HitResultActorIsOwnedByTargetActor())
-						{
-							Listener.RegisterStimulus(TargetActor, FAIStimulus(*this, 1.f, TargetLocation, Listener.CachedLocation));
-							SightQuery->bLastResult = true;
-							SightQuery->LastSeenLocation = TargetLocation;
-						}
-						// communicate failure only if we've seen give actor before
-						else if (SightQuery->bLastResult == true)
-						{
-							Listener.RegisterStimulus(TargetActor, FAIStimulus(*this, 0.f, TargetLocation, Listener.CachedLocation, FAIStimulus::SensingFailed));
-							SightQuery->bLastResult = false;
-							SightQuery->LastSeenLocation = FAISystem::InvalidLocation;
-						}
-
-						if (SightQuery->bLastResult == false)
-						{
-							SIGHT_LOG_LOCATION(ListenerPtr->GetOwner(), TargetLocation, 25.f, FColor::Red, TEXT(""));
-						}
-					}
-				}
-				// communicate failure only if we've seen give actor before
-				else if (SightQuery->bLastResult)
-				{
-					SIGHT_LOG_SEGMENT(ListenerPtr->GetOwner(), Listener.CachedLocation, TargetLocation, FColor::Red, TEXT("TargetID %d"), Target.TargetId);
-					Listener.RegisterStimulus(TargetActor, FAIStimulus(*this, 0.f, TargetLocation, Listener.CachedLocation, FAIStimulus::SensingFailed));
-					SightQuery->bLastResult = false;
-				}
+			const bool bIsVisible = ComputeVisibility(World, *SightQuery, Listener, ListenerBodyActor, Target, TargetActor, PropDigest, StimulusStrength, SeenLocation, NumberOfLoSChecksPerformed);
 
-				SightQuery->Importance = CalcQueryImportance(Listener, TargetLocation, SightRadiusSq);
-				const bool bShouldBeInRange = SightQuery->Importance > 0.0f;
-				if (bIsInRangeQuery != bShouldBeInRange)
-				{
-					QueryOperations.Add(FQueryOperation(bIsInRangeQuery, EOperationType::SwapList, bIsInRangeQuery ? InRangeIndex : OutOfRangeIndex));
-				}
+			TracesCount += NumberOfLoSChecksPerformed;
 
-				// restart query
-				SightQuery->OnProcessed();
-			}
-			else
+			const bool bWasVisible = SightQuery->bLastResult;
+			const FVector TargetLocation = TargetActor->GetActorLocation();
+			UpdateQueryVisibilityStatus(*SightQuery, Listener, bIsVisible, SeenLocation, StimulusStrength, TargetActor, TargetLocation);
+
+			const float SightRadiusSq = bWasVisible ? PropDigest.LoseSightRadiusSq : PropDigest.SightRadiusSq;
+			SightQuery->Importance = CalcQueryImportance(Listener, TargetLocation, SightRadiusSq);
+			const bool bShouldBeInRange = SightQuery->Importance > 0.0f;
+			if (bIsInRangeQuery != bShouldBeInRange)
 			{
-				// put this index to "to be removed" array
-				QueryOperations.Add( FQueryOperation(bIsInRangeQuery, EOperationType::Remove, bIsInRangeQuery ? InRangeIndex : OutOfRangeIndex) );
-				if (TargetActor == nullptr)
-				{
-					InvalidTargets.AddUnique(SightQuery->TargetId);
-				}
+				QueryOperations.Add(FQueryOperation(bIsInRangeQuery, EOperationType::SwapList, bIsInRangeQuery ? InRangeIndex : OutOfRangeIndex));
 			}
+
+			// restart query
+			SightQuery->OnProcessed();
 		}
 		else
 		{
-			break;
+			// put this index to "to be removed" array
+			QueryOperations.Add( FQueryOperation(bIsInRangeQuery, EOperationType::Remove, bIsInRangeQuery ? InRangeIndex : OutOfRangeIndex) );
+			if (TargetActor == nullptr)
+			{
+				InvalidTargets.AddUnique(SightQuery->TargetId);
+			}
 		}
 	}
 	NextOutOfRangeIndex = SightQueriesOutOfRange.Num() > 0 ? (NextOutOfRangeIndex + OutOfRangeItr) % SightQueriesOutOfRange.Num() : 0;
 
-#ifdef AISENSE_SIGHT_TIMESLICING_DEBUG
-	UE_LOG(LogAIPerception, VeryVerbose, TEXT("UAISense_Sight::Update processed %d sources in %f seconds [time slice limited? %d]"), NumQueriesProcessed, TimeSpent, bHitTimeSliceLimit ? 1 : 0);
+#if AISENSE_SIGHT_TIMESLICING_DEBUG
+	SlicingInfo.Stop();
+	UE_LOG(LogAIPerception, VeryVerbose, TEXT("UAISense_Sight::Update processed %d sources %s [time slice limited? %d]"), NumQueriesProcessed, *SlicingInfo.ToString(), bHitTimeSliceLimit ? 1 : 0);
 #else
 	UE_LOG(LogAIPerception, VeryVerbose, TEXT("UAISense_Sight::Update processed %d sources [time slice limited? %d]"), NumQueriesProcessed, bHitTimeSliceLimit ? 1 : 0);
 #endif // AISENSE_SIGHT_TIMESLICING_DEBUG
 
 	if (QueryOperations.Num() > 0)
 	{
+		SCOPE_CYCLE_COUNTER(STAT_AI_Sense_Sight_QueryOperations);
+
 		// Sort by InRange and by descending Index 
 		QueryOperations.Sort([](const FQueryOperation& LHS, const FQueryOperation& RHS)->bool
 		{
@@ -468,6 +449,79 @@ float UAISense_Sight::Update()
 	return 0.f;
 }
 
+bool UAISense_Sight::ComputeVisibility(const UWorld* World, FAISightQuery& SightQuery, FPerceptionListener& Listener, const AActor* ListenerActor, FAISightTarget& Target, AActor* TargetActor, const FDigestedSightProperties& PropDigest, float& OutStimulusStrength, FVector& OutSeenLocation, int32& OutNumberOfLoSChecksPerformed) const
+{
+	SCOPE_CYCLE_COUNTER(STAT_AI_Sense_Sight_ComputeVisibility);
+
+	// @Note that automagical "seeing" does not care about sight range nor vision cone
+	if (ShouldAutomaticallySeeTarget(PropDigest, &SightQuery, Listener, TargetActor, OutStimulusStrength))
+	{
+		OutSeenLocation = FAISystem::InvalidLocation;
+		return true;
+	}
+
+	const FVector TargetLocation = TargetActor->GetActorLocation();
+	const float SightRadiusSq = SightQuery.bLastResult ? PropDigest.LoseSightRadiusSq : PropDigest.SightRadiusSq;
+	if (!FAISystem::CheckIsTargetInSightCone(Listener.CachedLocation, Listener.CachedDirection, PropDigest.PeripheralVisionAngleCos, PropDigest.PointOfViewBackwardOffset, PropDigest.NearClippingRadiusSq, SightRadiusSq, TargetLocation))
+	{
+		return false;
+	}
+
+	if (Target.SightTargetInterface != nullptr)
+	{
+		const bool bWasVisible = SightQuery.bLastResult;
+		const bool bCanBeSeen = Target.SightTargetInterface->CanBeSeenFrom(Listener.CachedLocation, OutSeenLocation, OutNumberOfLoSChecksPerformed, OutStimulusStrength, ListenerActor, &bWasVisible, &SightQuery.UserData);
+		return bCanBeSeen;
+	}
+	else
+	{
+		// we need to do tests ourselves
+		FHitResult HitResult;
+		const bool bHit = World->LineTraceSingleByChannel(HitResult, Listener.CachedLocation, TargetLocation, DefaultSightCollisionChannel, FCollisionQueryParams(SCENE_QUERY_STAT(AILineOfSight), true, ListenerActor));
+
+		++OutNumberOfLoSChecksPerformed;
+
+		auto HitResultActorIsOwnedByTargetActor = [&HitResult, TargetActor]()
+		{
+			AActor* HitResultActor = HitResult.HitObjectHandle.FetchActor();
+			return (HitResultActor ? HitResultActor->IsOwnedBy(TargetActor) : false);
+		};
+
+		if (bHit == false || HitResultActorIsOwnedByTargetActor())
+		{
+			OutSeenLocation = TargetLocation;
+			return true;
+		}
+		else
+		{
+			return false;
+		}
+	}
+}
+
+void UAISense_Sight::UpdateQueryVisibilityStatus(FAISightQuery& SightQuery, FPerceptionListener& Listener, const bool bIsVisible, const FVector& SeenLocation, const float StimulusStrength, AActor* TargetActor, const FVector& TargetLocation) const
+{
+	if (bIsVisible)
+	{
+		const bool bHasValidSeenLocation = SeenLocation != FAISystem::InvalidLocation;
+		Listener.RegisterStimulus(TargetActor, FAIStimulus(*this, StimulusStrength, bHasValidSeenLocation ? SeenLocation : SightQuery.LastSeenLocation, Listener.CachedLocation));
+		SightQuery.bLastResult = true;
+		if (bHasValidSeenLocation)
+		{
+			SightQuery.LastSeenLocation = SeenLocation;
+		}
+	}
+	// communicate failure only if we've seen given actor before
+	else if (SightQuery.bLastResult == true)
+	{
+		Listener.RegisterStimulus(TargetActor, FAIStimulus(*this, 0.f, TargetLocation, Listener.CachedLocation, FAIStimulus::SensingFailed));
+		SightQuery.bLastResult = false;
+		SightQuery.LastSeenLocation = FAISystem::InvalidLocation;
+	}
+
+	SIGHT_LOG_SEGMENT(ListenerPtr->GetOwner(), Listener.CachedLocation, TargetLocation, bIsVisible ? FColor::Green : FColor::Red, TEXT("TargetID %d"), Target.TargetId);
+}
+
 void UAISense_Sight::RegisterEvent(const FAISightEvent& Event)
 {
 
