Run behavior task: make it explicit when the sub tree is not compatible
#tests PIE
#OKforGitHub public

diff --git a/Engine/Source/Runtime/AIModule/Private/BehaviorTree/Tasks/BTTask_RunBehavior.cpp b/Engine/Source/Runtime/AIModule/Private/BehaviorTree/Tasks/BTTask_RunBehavior.cpp
--- a/Engine/Source/Runtime/AIModule/Private/BehaviorTree/Tasks/BTTask_RunBehavior.cpp
+++ b/Engine/Source/Runtime/AIModule/Private/BehaviorTree/Tasks/BTTask_RunBehavior.cpp
@@ -1,6 +1,7 @@
 // Copyright Epic Games, Inc. All Rights Reserved.
 
 #include "BehaviorTree/Tasks/BTTask_RunBehavior.h"
+#include "BehaviorTree/BlackboardData.h"
 #include "VisualLogger/VisualLogger.h"
 #include "AIController.h"
 
@@ -41,7 +42,16 @@ void UBTTask_RunBehavior::OnSubtreeDeactivated(UBehaviorTreeComponent& OwnerComp
 
 FString UBTTask_RunBehavior::GetStaticDescription() const
 {
-	return FString::Printf(TEXT("%s: %s"), *Super::GetStaticDescription(), *GetNameSafe(BehaviorAsset));
+	bool bIsBBCompatible = false;
+	if (const UBlackboardData* BlackboardData = GetBlackboardAsset())
+	{
+		if (const UBlackboardData* OtherBlackboardData = BehaviorAsset ? BehaviorAsset->GetBlackboardAsset() : nullptr )
+		{
+			bIsBBCompatible = BlackboardData == OtherBlackboardData || BlackboardData->IsChildOf(*OtherBlackboardData);
+		}
+	}
+		
+	return FString::Printf(TEXT("%s: %s%s"), *Super::GetStaticDescription(), *GetNameSafe(BehaviorAsset), bIsBBCompatible ? TEXT("") : TEXT(" (Blackboard not compatible)"));
 }
 
 #if WITH_EDITOR
