Fix bug where decorator with LowerPrio abortion type on a subtree node was not reactivated upon exiting the subtree
UDN case 00484061

diff --git a/Engine/Source/Developer/AITestSuite/Private/Tests/BTTest.cpp b/Engine/Source/Developer/AITestSuite/Private/Tests/BTTest.cpp
--- a/Engine/Source/Developer/AITestSuite/Private/Tests/BTTest.cpp
+++ b/Engine/Source/Developer/AITestSuite/Private/Tests/BTTest.cpp
@@ -2109,6 +2109,58 @@ struct FAITest_BTSubtreeAbortOutToLowerPrio : public FAITest_SimpleBT
 };
 IMPLEMENT_AI_LATENT_TEST(FAITest_BTSubtreeAbortOutToLowerPrio, "System.AI.Behavior Trees.Subtree: abort out lower prio")
 
+
+/* This unittest come from UDN case 00484061 */
+struct FAITest_BTSubtreeAbortOutToLowerPrio2 : public FAITest_SimpleBT
+{
+	FAITest_BTSubtreeAbortOutToLowerPrio2()
+	{
+		enum
+		{
+			SubtreeTaskExecute = 1,
+			SubtreeTaskTick,
+			BBDecoratorBecomeRelevant,
+			BBDecoratorCeaseRelevant,
+			MainTaskExecute,
+			MainTaskTick
+		};
+
+		// SubTree
+		UBehaviorTree* ChildAsset = &FBTBuilder::CreateBehaviorTree(*BTAsset);
+		if (ChildAsset)
+		{
+			AddAutoDestroyObject(*ChildAsset);
+			UBTCompositeNode& CompNode = FBTBuilder::AddSequence(*ChildAsset);
+			{
+				// SubtreeTask
+				FBTBuilder::AddTask(CompNode, SubtreeTaskExecute, EBTNodeResult::Failed, 1, SubtreeTaskTick);
+			}
+		}
+
+		UBTCompositeNode& CompNode = FBTBuilder::AddSelector(*BTAsset);
+		{
+			// SubTreeTask
+			FBTBuilder::AddTaskSubtree(CompNode, ChildAsset);
+
+			// BBDecorator
+			FBTBuilder::WithDecoratorBlackboard(CompNode, EBasicKeyOperation::NotSet, EBTFlowAbortMode::LowerPriority, TEXT("Bool1"), BBDecoratorBecomeRelevant, BBDecoratorCeaseRelevant);
+
+			// MainTask
+			FBTBuilder::AddTask(CompNode, MainTaskExecute, EBTNodeResult::Succeeded, 1, MainTaskTick);
+		}
+
+		ExpectedResult.Add(SubtreeTaskExecute);
+		ExpectedResult.Add(SubtreeTaskTick);
+		ExpectedResult.Add(SubtreeTaskTick);
+		ExpectedResult.Add(BBDecoratorBecomeRelevant);
+		ExpectedResult.Add(MainTaskExecute);
+		ExpectedResult.Add(MainTaskTick);
+		ExpectedResult.Add(MainTaskTick);
+		ExpectedResult.Add(BBDecoratorCeaseRelevant);
+	}
+};
+IMPLEMENT_AI_LATENT_TEST(FAITest_BTSubtreeAbortOutToLowerPrio2, "System.AI.Behavior Trees.Subtree: abort out lower prio 2")
+
 struct FAITest_BTServiceInstantTask : public FAITest_SimpleBT
 {
 	FAITest_BTServiceInstantTask()
diff --git a/Engine/Source/Runtime/AIModule/Classes/BehaviorTree/BTCompositeNode.h b/Engine/Source/Runtime/AIModule/Classes/BehaviorTree/BTCompositeNode.h
--- a/Engine/Source/Runtime/AIModule/Classes/BehaviorTree/BTCompositeNode.h
+++ b/Engine/Source/Runtime/AIModule/Classes/BehaviorTree/BTCompositeNode.h
@@ -117,9 +117,22 @@ class AIMODULE_API UBTCompositeNode : public UBTNode
 	void OnChildActivation(FBehaviorTreeSearchData& SearchData, const UBTNode& ChildNode) const;
 	void OnChildActivation(FBehaviorTreeSearchData& SearchData, int32 ChildIndex) const;
 
-	/** called after child has finished search */
-	void OnChildDeactivation(FBehaviorTreeSearchData& SearchData, const UBTNode& ChildNode, EBTNodeResult::Type& NodeResult, const bool bIsRequestFromActiveInstance) const;
-	void OnChildDeactivation(FBehaviorTreeSearchData& SearchData, int32 ChildIndex, EBTNodeResult::Type& NodeResult, const bool bIsRequestFromActiveInstance) const;
+	/**
+	 * Notification called after child has finished search
+	 * @param SearchData for any new addition or removal of extra aux nodes/ descriptor
+	 * @param ChildNode is the children being deactivated
+	 * @param NodeResult the raison of the deactivation
+	 * @param bRequestedFromValidInstance the new requested search start is within the current active instances
+	 */
+	void OnChildDeactivation(FBehaviorTreeSearchData& SearchData, const UBTNode& ChildNode, EBTNodeResult::Type& NodeResult, const bool bRequestedFromValidInstance) const;
+	/**
+	 * Notification called after child has finished search
+	 * @param SearchData for any new addition or removal of extra aux nodes/ descriptor
+	 * @param ChildIndex of the child node being deactivated
+	 * @param NodeResult the raison of the deactivation
+	 * @param bRequestedFromValidInstance the new request search start is within the current active instances
+	 */
+	void OnChildDeactivation(FBehaviorTreeSearchData& SearchData, int32 ChildIndex, EBTNodeResult::Type& NodeResult, const bool bRequestedFromValidInstance) const;
 
 	/** called when start enters this node */
 	void OnNodeActivation(FBehaviorTreeSearchData& SearchData) const;
@@ -176,9 +189,9 @@ class AIMODULE_API UBTCompositeNode : public UBTNode
 
 	// Deprecated methods
 	UE_DEPRECATED(5.0, "This function is deprecated. Please use RequestBranchDeactivation instead.")
-	void OnChildDeactivation(FBehaviorTreeSearchData& SearchData, const UBTNode& ChildNode, EBTNodeResult::Type& NodeResult) const { OnChildDeactivation(SearchData, ChildNode, NodeResult, true/*bIsRequestFromActiveInstance*/); }
+	void OnChildDeactivation(FBehaviorTreeSearchData& SearchData, const UBTNode& ChildNode, EBTNodeResult::Type& NodeResult) const { OnChildDeactivation(SearchData, ChildNode, NodeResult, /*bRequestedFromValidInstance*/true); }
 	UE_DEPRECATED(5.0, "This function is deprecated. Please use RequestBranchDeactivation instead.")
-	void OnChildDeactivation(FBehaviorTreeSearchData& SearchData, int32 ChildIndex, EBTNodeResult::Type& NodeResult) const { OnChildDeactivation(SearchData, ChildIndex, NodeResult, true /*bIsRequestFromActiveInstance*/); }
+	void OnChildDeactivation(FBehaviorTreeSearchData& SearchData, int32 ChildIndex, EBTNodeResult::Type& NodeResult) const { OnChildDeactivation(SearchData, ChildIndex, NodeResult, /*bRequestedFromValidInstance*/true ); }
 
 protected:
 
diff --git a/Engine/Source/Runtime/AIModule/Private/BehaviorTree/BTCompositeNode.cpp b/Engine/Source/Runtime/AIModule/Private/BehaviorTree/BTCompositeNode.cpp
--- a/Engine/Source/Runtime/AIModule/Private/BehaviorTree/BTCompositeNode.cpp
+++ b/Engine/Source/Runtime/AIModule/Private/BehaviorTree/BTCompositeNode.cpp
@@ -123,12 +123,12 @@ void UBTCompositeNode::OnChildActivation(FBehaviorTreeSearchData& SearchData, in
 	NodeMemory->CurrentChild = ChildIndex;
 }
 
-void UBTCompositeNode::OnChildDeactivation(FBehaviorTreeSearchData& SearchData, const UBTNode& ChildNode, EBTNodeResult::Type& NodeResult, const bool bIsRequestFromActiveInstance) const
+void UBTCompositeNode::OnChildDeactivation(FBehaviorTreeSearchData& SearchData, const UBTNode& ChildNode, EBTNodeResult::Type& NodeResult, const bool bRequestedFromValidInstance) const
 {
-	OnChildDeactivation(SearchData, GetChildIndex(SearchData, ChildNode), NodeResult, bIsRequestFromActiveInstance);
+	OnChildDeactivation(SearchData, GetChildIndex(SearchData, ChildNode), NodeResult, bRequestedFromValidInstance);
 }
 
-void UBTCompositeNode::OnChildDeactivation(FBehaviorTreeSearchData& SearchData, int32 ChildIndex, EBTNodeResult::Type& NodeResult, const bool bIsRequestFromActiveInstance) const
+void UBTCompositeNode::OnChildDeactivation(FBehaviorTreeSearchData& SearchData, int32 ChildIndex, EBTNodeResult::Type& NodeResult, const bool bRequestedFromValidInstance) const
 {
 	const FBTCompositeChild& ChildInfo = Children[ChildIndex];
 
@@ -151,7 +151,7 @@ void UBTCompositeNode::OnChildDeactivation(FBehaviorTreeSearchData& SearchData,
 	const bool bCanNotify = !bUseDecoratorsDeactivationCheck || CanNotifyDecoratorsOnDeactivation(SearchData, ChildIndex, NodeResult);
 	if (bCanNotify)
 	{
-		NotifyDecoratorsOnDeactivation(SearchData, ChildIndex, NodeResult, bIsRequestFromActiveInstance);
+		NotifyDecoratorsOnDeactivation(SearchData, ChildIndex, NodeResult, bRequestedFromValidInstance);
 	}
 }
 
@@ -249,7 +249,7 @@ void UBTCompositeNode::NotifyDecoratorsOnActivation(FBehaviorTreeSearchData& Sea
 	}
 }
 
-void UBTCompositeNode::NotifyDecoratorsOnDeactivation(FBehaviorTreeSearchData& SearchData, int32 ChildIdx, EBTNodeResult::Type& NodeResult, const bool bIsRequestFromActiveInstance) const
+void UBTCompositeNode::NotifyDecoratorsOnDeactivation(FBehaviorTreeSearchData& SearchData, int32 ChildIdx, EBTNodeResult::Type& NodeResult, const bool bRequestedFromValidInstance) const
 {
 	const FBTCompositeChild& ChildInfo = Children[ChildIdx];
 	if (NodeResult == EBTNodeResult::Aborted)
@@ -272,7 +272,7 @@ void UBTCompositeNode::NotifyDecoratorsOnDeactivation(FBehaviorTreeSearchData& S
 			DecoratorOb->WrappedOnNodeProcessed(SearchData, NodeResult);
 			DecoratorOb->WrappedOnNodeDeactivation(SearchData, NodeResult);
 
-			if (!bIsRequestFromActiveInstance)
+			if (!bRequestedFromValidInstance)
 			{
 				UE_VLOG(SearchData.OwnerComp.GetOwner(), LogBehaviorTree, Verbose, TEXT("Removing decorator(%s) as request is not in active instance"), *UBehaviorTreeTypes::DescribeNodeHelper(this));
 				SearchData.AddUniqueUpdate(FBehaviorTreeSearchUpdate(DecoratorOb, SearchData.OwnerComp.GetActiveInstanceIdx(), EBTNodeUpdateMode::Remove));
diff --git a/Engine/Source/Runtime/AIModule/Private/BehaviorTree/BehaviorTreeComponent.cpp b/Engine/Source/Runtime/AIModule/Private/BehaviorTree/BehaviorTreeComponent.cpp
--- a/Engine/Source/Runtime/AIModule/Private/BehaviorTree/BehaviorTreeComponent.cpp
+++ b/Engine/Source/Runtime/AIModule/Private/BehaviorTree/BehaviorTreeComponent.cpp
@@ -1955,14 +1955,15 @@ void UBehaviorTreeComponent::ProcessExecutionRequest()
 			}
 			else if (ChildBranchIdx == BTSpecialChild::ReturnToParent)
 			{
-				const UBTCompositeNode* ChildNode = TestNode;
+				const UBTNode* ChildNode = TestNode;
+				const UBTCompositeNode* PrevTestNode = TestNode;
 				TestNode = TestNode->GetParentNode();
 
 				// does it want to move up the tree?
 				if (TestNode == NULL)
 				{
 					// special case for leaving instance: deactivate root manually
-					ChildNode->OnNodeDeactivation(SearchData, NodeResult);
+					PrevTestNode->OnNodeDeactivation(SearchData, NodeResult);
 
 					// don't remove top instance from stack, so it could be looped
 					if (ActiveInstanceIdx > 0)
@@ -1978,13 +1979,15 @@ void UBehaviorTreeComponent::ProcessExecutionRequest()
 						ActiveInstanceIdx--;
 
 						StoreDebuggerSearchStep(InstanceStack[ActiveInstanceIdx].ActiveNode, ActiveInstanceIdx, NodeResult);
-						TestNode = InstanceStack[ActiveInstanceIdx].ActiveNode->GetParentNode();
+						ChildNode = InstanceStack[ActiveInstanceIdx].ActiveNode;
+						TestNode = ChildNode->GetParentNode();
 					}
 				}
 
 				if (TestNode)
 				{
-					TestNode->OnChildDeactivation(SearchData, *ChildNode, NodeResult, ActiveInstanceIdx == ExecutionRequest.ExecuteInstanceIdx  /*bIsRequestInSameInstance*/);
+					const bool bRequestedFromValidInstance = ActiveInstanceIdx <= ExecutionRequest.ExecuteInstanceIdx;
+					TestNode->OnChildDeactivation(SearchData, *ChildNode, NodeResult, bRequestedFromValidInstance);
 				}
 			}
 			else if (TestNode->Children.IsValidIndex(ChildBranchIdx))
@@ -2183,7 +2186,8 @@ bool UBehaviorTreeComponent::DeactivateUpTo(const UBTCompositeNode* Node, uint16
 		if (NotifyParent)
 		{
 			OutLastDeactivatedChildIndex = NotifyParent->GetChildIndex(SearchData, *DeactivatedChild);
-			NotifyParent->OnChildDeactivation(SearchData, OutLastDeactivatedChildIndex, NodeResult, ActiveInstanceIdx == NodeInstanceIdx /*bIsRequestInSameInstance*/);
+			const bool bRequestedFromValidInstance = ActiveInstanceIdx <= NodeInstanceIdx;
+			NotifyParent->OnChildDeactivation(SearchData, OutLastDeactivatedChildIndex, NodeResult, bRequestedFromValidInstance);
 
 			BT_SEARCHLOG(SearchData, Verbose, TEXT("Deactivate node: %s"), *UBehaviorTreeTypes::DescribeNodeHelper(DeactivatedChild));
 			StoreDebuggerSearchStep(DeactivatedChild, ActiveInstanceIdx, NodeResult);
