Support more than 255 entries in a blackboard
#OKForGitHub public

diff --git a/Engine/Source/Runtime/AIModule/Classes/BehaviorTree/BehaviorTreeTypes.h b/Engine/Source/Runtime/AIModule/Classes/BehaviorTree/BehaviorTreeTypes.h
--- a/Engine/Source/Runtime/AIModule/Classes/BehaviorTree/BehaviorTreeTypes.h
+++ b/Engine/Source/Runtime/AIModule/Classes/BehaviorTree/BehaviorTreeTypes.h
@@ -47,9 +47,15 @@ namespace FBlackboard
 {
 	const FName KeySelf = TEXT("SelfActor");
 
-	typedef uint8 FKey;
+#ifdef AI_BLACKBOARD_KEY_SIZE_8
+	// this is the legacy BB key size. Add AI_BLACKBOARD_KEY_SIZE_8 to your *.target.cs to enable it.
+	using FKey = uint8;
+#else
+	// this is the default BB key size
+	using FKey = uint16;
+#endif
 
-	const FKey InvalidKey = FKey(-1);
+	constexpr FKey InvalidKey = FKey(-1);
 }
 
 enum class EBlackboardNotificationResult : uint8
@@ -64,10 +70,10 @@ DECLARE_DELEGATE_RetVal_TwoParams(EBlackboardNotificationResult, FOnBlackboardCh
 
 namespace BTSpecialChild
 {
-	const int32 NotInitialized = -1;	// special value for child indices: needs to be initialized
-	const int32 ReturnToParent = -2;	// special value for child indices: return to parent node
+	constexpr int32 NotInitialized = -1;	// special value for child indices: needs to be initialized
+	constexpr int32 ReturnToParent = -2;	// special value for child indices: return to parent node
 	
-	const uint8 OwnedByComposite = MAX_uint8;	// special value for aux node's child index: owned by composite node instead of a task
+	constexpr uint8 OwnedByComposite = MAX_uint8;	// special value for aux node's child index: owned by composite node instead of a task
 }
 
 UENUM(BlueprintType)
@@ -180,7 +186,7 @@ struct FBehaviorTreeParallelTask
 	/** additional mode data used for context switching */
 	EBTTaskStatus::Type Status;
 
-	FBehaviorTreeParallelTask() : TaskNode(NULL) {}
+	FBehaviorTreeParallelTask() : TaskNode(nullptr) {}
 	FBehaviorTreeParallelTask(const UBTTaskNode* InTaskNode, EBTTaskStatus::Type InStatus) : TaskNode(InTaskNode), Status(InStatus) {}
 
 	bool operator==(const FBehaviorTreeParallelTask& Other) const { return TaskNode == Other.TaskNode; }
@@ -225,7 +231,7 @@ struct FBehaviorTreeDebuggerInstance
 		FNodeFlowData() : ExecutionIndex(INDEX_NONE), bPassed(0), bTrigger(0), bDiscardedTrigger(0) {}
 	};
 
-	FBehaviorTreeDebuggerInstance() : TreeAsset(NULL), RootNode(NULL) {}
+	FBehaviorTreeDebuggerInstance() : TreeAsset(nullptr), RootNode(nullptr) {}
 
 	/** behavior tree asset */
 	UBehaviorTree* TreeAsset;
@@ -603,10 +609,8 @@ struct AIMODULE_API FBlackboardKeySelector
 
 protected:
 	/** ID of selected key */
-	UPROPERTY(transient, EditInstanceOnly, BlueprintReadWrite, Category = Blackboard)
-	uint8 SelectedKeyID;
-	// SelectedKeyId type should be FBlackboard::FKey, but typedefs are not supported by UHT
-	static_assert(sizeof(uint8) == sizeof(FBlackboard::FKey), "FBlackboardKeySelector::SelectedKeyId should be of FBlackboard::FKey-compatible type.");
+	UPROPERTY(transient, EditInstanceOnly, BlueprintReadWrite, Category = Blackboard, meta = (ClampMin = "0", UIMin = "0"))
+	int32 SelectedKeyID;
 
 	// Requires BlueprintReadWrite so that blueprint creators (using MakeBlackboardKeySelector) can specify whether or not None is Allowed.
 	UPROPERTY(EditAnywhere, BlueprintReadWrite, Category = Blackboard, Meta = (Tooltip = ""))
diff --git a/Engine/Source/Runtime/AIModule/Private/BehaviorTree/BlackboardComponent.cpp b/Engine/Source/Runtime/AIModule/Private/BehaviorTree/BlackboardComponent.cpp
--- a/Engine/Source/Runtime/AIModule/Private/BehaviorTree/BlackboardComponent.cpp
+++ b/Engine/Source/Runtime/AIModule/Private/BehaviorTree/BlackboardComponent.cpp
@@ -178,7 +178,10 @@ bool UBlackboardComponent::InitializeBlackboard(UBlackboardData& NewAsset)
 		for (int32 Index = 0; Index < InitList.Num(); Index++)
 		{
 			const FBlackboardEntry* KeyData = BlackboardAsset->GetKey(InitList[Index].KeyID);
-			KeyData->KeyType->InitializeKey(*this, InitList[Index].KeyID);
+			if (ensureMsgf(KeyData && KeyData->KeyType, TEXT("Too many elements in blackboard %s (%u)"),*GetNameSafe(BlackboardAsset), InitList.Num()))
+			{
+				KeyData->KeyType->InitializeKey(*this, InitList[Index].KeyID);
+			}
 		}
 
 		// naive initial synchronization with one of already instantiated blackboards using the same BB asset
