Provide an optionnal way in AISense_Sight to perform asynchronous collision tests :
- add a new method in UAISightTargetInterface that support asynchronous queries
- add new queries list in AISense_Sight to store pending queries

diff --git a/Engine/Source/Runtime/AIModule/Classes/Perception/AISense_Sight.h b/Engine/Source/Runtime/AIModule/Classes/Perception/AISense_Sight.h
--- a/Engine/Source/Runtime/AIModule/Classes/Perception/AISense_Sight.h
+++ b/Engine/Source/Runtime/AIModule/Classes/Perception/AISense_Sight.h
@@ -5,6 +5,7 @@
 #include "CoreMinimal.h"
 #include "UObject/ObjectMacros.h"
 #include "GenericTeamAgentInterface.h"
+#include "Misc/MTAccessDetector.h"
 #include "Perception/AISense.h"
 #include "AISense_Sight.generated.h"
 
@@ -46,7 +47,7 @@ struct AIMODULE_API FAISightEvent
 	}
 };
 
-struct FAISightTarget
+struct AIMODULE_API FAISightTarget
 {
 	typedef uint32 FTargetId;
 	static const FTargetId InvalidTargetId;
@@ -122,6 +123,24 @@ struct FAISightQuery
 	};
 };
 
+struct AIMODULE_API FAISightQueryID
+{
+	FPerceptionListenerID ObserverId;
+	FAISightTarget::FTargetId TargetId;
+
+	FAISightQueryID(FPerceptionListenerID ListenerId = FPerceptionListenerID::InvalidID(), FAISightTarget::FTargetId Target = FAISightTarget::InvalidTargetId)
+	: ObserverId(ListenerId), TargetId(Target)
+	{
+	}
+
+	FAISightQueryID(const FAISightQuery& Query)
+	: ObserverId(Query.ObserverId), TargetId(Query.TargetId)
+	{
+	}
+};
+
+DECLARE_DELEGATE_FiveParams(FOnPendingVisibilityQueryProcessedDelegate, const FAISightQueryID&, const bool, const float, const FVector&, const TOptional<int32>&);
+
 UCLASS(ClassGroup=AI, config=Game)
 class AIMODULE_API UAISense_Sight : public UAISense
 {
@@ -140,8 +159,15 @@ public:
 
 		FDigestedSightProperties();
 		FDigestedSightProperties(const UAISenseConfig_Sight& SenseConfig);
-	};	
-	
+	};
+
+	enum class EVisibilityResult
+	{
+		Visible,
+		NotVisible,
+		Pending
+	};
+
 	typedef TMap<FAISightTarget::FTargetId, FAISightTarget> FTargetsContainer;
 	FTargetsContainer ObservedTargets;
 	TMap<FPerceptionListenerID, FDigestedSightProperties> DigestedProperties;
@@ -153,11 +179,16 @@ public:
 	bool bSightQueriesOutOfRangeDirty = true;
 	TArray<FAISightQuery> SightQueriesOutOfRange;
 	TArray<FAISightQuery> SightQueriesInRange;
+	TArray<FAISightQuery> SightQueriesPending;
 
 protected:
 	UPROPERTY(EditDefaultsOnly, Category = "AI Perception", config)
 	int32 MaxTracesPerTick;
 
+	/** Maximum number of asynchronous traces that can be requested in a single update call*/
+	UPROPERTY(EditDefaultsOnly, Category = "AI Perception", config)
+	int32 MaxAsyncTracesPerTick;
+
 	UPROPERTY(EditDefaultsOnly, Category = "AI Perception", config)
 	int32 MinQueriesPerTimeSliceCheck;
 
@@ -175,8 +206,19 @@ protected:
 	UPROPERTY(EditDefaultsOnly, Category = "AI Perception", config)
 	float SightLimitQueryImportance;
 
+	/** Defines the amount of async trace queries to prevent based on the number of pending queries at the start of an update.
+	 * 1 means that the async trace budget is slashed by the pending queries count
+	 * 0 means that the async trace budget is not impacted by the pending queries
+	 */
+	UPROPERTY(EditDefaultsOnly, Category = "AI Perception", config)
+	float PendingQueriesBudgetReductionRatio;
+
 	ECollisionChannel DefaultSightCollisionChannel;
 
+	FOnPendingVisibilityQueryProcessedDelegate OnPendingVisibilityQueryProcessedDelegate;
+
+	UE_MT_DECLARE_RW_ACCESS_DETECTOR(QueriesListAccessDetector);
+
 public:
 
 	virtual void PostInitProperties() override;
@@ -192,9 +234,10 @@ public:
 protected:
 	virtual float Update() override;
 
-	bool ComputeVisibility(const UWorld* World, FAISightQuery& SightQuery, FPerceptionListener& Listener, const AActor* ListenerActor, FAISightTarget& Target, AActor* TargetActor, const FDigestedSightProperties& PropDigest, float& OutStimulusStrength, FVector& OutSeenLocation, int32& OutNumberOfLoSChecksPerformed) const;
+	EVisibilityResult ComputeVisibility(const UWorld* World, FAISightQuery& SightQuery, FPerceptionListener& Listener, const AActor* ListenerActor, FAISightTarget& Target, AActor* TargetActor, const FDigestedSightProperties& PropDigest, float& OutStimulusStrength, FVector& OutSeenLocation, int32& OutNumberOfLoSChecksPerformed, int32& OutNumberOfAsyncLosCheckRequested) const;
 	virtual bool ShouldAutomaticallySeeTarget(const FDigestedSightProperties& PropDigest, FAISightQuery* SightQuery, FPerceptionListener& Listener, AActor* TargetActor, float& OutStimulusStrength) const;
 	void UpdateQueryVisibilityStatus(FAISightQuery& SightQuery, FPerceptionListener& Listener, const bool bIsVisible, const FVector& SeenLocation, const float StimulusStrength, AActor* TargetActor, const FVector& TargetLocation) const;
+	void OnPendingVisibilityQueryProcessed(const FAISightQueryID& QueryID, const bool bIsVisible, const float StimulusStrength, const FVector& SeenLocation, const TOptional<int32>& UserData);
 
 	void OnNewListenerImpl(const FPerceptionListener& NewListener);
 	void OnListenerUpdateImpl(const FPerceptionListener& UpdatedListener);
diff --git a/Engine/Source/Runtime/AIModule/Classes/Perception/AISightTargetInterface.h b/Engine/Source/Runtime/AIModule/Classes/Perception/AISightTargetInterface.h
--- a/Engine/Source/Runtime/AIModule/Classes/Perception/AISightTargetInterface.h
+++ b/Engine/Source/Runtime/AIModule/Classes/Perception/AISightTargetInterface.h
@@ -3,6 +3,7 @@
 #pragma once
 
 #include "CoreMinimal.h"
+#include "AISense_Sight.h"
 #include "UObject/ObjectMacros.h"
 #include "UObject/Interface.h"
 #include "AISightTargetInterface.generated.h"
@@ -15,20 +16,54 @@ class AIMODULE_API UAISightTargetInterface : public UInterface
 	GENERATED_UINTERFACE_BODY()
 };
 
+struct AIMODULE_API FCanBeSeenFromContext
+{
+	/** The query identifier used by the delegate call to find the appropriate query */
+	FAISightQueryID SightQueryID;
+	/** The location of the observer */
+	FVector ObserverLocation;
+	/** The actor to ignore when doing test */
+	const AActor* IgnoreActor = nullptr;
+	/** If available, it is the previous visibility state */
+	const bool* bWasVisible = nullptr;
+};
+
 class AIMODULE_API IAISightTargetInterface
 {
 	GENERATED_IINTERFACE_BODY()
 
+	/**	
+	 * The method needs to check whether the implementer is visible from given observer's location. 
+	 * This version allows an asynchronous answer by returning Pending and using the provided Delegate when the information is computed
+	 * @param Context								Struct containing all the information required to perform the query
+	 * @param OutSeenLocation						The first visible target location
+	 * @param OutNumberOfLoSChecksPerformed			The number of synchronous queries done. This is used to stop the perception tick sooner when too many LoS checks are done
+	 * @param OutNumberOfAsyncLosCheckRequested		The number of asynchronous queries done. This is used to stop the perception tick sooner when too many asynchronous LoS checks are done
+	 * @param OutSightStrength						The sight strength for how well the target is seen
+	 * @param UserData								If available, it is a data passed between visibility tests for the users to store whatever they want
+	 * @param Delegate								If available, the delegate to call later on, if the result needed asynchronous checks. Calling this delegate within the CanBeSeenFrom call will ensure
+	 * @return										Undefined to indicate that we need to try another CanBeSeenFrom implementation, Visible/NotVisible if visible/not visible from the observer's location, Pending if the check is not finished
+	 */
+	virtual UAISense_Sight::EVisibilityResult CanBeSeenFrom(const FCanBeSeenFromContext& Context, FVector& OutSeenLocation, int32& OutNumberOfLoSChecksPerformed, int32& OutNumberOfAsyncLosCheckRequested, float& OutSightStrength, int32* UserData = nullptr, const FOnPendingVisibilityQueryProcessedDelegate* Delegate = nullptr)
+	{
+		OutNumberOfAsyncLosCheckRequested = 0;
+		PRAGMA_DISABLE_DEPRECATION_WARNINGS
+		const bool bCanBeSeenFrom = CanBeSeenFrom(Context.ObserverLocation, OutSeenLocation, OutNumberOfLoSChecksPerformed, OutSightStrength, Context.IgnoreActor, Context.bWasVisible, UserData);
+		PRAGMA_ENABLE_DEPRECATION_WARNINGS
+		return bCanBeSeenFrom ? UAISense_Sight::EVisibilityResult::Visible : UAISense_Sight::EVisibilityResult::NotVisible;
+	}
+
 	/**	
 	 * The method needs to check whether the implementer is visible from given observer's location. 
 	 * @param ObserverLocation	The location of the observer
 	 * @param OutSeenLocation	The first visible target location
-	 * @param OutSightStrengh	The sight strength for how well the target is seen
+	 * @param OutSightStrength	The sight strength for how well the target is seen
 	 * @param IgnoreActor		The actor to ignore when doing test
 	 * @param bWasVisible		If available, it is the previous visibility state
 	 * @param UserData			If available, it is a data passed between visibility tests for the users to store whatever they want
 	 * @return	True if visible from the observer's location
 	 */
+	UE_DEPRECATED(5.2, "This function is deprecated. Use the new CanBeSeenFrom method signature")
 	virtual bool CanBeSeenFrom(const FVector& ObserverLocation, FVector& OutSeenLocation, int32& NumberOfLoSChecksPerformed, float& OutSightStrength, const AActor* IgnoreActor = nullptr, const bool* bWasVisible = nullptr, int32* UserData = nullptr) const
 	{ 
 		NumberOfLoSChecksPerformed = 0;
diff --git a/Engine/Source/Runtime/AIModule/Private/Perception/AISense_Sight.cpp b/Engine/Source/Runtime/AIModule/Private/Perception/AISense_Sight.cpp
--- a/Engine/Source/Runtime/AIModule/Private/Perception/AISense_Sight.cpp
+++ b/Engine/Source/Runtime/AIModule/Private/Perception/AISense_Sight.cpp
@@ -37,7 +37,9 @@ DECLARE_CYCLE_STAT(TEXT("Perception Sense: Sight, Process pending result"), STAT
 
 
 static const int32 DefaultMaxTracesPerTick = 6;
+static const int32 DefaultMaxAsyncTracesPerTick = 10;
 static const int32 DefaultMinQueriesPerTimeSliceCheck = 40;
+static const float DefaultPendingQueriesBudgetReductionRatio = 0.5f;
 
 enum class EForEachResult : uint8
 {
@@ -121,11 +123,13 @@ UAISense_Sight::FDigestedSightProperties::FDigestedSightProperties()
 UAISense_Sight::UAISense_Sight(const FObjectInitializer& ObjectInitializer)
 	: Super(ObjectInitializer)
 	, MaxTracesPerTick(DefaultMaxTracesPerTick)
+	, MaxAsyncTracesPerTick(DefaultMaxAsyncTracesPerTick)
 	, MinQueriesPerTimeSliceCheck(DefaultMinQueriesPerTimeSliceCheck)
 	, MaxTimeSlicePerTick(0.005) // 5ms
 	, HighImportanceQueryDistanceThreshold(300.f)
 	, MaxQueryImportance(60.f)
 	, SightLimitQueryImportance(10.f)
+	, PendingQueriesBudgetReductionRatio(DefaultPendingQueriesBudgetReductionRatio)
 {
 	if (HasAnyFlags(RF_ClassDefaultObject) == false)
 	{
@@ -135,6 +139,8 @@ UAISense_Sight::UAISense_Sight(const FObjectInitializer& ObjectInitializer)
 		OnNewListenerDelegate.BindUObject(this, &UAISense_Sight::OnNewListenerImpl);
 		OnListenerUpdateDelegate.BindUObject(this, &UAISense_Sight::OnListenerUpdateImpl);
 		OnListenerRemovedDelegate.BindUObject(this, &UAISense_Sight::OnListenerRemovedImpl);
+
+		OnPendingVisibilityQueryProcessedDelegate.BindUObject(this, &UAISense_Sight::OnPendingVisibilityQueryProcessed);
 	}
 
 	NotifyType = EAISenseNotifyType::OnPerceptionChange;
@@ -235,6 +241,8 @@ float UAISense_Sight::Update()
 		return SuspendNextUpdate;
 	}
 
+	UE_MT_SCOPED_WRITE_ACCESS(QueriesListAccessDetector);
+
 	// sort Sight Queries
 	{
 		auto RecalcScore = [](FAISightQuery& SightQuery)->EForEachResult
@@ -259,8 +267,9 @@ float UAISense_Sight::Update()
 	}
 
 	int32 TracesCount = 0;
+	int32 AsyncTracesCount = FMath::Max(0, static_cast<int32>(PendingQueriesBudgetReductionRatio * SightQueriesPending.Num()));	// pending queries should be requesting async collisions traces at this frame, so we might want to restrain ourself in this update
 	int32 NumQueriesProcessed = 0;
-	double TimeSliceEnd = FPlatformTime::Seconds() + MaxTimeSlicePerTick;
+	const double TimeSliceEnd = FPlatformTime::Seconds() + MaxTimeSlicePerTick;
 	bool bHitTimeSliceLimit = false;
 #if AISENSE_SIGHT_TIMESLICING_DEBUG
 	FTimingSlicingInfo SlicingInfo;
@@ -269,7 +278,8 @@ float UAISense_Sight::Update()
 	enum class EOperationType : uint8
 	{
 		Remove,
-		SwapList
+		SwapList,
+		MoveToPending
 	};
 	struct FQueryOperation
 	{
@@ -296,7 +306,7 @@ float UAISense_Sight::Update()
 			bHitTimeSliceLimit = true;
 		}
 
-		if (bHitTimeSliceLimit || TracesCount >= MaxTracesPerTick)
+		if (bHitTimeSliceLimit || TracesCount >= MaxTracesPerTick || AsyncTracesCount >= MaxAsyncTracesPerTick)
 		{
 			break;
 		}
@@ -339,25 +349,40 @@ float UAISense_Sight::Update()
 			float StimulusStrength = 1.f;
 			FVector SeenLocation(0.f);
 			int32 NumberOfLoSChecksPerformed = 0;
+			int32 NumberOfAsyncLosCheckRequested = 0;
 
-			const bool bIsVisible = ComputeVisibility(World, *SightQuery, Listener, ListenerBodyActor, Target, TargetActor, PropDigest, StimulusStrength, SeenLocation, NumberOfLoSChecksPerformed);
+			const EVisibilityResult VisibilityResult = ComputeVisibility(World, *SightQuery, Listener, ListenerBodyActor, Target, TargetActor, PropDigest, StimulusStrength, SeenLocation, NumberOfLoSChecksPerformed, NumberOfAsyncLosCheckRequested);
 
 			TracesCount += NumberOfLoSChecksPerformed;
+			AsyncTracesCount += NumberOfAsyncLosCheckRequested;
 
-			const bool bWasVisible = SightQuery->bLastResult;
-			const FVector TargetLocation = TargetActor->GetActorLocation();
-			UpdateQueryVisibilityStatus(*SightQuery, Listener, bIsVisible, SeenLocation, StimulusStrength, TargetActor, TargetLocation);
-
-			const float SightRadiusSq = bWasVisible ? PropDigest.LoseSightRadiusSq : PropDigest.SightRadiusSq;
-			SightQuery->Importance = CalcQueryImportance(Listener, TargetLocation, SightRadiusSq);
-			const bool bShouldBeInRange = SightQuery->Importance > 0.0f;
-			if (bIsInRangeQuery != bShouldBeInRange)
+			if (VisibilityResult == EVisibilityResult::Pending)
 			{
-				QueryOperations.Add(FQueryOperation(bIsInRangeQuery, EOperationType::SwapList, bIsInRangeQuery ? InRangeIndex : OutOfRangeIndex));
+				QueryOperations.Add(FQueryOperation(bIsInRangeQuery, EOperationType::MoveToPending, bIsInRangeQuery ? InRangeIndex : OutOfRangeIndex));
+
+				// Call this to be able to know how many frames were used to process the visibility
+				SightQuery->OnProcessed();
 			}
+			else
+			{
+				UE_CLOG(VisibilityResult != EVisibilityResult::Visible && VisibilityResult != EVisibilityResult::NotVisible, LogAIPerception, Error, TEXT("UAISense_Sight::Update received invalid Visibility result [%d] for query between Listener %s and Target %s. We'll consider it as NotVisible"), VisibilityResult, *GetNameSafe(ListenerBodyActor), *GetNameSafe(TargetActor));
 
-			// restart query
-			SightQuery->OnProcessed();
+				const bool bIsVisible = VisibilityResult == EVisibilityResult::Visible;
+				const bool bWasVisible = SightQuery->bLastResult;
+				const FVector TargetLocation = TargetActor->GetActorLocation();
+				UpdateQueryVisibilityStatus(*SightQuery, Listener, bIsVisible, SeenLocation, StimulusStrength, TargetActor, TargetLocation);
+
+				const float SightRadiusSq = bWasVisible ? PropDigest.LoseSightRadiusSq : PropDigest.SightRadiusSq;
+				SightQuery->Importance = CalcQueryImportance(Listener, TargetLocation, SightRadiusSq);
+				const bool bShouldBeInRange = SightQuery->Importance > 0.0f;
+				if (bIsInRangeQuery != bShouldBeInRange)
+				{
+					QueryOperations.Add(FQueryOperation(bIsInRangeQuery, EOperationType::SwapList, bIsInRangeQuery ? InRangeIndex : OutOfRangeIndex));
+				}
+
+				// restart query
+				SightQuery->OnProcessed();
+			}
 		}
 		else
 		{
@@ -393,16 +418,33 @@ float UAISense_Sight::Update()
 		TArray<FAISightQuery> SightQueriesOutOfRangeToInsert;
 		for (FQueryOperation& Operation : QueryOperations)
 		{
-			if (Operation.OpType == EOperationType::SwapList)
+			switch (Operation.OpType)
 			{
-				if (Operation.bInRange)
+			case EOperationType::SwapList:
 				{
-					SightQueriesOutOfRangeToInsert.Push(SightQueriesInRange[Operation.Index]);
+					if (Operation.bInRange)
+					{
+						SightQueriesOutOfRangeToInsert.Push(SightQueriesInRange[Operation.Index]);
+					}
+					else
+					{
+						SightQueriesInRange.Add(SightQueriesOutOfRange[Operation.Index]);
+					}
 				}
-				else
+				break;
+
+			case EOperationType::MoveToPending:
 				{
-					SightQueriesInRange.Add(SightQueriesOutOfRange[Operation.Index]);
+					SightQueriesPending.Add(Operation.bInRange ? SightQueriesInRange[Operation.Index] : SightQueriesOutOfRange[Operation.Index]);
 				}
+				break;
+
+			case EOperationType::Remove:
+				break;
+
+			default:
+				check(false);
+				break;
 			}
 
 			if (Operation.bInRange)
@@ -449,7 +491,7 @@ float UAISense_Sight::Update()
 	return 0.f;
 }
 
-bool UAISense_Sight::ComputeVisibility(const UWorld* World, FAISightQuery& SightQuery, FPerceptionListener& Listener, const AActor* ListenerActor, FAISightTarget& Target, AActor* TargetActor, const FDigestedSightProperties& PropDigest, float& OutStimulusStrength, FVector& OutSeenLocation, int32& OutNumberOfLoSChecksPerformed) const
+UAISense_Sight::EVisibilityResult UAISense_Sight::ComputeVisibility(const UWorld* World, FAISightQuery& SightQuery, FPerceptionListener& Listener, const AActor* ListenerActor, FAISightTarget& Target, AActor* TargetActor, const FDigestedSightProperties& PropDigest, float& OutStimulusStrength, FVector& OutSeenLocation, int32& OutNumberOfLoSChecksPerformed, int32& OutNumberOfAsyncLosCheckRequested) const
 {
 	SCOPE_CYCLE_COUNTER(STAT_AI_Sense_Sight_ComputeVisibility);
 
@@ -457,21 +499,26 @@ bool UAISense_Sight::ComputeVisibility(const UWorld* World, FAISightQuery& Sight
 	if (ShouldAutomaticallySeeTarget(PropDigest, &SightQuery, Listener, TargetActor, OutStimulusStrength))
 	{
 		OutSeenLocation = FAISystem::InvalidLocation;
-		return true;
+		return EVisibilityResult::Visible;
 	}
 
 	const FVector TargetLocation = TargetActor->GetActorLocation();
 	const float SightRadiusSq = SightQuery.bLastResult ? PropDigest.LoseSightRadiusSq : PropDigest.SightRadiusSq;
 	if (!FAISystem::CheckIsTargetInSightCone(Listener.CachedLocation, Listener.CachedDirection, PropDigest.PeripheralVisionAngleCos, PropDigest.PointOfViewBackwardOffset, PropDigest.NearClippingRadiusSq, SightRadiusSq, TargetLocation))
 	{
-		return false;
+		return EVisibilityResult::NotVisible;
 	}
 
 	if (Target.SightTargetInterface != nullptr)
 	{
 		const bool bWasVisible = SightQuery.bLastResult;
-		const bool bCanBeSeen = Target.SightTargetInterface->CanBeSeenFrom(Listener.CachedLocation, OutSeenLocation, OutNumberOfLoSChecksPerformed, OutStimulusStrength, ListenerActor, &bWasVisible, &SightQuery.UserData);
-		return bCanBeSeen;
+		FCanBeSeenFromContext Context;
+		Context.SightQueryID = FAISightQueryID(SightQuery);
+		Context.ObserverLocation = Listener.CachedLocation;
+		Context.IgnoreActor = ListenerActor;
+		Context.bWasVisible = &bWasVisible;
+
+		return Target.SightTargetInterface->CanBeSeenFrom(Context, OutSeenLocation, OutNumberOfLoSChecksPerformed, OutNumberOfAsyncLosCheckRequested, OutStimulusStrength, &SightQuery.UserData, &OnPendingVisibilityQueryProcessedDelegate);
 	}
 	else
 	{
@@ -490,11 +537,11 @@ bool UAISense_Sight::ComputeVisibility(const UWorld* World, FAISightQuery& Sight
 		if (bHit == false || HitResultActorIsOwnedByTargetActor())
 		{
 			OutSeenLocation = TargetLocation;
-			return true;
+			return EVisibilityResult::Visible;
 		}
 		else
 		{
-			return false;
+			return EVisibilityResult::NotVisible;
 		}
 	}
 }
@@ -522,6 +569,63 @@ void UAISense_Sight::UpdateQueryVisibilityStatus(FAISightQuery& SightQuery, FPer
 	SIGHT_LOG_SEGMENT(ListenerPtr->GetOwner(), Listener.CachedLocation, TargetLocation, bIsVisible ? FColor::Green : FColor::Red, TEXT("TargetID %d"), Target.TargetId);
 }
 
+void UAISense_Sight::OnPendingVisibilityQueryProcessed(const FAISightQueryID& QueryID, const bool bIsVisible, const float StimulusStrength, const FVector& SeenLocation, const TOptional<int32>& UserData)
+{
+	SCOPE_CYCLE_COUNTER(STAT_AI_Sense_Sight_ProcessPendingQuery);
+
+	UE_MT_SCOPED_WRITE_ACCESS(QueriesListAccessDetector);
+
+	const int32 QueryIdx = SightQueriesPending.IndexOfByPredicate([&QueryID](const FAISightQuery& Element)
+		{
+			return Element.ObserverId == QueryID.ObserverId
+				&& Element.TargetId == QueryID.TargetId;
+		});
+
+	if (QueryIdx == INDEX_NONE)
+	{
+		// the query is not pending. It must have been removed because the source or the target have been removed
+		return;
+	}
+
+	FAISightQuery SightQuery = SightQueriesPending[QueryIdx];
+	SightQueriesPending.RemoveAtSwap(QueryIdx, 1, false);
+
+	AIPerception::FListenerMap& ListenersMap = *GetListeners();
+	FPerceptionListener* Listener = ListenersMap.Find(SightQuery.ObserverId);
+	FAISightTarget* Target = ObservedTargets.Find(SightQuery.TargetId);
+	AActor* TargetActor = Target ? Target->Target.Get() : nullptr;
+	if (TargetActor == nullptr || Listener == nullptr)
+	{
+		return;
+	}
+
+	const bool bWasVisible = SightQuery.bLastResult;
+	const FVector TargetLocation = TargetActor->GetActorLocation();
+	UpdateQueryVisibilityStatus(SightQuery, *Listener, bIsVisible, SeenLocation, StimulusStrength, TargetActor, TargetLocation);
+
+	if (UserData.IsSet())
+	{
+		SightQuery.UserData = UserData.GetValue();
+	}
+
+	// Call this to be able to have an accurate tick time
+	SightQuery.OnProcessed();
+
+	const FDigestedSightProperties& PropDigest = DigestedProperties[SightQuery.ObserverId];
+	const float SightRadiusSq = bWasVisible ? PropDigest.LoseSightRadiusSq : PropDigest.SightRadiusSq;
+	SightQuery.Importance = CalcQueryImportance(*Listener, TargetLocation, SightRadiusSq);
+	const bool bShouldBeInRange = SightQuery.Importance > 0.0f;
+	if (bShouldBeInRange)
+	{
+		SightQueriesInRange.Add(SightQuery);
+	}
+	else
+	{
+		SightQueriesOutOfRange.Insert(SightQuery, NextOutOfRangeIndex);
+		++NextOutOfRangeIndex;
+	}
+}
+
 void UAISense_Sight::RegisterEvent(const FAISightEvent& Event)
 {
 
@@ -534,9 +638,11 @@ void UAISense_Sight::RegisterSource(AActor& SourceActor)
 
 void UAISense_Sight::UnregisterSource(AActor& SourceActor)
 {
+	UE_MT_SCOPED_WRITE_ACCESS(QueriesListAccessDetector);
+
 	const FAISightTarget::FTargetId AsTargetId = SourceActor.GetUniqueID();
 	FAISightTarget AsTarget;
-	
+
 	if (ObservedTargets.RemoveAndCopyValue(AsTargetId, AsTarget) 
 		&& (SightQueriesInRange.Num() + SightQueriesOutOfRange.Num()) > 0)
 	{
@@ -570,6 +676,7 @@ void UAISense_Sight::UnregisterSource(AActor& SourceActor)
 			{
 				bSightQueriesOutOfRangeDirty = true;
 			}
+			ReverseForEach(SightQueriesPending, RemoveQuery);
 		}
 	}
 }
@@ -811,8 +918,10 @@ void UAISense_Sight::RemoveAllQueriesByListener(const FPerceptionListener& Liste
 		return;
 	}
 
+	UE_MT_SCOPED_WRITE_ACCESS(QueriesListAccessDetector);
+
 	const uint32 ListenerId = Listener.GetListenerID();
-	
+
 	auto RemoveQuery = [&ListenerId, &OnRemoveFunc](TArray<FAISightQuery>& SightQueries, const int32 QueryIndex)->EReverseForEachResult
 	{
 		const FAISightQuery& SightQuery = SightQueries[QueryIndex];
@@ -833,11 +942,13 @@ void UAISense_Sight::RemoveAllQueriesByListener(const FPerceptionListener& Liste
 	{
 		bSightQueriesOutOfRangeDirty = true;
 	}
+	ReverseForEach(SightQueriesPending, RemoveQuery);
 }
 
 void UAISense_Sight::RemoveAllQueriesToTarget(const FAISightTarget::FTargetId& TargetId, const TFunction<void(const FAISightQuery&)>& OnRemoveFunc/*= nullptr */)
 {
 	SCOPE_CYCLE_COUNTER(STAT_AI_Sense_Sight_RemoveToTarget);
+	UE_MT_SCOPED_WRITE_ACCESS(QueriesListAccessDetector);
 
 	auto RemoveQuery = [&TargetId, &OnRemoveFunc](TArray<FAISightQuery>& SightQueries, const int32 QueryIndex)->EReverseForEachResult
 	{
@@ -859,6 +970,7 @@ void UAISense_Sight::RemoveAllQueriesToTarget(const FAISightTarget::FTargetId& T
 	{
 		bSightQueriesOutOfRangeDirty = true;
 	}
+	ReverseForEach(SightQueriesPending, RemoveQuery);
 }
 
 void UAISense_Sight::OnListenerForgetsActor(const FPerceptionListener& Listener, AActor& ActorToForget)
@@ -879,12 +991,17 @@ void UAISense_Sight::OnListenerForgetsActor(const FPerceptionListener& Listener,
 
 	if (ForEach(SightQueriesInRange, ForgetPreviousResult) == EForEachResult::Continue)
 	{
-		ForEach(SightQueriesOutOfRange, ForgetPreviousResult);
+		if (ForEach(SightQueriesOutOfRange, ForgetPreviousResult) == EForEachResult::Continue)
+		{
+			ForEach(SightQueriesPending, ForgetPreviousResult);
+		}
 	}
 }
 
 void UAISense_Sight::OnListenerForgetsAll(const FPerceptionListener& Listener)
 {
+	UE_MT_SCOPED_WRITE_ACCESS(QueriesListAccessDetector);
+
 	const uint32 ListenerId = Listener.GetListenerID();
 
 	auto ForgetPreviousResult = [&ListenerId](FAISightQuery& SightQuery)->EForEachResult
@@ -898,5 +1015,6 @@ void UAISense_Sight::OnListenerForgetsAll(const FPerceptionListener& Listener)
 
 	ForEach(SightQueriesInRange, ForgetPreviousResult);
 	ForEach(SightQueriesOutOfRange, ForgetPreviousResult);
+	ForEach(SightQueriesPending, ForgetPreviousResult);
 }
 
