Compile time optimisation: do not include PathFollowingComponent in AIController
#tests Profile compile time

diff --git a/Engine/Source/Runtime/AIModule/Classes/AIController.h b/Engine/Source/Runtime/AIModule/Classes/AIController.h
--- a/Engine/Source/Runtime/AIModule/Classes/AIController.h
+++ b/Engine/Source/Runtime/AIModule/Classes/AIController.h
@@ -13,7 +13,9 @@
 #include "GameplayTask.h"
 #include "GameFramework/Pawn.h"
 #include "GameFramework/Controller.h"
+#if UE_ENABLE_INCLUDE_ORDER_DEPRECATED_IN_5_2
 #include "Navigation/PathFollowingComponent.h"
+#endif
 #include "Perception/AIPerceptionListenerInterface.h"
 #include "GenericTeamAgentInterface.h"
 #include "VisualLogger/VisualLoggerDebugSnapshotInterface.h"
@@ -28,9 +30,13 @@ class UBrainComponent;
 class UCanvas;
 class UGameplayTaskResource;
 class UGameplayTasksComponent;
+class UPathFollowingComponent;
 class UPawnAction;
 class UPawnActionsComponent;
 struct FVisualLogEntry;
+namespace EPathFollowingRequestResult {	enum Type : int; }
+namespace EPathFollowingResult { enum Type : int; }
+namespace EPathFollowingStatus { enum Type : int; }
 
 #if ENABLE_VISUAL_LOG
 struct FVisualLogEntry;
@@ -191,7 +197,7 @@ public:
 	 *  @param OutPath - optional output param, filled in with assigned path
 	 *  @return struct holding MoveId and enum code
 	 */
-	virtual FPathFollowingRequestResult MoveTo(const FAIMoveRequest& MoveRequest, FNavPathSharedPtr* OutPath = nullptr);
+	virtual struct FPathFollowingRequestResult MoveTo(const FAIMoveRequest& MoveRequest, FNavPathSharedPtr* OutPath = nullptr);
 
 	/** Passes move request and path object to path following */
 	virtual FAIRequestID RequestMove(const FAIMoveRequest& MoveRequest, FNavPathSharedPtr Path);
@@ -222,13 +228,13 @@ public:
 	virtual void StopMovement() override;
 
 	/** Called on completing current movement request */
-	virtual void OnMoveCompleted(FAIRequestID RequestID, const FPathFollowingResult& Result);
+	virtual void OnMoveCompleted(FAIRequestID RequestID, const struct FPathFollowingResult& Result);
 
 	UE_DEPRECATED_FORGAME(4.13, "This function is now deprecated, please use version with EPathFollowingResultDetails parameter.")
 	virtual void OnMoveCompleted(FAIRequestID RequestID, EPathFollowingResult::Type Result);
 
 	/** Returns the Move Request ID for the current move */
-	FORCEINLINE FAIRequestID GetCurrentMoveRequestID() const { return GetPathFollowingComponent() ? GetPathFollowingComponent()->GetCurrentRequestId() : FAIRequestID::InvalidRequest; }
+	FAIRequestID GetCurrentMoveRequestID() const;
 
 	/** Blueprint notification that we've completed the current movement request */
 	UPROPERTY(BlueprintAssignable, meta = (DisplayName = "MoveCompleted"))
@@ -385,7 +391,7 @@ public:
 	// INavAgentInterface
 	//----------------------------------------------------------------------//
 	virtual bool IsFollowingAPath() const override;
-	virtual IPathFollowingAgentInterface* GetPathFollowingAgent() const override { return PathFollowingComponent; }
+	virtual IPathFollowingAgentInterface* GetPathFollowingAgent() const override;
 
 	//----------------------------------------------------------------------//
 	// IGenericTeamAgentInterface
diff --git a/Engine/Source/Runtime/AIModule/Classes/Navigation/PathFollowingComponent.h b/Engine/Source/Runtime/AIModule/Classes/Navigation/PathFollowingComponent.h
--- a/Engine/Source/Runtime/AIModule/Classes/Navigation/PathFollowingComponent.h
+++ b/Engine/Source/Runtime/AIModule/Classes/Navigation/PathFollowingComponent.h
@@ -33,7 +33,7 @@ class UNavigationComponent;
 UENUM(BlueprintType)
 namespace EPathFollowingStatus
 {
-	enum Type
+	enum Type : int
 	{
 		/** No requests */
 		Idle,
@@ -52,7 +52,7 @@ namespace EPathFollowingStatus
 UENUM(BlueprintType)
 namespace EPathFollowingResult
 {
-	enum Type
+	enum Type : int
 	{
 		/** Reached destination */
 		Success,
@@ -140,7 +140,7 @@ struct AIMODULE_API FPathFollowingResult
 UENUM(BlueprintType)
 namespace EPathFollowingAction
 {
-	enum Type
+	enum Type : int
 	{
 		Error,
 		NoMove,
@@ -153,7 +153,7 @@ namespace EPathFollowingAction
 UENUM(BlueprintType)
 namespace EPathFollowingRequestResult
 {
-	enum Type
+	enum Type : int
 	{
 		Failed,
 		AlreadyAtGoal,
@@ -172,7 +172,7 @@ struct AIMODULE_API FPathFollowingRequestResult
 
 namespace EPathFollowingDebugTokens
 {
-	enum Type
+	enum Type : int
 	{
 		Description,
 		ParamName,
@@ -184,7 +184,7 @@ namespace EPathFollowingDebugTokens
 // DEPRECATED, please use EPathFollowingResultDetails instead, will be removed with deprecated override of AbortMove function
 namespace EPathFollowingMessage
 {
-	enum Type
+	enum Type : int
 	{
 		NoPath,
 		OtherRequest,
diff --git a/Engine/Source/Runtime/AIModule/Private/AIController.cpp b/Engine/Source/Runtime/AIModule/Private/AIController.cpp
--- a/Engine/Source/Runtime/AIModule/Private/AIController.cpp
+++ b/Engine/Source/Runtime/AIModule/Private/AIController.cpp
@@ -23,6 +23,7 @@
 #include "GameplayTasksComponent.h"
 #include "Tasks/GameplayTask_ClaimResource.h"
 #include "NetworkingDistanceConstants.h"
+#include "Navigation/PathFollowingComponent.h"
 
 #include UE_INLINE_GENERATED_CPP_BY_NAME(AIController)
 
@@ -871,6 +872,11 @@ bool AAIController::IsFollowingAPath() const
 	return (PathFollowingComponent != nullptr) && (PathFollowingComponent->GetStatus() != EPathFollowingStatus::Idle);
 }
 
+IPathFollowingAgentInterface* AAIController::GetPathFollowingAgent() const
+{
+	return PathFollowingComponent;
+}
+
 FVector AAIController::GetImmediateMoveDestination() const
 {
 	return (PathFollowingComponent) ? PathFollowingComponent->GetCurrentTargetLocation() : FVector::ZeroVector;
@@ -895,6 +901,11 @@ void AAIController::OnMoveCompleted(FAIRequestID RequestID, EPathFollowingResult
 	// deprecated
 }
 
+FAIRequestID AAIController::GetCurrentMoveRequestID() const
+{
+	return GetPathFollowingComponent() ? GetPathFollowingComponent()->GetCurrentRequestId() : FAIRequestID::InvalidRequest;
+}
+
 bool AAIController::RunBehaviorTree(UBehaviorTree* BTAsset)
 {
 	// @todo: find BrainComponent and see if it's BehaviorTreeComponent
diff --git a/Engine/Source/Runtime/AIModule/Private/BehaviorTree/Decorators/BTDecorator_ReachedMoveGoal.cpp b/Engine/Source/Runtime/AIModule/Private/BehaviorTree/Decorators/BTDecorator_ReachedMoveGoal.cpp
--- a/Engine/Source/Runtime/AIModule/Private/BehaviorTree/Decorators/BTDecorator_ReachedMoveGoal.cpp
+++ b/Engine/Source/Runtime/AIModule/Private/BehaviorTree/Decorators/BTDecorator_ReachedMoveGoal.cpp
@@ -2,6 +2,7 @@
 
 #include "BehaviorTree/Decorators/BTDecorator_ReachedMoveGoal.h"
 #include "AIController.h"
+#include "Navigation/PathFollowingComponent.h"
 
 #include UE_INLINE_GENERATED_CPP_BY_NAME(BTDecorator_ReachedMoveGoal)
 
diff --git a/Engine/Source/Runtime/AIModule/Private/GameplayDebugger/GameplayDebuggerCategory_AI.cpp b/Engine/Source/Runtime/AIModule/Private/GameplayDebugger/GameplayDebuggerCategory_AI.cpp
--- a/Engine/Source/Runtime/AIModule/Private/GameplayDebugger/GameplayDebuggerCategory_AI.cpp
+++ b/Engine/Source/Runtime/AIModule/Private/GameplayDebugger/GameplayDebuggerCategory_AI.cpp
@@ -12,6 +12,7 @@
 #include "GameFramework/PlayerController.h"
 #include "Engine/Engine.h"
 #include "EngineGlobals.h"
+#include "Navigation/PathFollowingComponent.h"
 #include "NavMesh/RecastNavMesh.h"
 #include "Engine/Canvas.h"
 #include "AIController.h"
