Fixed a case where any blueprint base decorator, during a transition, would  RequestExecution but was not in a state to make the right call. The problem come with the fact that the RequestExecution is deducing through the current state of the node to know if this request is a ActivateBranch or a DeactivateBranch. But in this case the blueprint base decorator is asking if it is active the during the transition through the call const bool bIsOnActiveBranch = OwnerComp.IsExecutingBranch(GetMyNode(), GetChildIndex()); and would get the wrong answer. So instead of deducing what the RequestExecution meant, we are now using new API methods that are requesting either a RequestBranchActivation or a RequestBranchDeactivation. During a transition, those get queued up and unqueued post transition and the code now make sure the request still holds and make sense.
UDN case 00489819

diff --git a/Engine/Source/Developer/AITestSuite/Classes/BehaviorTree/TestBTDecorator_Blueprint.h b/Engine/Source/Developer/AITestSuite/Classes/BehaviorTree/TestBTDecorator_Blueprint.h
new file mode 100644
--- /dev/null
+++ b/Engine/Source/Developer/AITestSuite/Classes/BehaviorTree/TestBTDecorator_Blueprint.h
@@ -0,0 +1,48 @@
+// Copyright Epic Games, Inc. All Rights Reserved.
+
+#pragma once 
+
+#include "BehaviorTree/Decorators/BTDecorator_BlueprintBase.h"
+
+#include "TestBTDecorator_Blueprint.generated.h"
+
+UENUM()
+enum class EBPConditionType
+{
+	NoCondition,
+	True,
+	False
+};
+
+UCLASS(meta = (HiddenNode))
+class UTestBTDecorator_Blueprint : public UBTDecorator_BlueprintBase
+{
+	GENERATED_UCLASS_BODY()
+
+public:
+
+	UPROPERTY()
+	EBPConditionType BPConditionType = EBPConditionType::True;
+	
+	UPROPERTY()
+	int32 LogIndexBecomeRelevant;
+
+	UPROPERTY()
+	int32 LogIndexCeaseRelevant;
+
+	UPROPERTY()
+	int32 LogIndexCalculate;
+
+	UPROPERTY()
+	FName ObservingKeyName = NAME_None;
+
+protected:
+	virtual void OnBecomeRelevant(UBehaviorTreeComponent& OwnerComp, uint8* NodeMemory) override;
+	virtual void OnCeaseRelevant(UBehaviorTreeComponent& OwnerComp, uint8* NodeMemory) override;
+	virtual bool CalculateRawConditionValue(UBehaviorTreeComponent& OwnerComp, uint8* NodeMemory) const override;
+	virtual void PostLoad() override;
+	virtual void PostInitProperties() override;
+
+private:
+	void LogExecution(int32 LogNumber) const;
+};
diff --git a/Engine/Source/Developer/AITestSuite/Private/BehaviorTree/TestBTDecorator_Blueprint.cpp b/Engine/Source/Developer/AITestSuite/Private/BehaviorTree/TestBTDecorator_Blueprint.cpp
new file mode 100644
--- /dev/null
+++ b/Engine/Source/Developer/AITestSuite/Private/BehaviorTree/TestBTDecorator_Blueprint.cpp
@@ -0,0 +1,68 @@
+// Copyright Epic Games, Inc. All Rights Reserved.
+
+#include "BehaviorTree/TestBTDecorator_Blueprint.h"
+#include "MockAI_BT.h"
+
+#include UE_INLINE_GENERATED_CPP_BY_NAME(TestBTDecorator_Blueprint)
+
+UTestBTDecorator_Blueprint::UTestBTDecorator_Blueprint(const FObjectInitializer& ObjectInitializer) : Super(ObjectInitializer)
+, LogIndexBecomeRelevant(INDEX_NONE)
+, LogIndexCeaseRelevant(INDEX_NONE)
+, LogIndexCalculate(INDEX_NONE)
+{
+}
+
+void UTestBTDecorator_Blueprint::OnBecomeRelevant(UBehaviorTreeComponent& OwnerComp, uint8* NodeMemory)
+{
+	LogExecution(LogIndexBecomeRelevant);
+	Super::OnBecomeRelevant(OwnerComp, NodeMemory);
+}
+
+void UTestBTDecorator_Blueprint::OnCeaseRelevant(UBehaviorTreeComponent& OwnerComp, uint8* NodeMemory)
+{
+	LogExecution(LogIndexCeaseRelevant);
+	Super::OnCeaseRelevant(OwnerComp, NodeMemory);
+}
+
+bool UTestBTDecorator_Blueprint::CalculateRawConditionValue(UBehaviorTreeComponent& OwnerComp, uint8* NodeMemory) const
+{
+	LogExecution(LogIndexCalculate);
+	return BPConditionType == EBPConditionType::True ? true : false;
+}
+
+void UTestBTDecorator_Blueprint::PostLoad()
+{
+	// Skip UBTDecorator_BlueprintBase as we are faking the fetching of blueprint BB keys
+	UBTDecorator::PostLoad();
+
+	// Let's fake this blueprint has a condition check as we overriden the CalculateRawConditionValue
+	PerformConditionCheckImplementations = BPConditionType == EBPConditionType::NoCondition ? 0 : 1;
+
+	if (!ObservingKeyName.IsNone())
+	{
+		ObservedKeyNames.Reset();
+		ObservedKeyNames.Add(ObservingKeyName);
+		bIsObservingBB = true;
+	}
+
+	if (PerformConditionCheckImplementations || bIsObservingBB)
+	{
+		bNotifyBecomeRelevant = true;
+		bNotifyCeaseRelevant = true;
+	}
+}
+
+void UTestBTDecorator_Blueprint::PostInitProperties()
+{
+	// Skip UBTDecorator_BlueprintBase as we are faking the fetching of blueprint BB keys
+	UBTDecorator::PostInitProperties();
+}
+
+void UTestBTDecorator_Blueprint::LogExecution(int32 LogNumber) const
+{
+	if (LogNumber != INDEX_NONE)
+	{
+		UMockAI_BT::ExecutionLog.Add(LogNumber);
+	}
+}
+
diff --git a/Engine/Source/Developer/AITestSuite/Private/Tests/BTTest.cpp b/Engine/Source/Developer/AITestSuite/Private/Tests/BTTest.cpp
--- a/Engine/Source/Developer/AITestSuite/Private/Tests/BTTest.cpp
+++ b/Engine/Source/Developer/AITestSuite/Private/Tests/BTTest.cpp
@@ -153,6 +153,84 @@ struct FAITest_BTBasicSequence : public FAITest_SimpleBT
 };
 IMPLEMENT_AI_LATENT_TEST(FAITest_BTBasicSequence, "System.AI.Behavior Trees.Composite node: sequence")
 
+struct FAITest_BTDecoratorBlueprint : public FAITest_SimpleBT
+{
+	FAITest_BTDecoratorBlueprint()
+	{
+		enum
+		{
+			DecoratorBlueprintCalculate = 1,
+			DecoratorBlueprintBecomeRelevant,
+			DecoratorBlueprintCeaseRelevant,
+			Task1Execute,
+			Task2Execute,
+			Task3Execute,
+		};
+
+		UBTCompositeNode& CompNode = FBTBuilder::AddSelector(*BTAsset);
+		{
+			UBTCompositeNode& CompNode1 = FBTBuilder::AddSequence(CompNode);
+			FBTBuilder::WithDecoratorBlueprint(CompNode, EBTFlowAbortMode::Both, EBPConditionType::True, DecoratorBlueprintBecomeRelevant, DecoratorBlueprintCeaseRelevant, DecoratorBlueprintCalculate);
+			{
+				// Task1
+				FBTBuilder::AddTask(CompNode1, Task1Execute, EBTNodeResult::Succeeded);
+
+				// Task2
+				FBTBuilder::AddTask(CompNode1, Task2Execute, EBTNodeResult::Succeeded, /*ExecutionTicks*/1);
+			}
+			// Task3
+			FBTBuilder::AddTask(CompNode, Task3Execute, EBTNodeResult::Failed);
+		}
+
+		ExpectedResult.Add(DecoratorBlueprintCalculate/*1*/);
+		ExpectedResult.Add(DecoratorBlueprintBecomeRelevant/*2*/);
+		ExpectedResult.Add(DecoratorBlueprintCalculate/*1*/);
+		ExpectedResult.Add(Task1Execute/*4*/);
+		ExpectedResult.Add(DecoratorBlueprintCalculate/*1*/);
+		ExpectedResult.Add(Task2Execute/*5*/);
+		ExpectedResult.Add(DecoratorBlueprintCalculate/*1*/);
+		ExpectedResult.Add(DecoratorBlueprintCalculate/*1*/);
+		ExpectedResult.Add(DecoratorBlueprintCeaseRelevant/*3*/);
+	}
+};
+IMPLEMENT_AI_LATENT_TEST(FAITest_BTDecoratorBlueprint, "System.AI.Behavior Trees.Decorator: blueprint")
+
+struct FAITest_BTDecoratorBlueprintNoConditionImpl : public FAITest_SimpleBT
+{
+	FAITest_BTDecoratorBlueprintNoConditionImpl()
+	{
+		enum
+		{
+			DecoratorBlueprintCalculate = 1,
+			DecoratorBlueprintBecomeRelevant,
+			DecoratorBlueprintCeaseRelevant,
+			Task1Execute,
+			Task2Execute,
+			Task3Execute,
+		};
+
+		UBTCompositeNode& CompNode = FBTBuilder::AddSelector(*BTAsset);
+		{
+			// Task1
+			FBTBuilder::AddTask(CompNode, Task1Execute, EBTNodeResult::Succeeded);
+			FBTBuilder::WithDecoratorBlueprint(CompNode, EBTFlowAbortMode::LowerPriority, EBPConditionType::NoCondition, DecoratorBlueprintBecomeRelevant, DecoratorBlueprintCeaseRelevant, DecoratorBlueprintCalculate, TEXT("Bool1"));
+
+			// Task2
+			FBTBuilder::AddTaskFlagChange(CompNode, true, EBTNodeResult::Succeeded, TEXT("Bool1"));
+			FBTBuilder::WithDecoratorBlackboard(CompNode, EBasicKeyOperation::NotSet, EBTFlowAbortMode::Self, TEXT("Bool1"));
+
+			// Task3
+			FBTBuilder::AddTask(CompNode, Task3Execute, EBTNodeResult::Succeeded, /*ExecutionTicks*/1);
+		}
+
+		ExpectedResult.Add(DecoratorBlueprintCalculate/*1*/);
+		ExpectedResult.Add(DecoratorBlueprintBecomeRelevant/*2*/);
+		ExpectedResult.Add(Task3Execute/*6*/);
+		ExpectedResult.Add(DecoratorBlueprintCeaseRelevant/*3*/);
+	}
+};
+IMPLEMENT_AI_LATENT_TEST(FAITest_BTDecoratorBlueprintNoConditionImpl, "System.AI.Behavior Trees.Decorator: blueprint no condition impl")
+
 struct FAITest_BTBasicParallelWait : public FAITest_SimpleBT
 {
 	FAITest_BTBasicParallelWait()
diff --git a/Engine/Source/Developer/AITestSuite/Public/BTBuilder.h b/Engine/Source/Developer/AITestSuite/Public/BTBuilder.h
--- a/Engine/Source/Developer/AITestSuite/Public/BTBuilder.h
+++ b/Engine/Source/Developer/AITestSuite/Public/BTBuilder.h
@@ -17,6 +17,7 @@
 
 #include "BehaviorTree/TestBTDecorator_Blackboard.h"
 #include "BehaviorTree/TestBTDecorator_DelayedAbort.h"
+#include "BehaviorTree/TestBTDecorator_Blueprint.h"
 #include "BehaviorTree/TestBTService_Log.h"
 #include "BehaviorTree/TestBTService_BTStopAction.h"
 #include "BehaviorTree/TestBTTask_LatentWithFlags.h"
@@ -301,6 +302,21 @@ struct FBTBuilder
 		AbortDecorator.bOnlyOnce = bAbortOnlyOnce;
 	}
 
+	static void WithDecoratorBlueprint(UBTCompositeNode& ParentNode, EBTFlowAbortMode::Type Observer, EBPConditionType BPConditionType = EBPConditionType::True,
+		int32 LogIndexBecomeRelevant = -1, int32 LogIndexCeaseRelevant = -1, int32 LogIndexCalculate = -1, FName ObservingKeyName = NAME_None )
+	{
+		UTestBTDecorator_Blueprint& BPDecorator = WithDecorator<UTestBTDecorator_Blueprint>(ParentNode);
+		BPDecorator.LogIndexBecomeRelevant = LogIndexBecomeRelevant;
+		BPDecorator.LogIndexCeaseRelevant = LogIndexCeaseRelevant;
+		BPDecorator.LogIndexCalculate = LogIndexCalculate;
+		BPDecorator.BPConditionType = BPConditionType;
+		BPDecorator.ObservingKeyName = ObservingKeyName;
+
+		FByteProperty* ObserverProp = FindFProperty<FByteProperty>(UBTDecorator_Blackboard::StaticClass(), TEXT("FlowAbortMode"));
+		uint8* ObserverPropData = ObserverProp->ContainerPtrToValuePtr<uint8>(&BPDecorator);
+		ObserverProp->SetIntPropertyValue(ObserverPropData, (uint64)Observer);
+	}
+
 	static void WithDecoratorLoop(UBTCompositeNode& ParentNode, int32 NumLoops = 2)
 	{
 		UBTDecorator_Loop& LoopDecorator = WithDecorator<UBTDecorator_Loop>(ParentNode);
diff --git a/Engine/Source/Runtime/AIModule/Classes/BehaviorTree/Decorators/BTDecorator_BlueprintBase.h b/Engine/Source/Runtime/AIModule/Classes/BehaviorTree/Decorators/BTDecorator_BlueprintBase.h
--- a/Engine/Source/Runtime/AIModule/Classes/BehaviorTree/Decorators/BTDecorator_BlueprintBase.h
+++ b/Engine/Source/Runtime/AIModule/Classes/BehaviorTree/Decorators/BTDecorator_BlueprintBase.h
@@ -49,12 +49,26 @@ public:
 
 	virtual FString GetStaticDescription() const override;
 	virtual void DescribeRuntimeValues(const UBehaviorTreeComponent& OwnerComp, uint8* NodeMemory, EBTDescriptionVerbosity::Type Verbosity, TArray<FString>& Values) const override;
-	virtual bool CalculateRawConditionValue(UBehaviorTreeComponent& OwnerComp, uint8* NodeMemory) const override final;
+	virtual bool CalculateRawConditionValue(UBehaviorTreeComponent& OwnerComp, uint8* NodeMemory) const override;
 	virtual void OnInstanceDestroyed(UBehaviorTreeComponent& OwnerComp) override;
 	virtual void InitializeFromAsset(UBehaviorTree& Asset) override;
 
-	/** return if this decorator should abort in current circumstances */
-	bool GetShouldAbort(UBehaviorTreeComponent& OwnerComp) const;
+	enum class EAbortType : uint8
+	{
+		NoAbort,
+		ActivateBranch,
+		DeactivateBranch,
+		Unknown,
+	};
+
+	/** return this decorator abort type in current circumstances */
+	EAbortType EvaluateAbortType(UBehaviorTreeComponent& OwnerComp) const;
+
+	UE_DEPRECATED(5.2, "Use GetShouldAbortType instead")
+	bool GetShouldAbort(UBehaviorTreeComponent& OwnerComp) const
+	{
+		return  EvaluateAbortType(OwnerComp) != EAbortType::NoAbort;
+	}
 
 	virtual void SetOwner(AActor* ActorOwner) override;
 
@@ -116,8 +130,6 @@ protected:
 	/** set if ReceiveConditionCheck is implemented by blueprint */
 	uint32 PerformConditionCheckImplementations : 2;
 
-	bool CalculateRawConditionValueImpl(UBehaviorTreeComponent& OwnerComp) const;
-	
 	virtual void OnBecomeRelevant(UBehaviorTreeComponent& OwnerComp, uint8* NodeMemory) override;
 	virtual void OnCeaseRelevant(UBehaviorTreeComponent& OwnerComp, uint8* NodeMemory) override;
 	virtual void OnNodeActivation(FBehaviorTreeSearchData& SearchData) override;
@@ -211,4 +223,7 @@ protected:
 	bool IsDecoratorObserverActive() const;
 
 	FORCEINLINE bool GetNeedsTickForConditionChecking() const { return PerformConditionCheckImplementations != 0 && (bIsObservingBB == false || bCheckConditionOnlyBlackBoardChanges == false); }
+
+	/** Handle abort request calls from the specified abort type */
+	void RequestAbort(UBehaviorTreeComponent& OwnerComp, const EAbortType Type);
 };
diff --git a/Engine/Source/Runtime/AIModule/Private/BehaviorTree/Decorators/BTDecorator_BlueprintBase.cpp b/Engine/Source/Runtime/AIModule/Private/BehaviorTree/Decorators/BTDecorator_BlueprintBase.cpp
--- a/Engine/Source/Runtime/AIModule/Private/BehaviorTree/Decorators/BTDecorator_BlueprintBase.cpp
+++ b/Engine/Source/Runtime/AIModule/Private/BehaviorTree/Decorators/BTDecorator_BlueprintBase.cpp
@@ -178,36 +178,65 @@ void UBTDecorator_BlueprintBase::TickNode(UBehaviorTreeComponent& OwnerComp, uin
 	}
 		
 	// possible this got ticked due to the decorator being configured as an observer
-	if (GetNeedsTickForConditionChecking() && GetShouldAbort(OwnerComp))
+	if (GetNeedsTickForConditionChecking())
 	{
-		OwnerComp.RequestExecution(this);
+		RequestAbort(OwnerComp, EvaluateAbortType(OwnerComp));
 	}
 }
 
-bool UBTDecorator_BlueprintBase::GetShouldAbort(UBehaviorTreeComponent& OwnerComp) const 
+UBTDecorator_BlueprintBase::EAbortType  UBTDecorator_BlueprintBase::EvaluateAbortType(UBehaviorTreeComponent& OwnerComp) const
 {
-	// if there's no condition-checking function implemented we always want to abort on any change
 	if (PerformConditionCheckImplementations == 0)
 	{
-		return true;
+		return EAbortType::Unknown;
+	}
+
+	if (FlowAbortMode == EBTFlowAbortMode::None)
+	{
+		return EAbortType::NoAbort;
 	}
 
 	const bool bIsOnActiveBranch = OwnerComp.IsExecutingBranch(GetMyNode(), GetChildIndex());
 
-	bool bShouldAbort = false;
+	EAbortType AbortType = EAbortType::NoAbort;
 	if (bIsOnActiveBranch)
 	{
-		bShouldAbort = (FlowAbortMode == EBTFlowAbortMode::Self || FlowAbortMode == EBTFlowAbortMode::Both) && CalculateRawConditionValueImpl(OwnerComp) == IsInversed();
+		if ((FlowAbortMode == EBTFlowAbortMode::Self || FlowAbortMode == EBTFlowAbortMode::Both) && CalculateRawConditionValue(OwnerComp, /*NodeMemory*/nullptr) == IsInversed())
+		{
+			AbortType = EAbortType::DeactivateBranch;
+		}
 	}
-	else
+	else 
 	{
-		bShouldAbort = (FlowAbortMode == EBTFlowAbortMode::LowerPriority || FlowAbortMode == EBTFlowAbortMode::Both) && CalculateRawConditionValueImpl(OwnerComp) != IsInversed();
-	}
+		if ((FlowAbortMode == EBTFlowAbortMode::LowerPriority || FlowAbortMode == EBTFlowAbortMode::Both) && CalculateRawConditionValue(OwnerComp, /*NodeMemory*/nullptr) != IsInversed())
+	    {
+			AbortType = EAbortType::ActivateBranch;
+	    }
+    }
+
+	return AbortType;
+}
 
-	return bShouldAbort;
+void UBTDecorator_BlueprintBase::RequestAbort(UBehaviorTreeComponent& OwnerComp, const EAbortType Type)
+{
+	switch (Type)
+	{
+	case EAbortType::NoAbort:
+		// Nothing to abort, continue
+		break;
+	case EAbortType::ActivateBranch:
+		OwnerComp.RequestBranchActivation(*this, false);
+		break;
+	case EAbortType::DeactivateBranch:
+		OwnerComp.RequestBranchDeactivation(*this);
+		break;
+	case EAbortType::Unknown:
+		checkf(false, TEXT("If decorator is active, it must know whether it needs to abort or not"))
+		break;
+	}
 }
 
-bool UBTDecorator_BlueprintBase::CalculateRawConditionValueImpl(UBehaviorTreeComponent& OwnerComp) const
+bool UBTDecorator_BlueprintBase::CalculateRawConditionValue(UBehaviorTreeComponent& OwnerComp, uint8* NodeMemory) const
 {
 	bool CurrentCallResult = false;
 	if (PerformConditionCheckImplementations != 0)
@@ -228,11 +257,6 @@ bool UBTDecorator_BlueprintBase::CalculateRawConditionValueImpl(UBehaviorTreeCom
 	return CurrentCallResult;
 }
 
-bool UBTDecorator_BlueprintBase::CalculateRawConditionValue(UBehaviorTreeComponent& OwnerComp, uint8* NodeMemory) const
-{
-	return CalculateRawConditionValueImpl(OwnerComp);	
-}
-
 bool UBTDecorator_BlueprintBase::IsDecoratorExecutionActive() const
 {
 	UBehaviorTreeComponent* OwnerComp = Cast<UBehaviorTreeComponent>(GetOuter());
@@ -279,10 +303,14 @@ void UBTDecorator_BlueprintBase::DescribeRuntimeValues(const UBehaviorTreeCompon
 
 EBlackboardNotificationResult UBTDecorator_BlueprintBase::OnBlackboardKeyValueChange(const UBlackboardComponent& Blackboard, FBlackboard::FKey ChangedKeyID)
 {
-	UBehaviorTreeComponent* BehaviorComp = (UBehaviorTreeComponent*)Blackboard.GetBrainComponent();
-	if (BehaviorComp && GetShouldAbort(*BehaviorComp))
+	UBehaviorTreeComponent* BehaviorComp = Cast<UBehaviorTreeComponent>(Blackboard.GetBrainComponent());
+	if (BehaviorComp)
 	{
-		BehaviorComp->RequestExecution(this);
+		const EAbortType Type = EvaluateAbortType(*BehaviorComp);
+		if (Type != EAbortType::Unknown)
+		{
+			RequestAbort(*BehaviorComp, Type);
+		}
 	}
 	return BehaviorComp ? EBlackboardNotificationResult::ContinueObserving : EBlackboardNotificationResult::RemoveObserver;
 }
