Move forward declarations at the beginning of the file

diff --git a/Engine/Source/Runtime/AIModule/Classes/AIController.h b/Engine/Source/Runtime/AIModule/Classes/AIController.h
--- a/Engine/Source/Runtime/AIModule/Classes/AIController.h
+++ b/Engine/Source/Runtime/AIModule/Classes/AIController.h
@@ -33,7 +33,7 @@ class UGameplayTasksComponent;
 class UPathFollowingComponent;
 class UPawnAction;
 class UPawnActionsComponent;
-struct FVisualLogEntry;
+
 namespace EPathFollowingRequestResult {	enum Type : int; }
 namespace EPathFollowingResult { enum Type : int; }
 namespace EPathFollowingStatus { enum Type : int; }
@@ -42,6 +42,8 @@ namespace EPathFollowingStatus { enum Type : int; }
 struct FVisualLogEntry;
 #endif // ENABLE_VISUAL_LOG
 struct FPathFindingQuery;
+struct FPathFollowingRequestResult;
+struct FPathFollowingResult;
 
 DECLARE_DYNAMIC_MULTICAST_DELEGATE_TwoParams(FAIMoveCompletedSignature, FAIRequestID, RequestID, EPathFollowingResult::Type, Result);
 
@@ -197,7 +199,7 @@ public:
 	 *  @param OutPath - optional output param, filled in with assigned path
 	 *  @return struct holding MoveId and enum code
 	 */
-	virtual struct FPathFollowingRequestResult MoveTo(const FAIMoveRequest& MoveRequest, FNavPathSharedPtr* OutPath = nullptr);
+	virtual FPathFollowingRequestResult MoveTo(const FAIMoveRequest& MoveRequest, FNavPathSharedPtr* OutPath = nullptr);
 
 	/** Passes move request and path object to path following */
 	virtual FAIRequestID RequestMove(const FAIMoveRequest& MoveRequest, FNavPathSharedPtr Path);
@@ -228,7 +230,7 @@ public:
 	virtual void StopMovement() override;
 
 	/** Called on completing current movement request */
-	virtual void OnMoveCompleted(FAIRequestID RequestID, const struct FPathFollowingResult& Result);
+	virtual void OnMoveCompleted(FAIRequestID RequestID, const FPathFollowingResult& Result);
 
 	UE_DEPRECATED_FORGAME(4.13, "This function is now deprecated, please use version with EPathFollowingResultDetails parameter.")
 	virtual void OnMoveCompleted(FAIRequestID RequestID, EPathFollowingResult::Type Result);
