Fix BlackboardComponent::IsValidKey returning false when KeyID is larger than the number of keys in the child blackboard. Compare with GetNumKeys to account for parent key count.

#tests PIE

diff --git a/Engine/Source/Runtime/AIModule/Classes/BehaviorTree/BlackboardComponent.h b/Engine/Source/Runtime/AIModule/Classes/BehaviorTree/BlackboardComponent.h
--- a/Engine/Source/Runtime/AIModule/Classes/BehaviorTree/BlackboardComponent.h
+++ b/Engine/Source/Runtime/AIModule/Classes/BehaviorTree/BlackboardComponent.h
@@ -204,7 +204,7 @@ public:
 	FORCEINLINE const uint8* GetKeyRawData(FBlackboard::FKey KeyID) const { return ValueMemory.Num() && ValueOffsets.IsValidIndex(KeyID) ? (ValueMemory.GetData() + ValueOffsets[KeyID]) : NULL; }
 
 	PRAGMA_DISABLE_DEPRECATION_WARNINGS // re BlackboardAsset
-	FORCEINLINE bool IsValidKey(FBlackboard::FKey KeyID) const { check(BlackboardAsset); return KeyID != FBlackboard::InvalidKey && BlackboardAsset->Keys.IsValidIndex(KeyID); }
+	FORCEINLINE bool IsValidKey(FBlackboard::FKey KeyID) const { check(BlackboardAsset); return BlackboardAsset->IsValidKey(KeyID); }
 	PRAGMA_ENABLE_DEPRECATION_WARNINGS // re BlackboardAsset
 
 	/** compares blackboard's values under specified keys */
diff --git a/Engine/Source/Runtime/AIModule/Classes/BehaviorTree/BlackboardData.h b/Engine/Source/Runtime/AIModule/Classes/BehaviorTree/BlackboardData.h
--- a/Engine/Source/Runtime/AIModule/Classes/BehaviorTree/BlackboardData.h
+++ b/Engine/Source/Runtime/AIModule/Classes/BehaviorTree/BlackboardData.h
@@ -86,6 +86,8 @@ public:
 
 	FORCEINLINE FBlackboard::FKey GetFirstKeyID() const { return FirstKeyID; }
 
+	bool IsValidKey(FBlackboard::FKey KeyID) const { return KeyID != FBlackboard::InvalidKey && KeyID < GetNumKeys(); }
+
 	/** @return key data */
 	const FBlackboardEntry* GetKey(FBlackboard::FKey KeyID) const;
 
