Blackboard component: Support const NotifyOwner
#tests Used locally with a const UObject

diff --git a/Engine/Source/Runtime/AIModule/Classes/BehaviorTree/BlackboardComponent.h b/Engine/Source/Runtime/AIModule/Classes/BehaviorTree/BlackboardComponent.h
--- a/Engine/Source/Runtime/AIModule/Classes/BehaviorTree/BlackboardComponent.h
+++ b/Engine/Source/Runtime/AIModule/Classes/BehaviorTree/BlackboardComponent.h
@@ -70,13 +70,13 @@ public:
 	bool HasValidAsset() const;
 
 	/** register observer for blackboard key */
-	FDelegateHandle RegisterObserver(FBlackboard::FKey KeyID, UObject* NotifyOwner, FOnBlackboardChangeNotification ObserverDelegate);
+	FDelegateHandle RegisterObserver(FBlackboard::FKey KeyID, const UObject* NotifyOwner, FOnBlackboardChangeNotification ObserverDelegate);
 
 	/** unregister observer from blackboard key */
 	void UnregisterObserver(FBlackboard::FKey KeyID, FDelegateHandle ObserverHandle);
 
 	/** unregister all observers associated with given owner */
-	void UnregisterObserversFrom(UObject* NotifyOwner);
+	void UnregisterObserversFrom(const UObject* NotifyOwner);
 
 	/** pause observer change notifications, any new ones will be added to a queue */
 	void PauseObserverNotifications();
@@ -97,7 +97,7 @@ public:
 	bool InitializeBlackboard(UBlackboardData& NewAsset);
 	
 	/** @return true if component can be used with specified blackboard asset */
-	bool IsCompatibleWith(UBlackboardData* TestAsset) const;
+	bool IsCompatibleWith(const UBlackboardData* TestAsset) const;
 
 	UFUNCTION(BlueprintCallable, Category="AI|Components|Blackboard")
 	UObject* GetValueAsObject(const FName& KeyName) const;
@@ -276,7 +276,7 @@ protected:
 	mutable TMultiMap<uint8, FOnBlackboardChangeNotificationInfo> Observers;
 	
 	/** observers registered from owner objects */
-	mutable TMultiMap<UObject*, FDelegateHandle> ObserverHandles;
+	mutable TMultiMap<const UObject*, FDelegateHandle> ObserverHandles;
 
 	/** queued key change notification, will be processed on ResumeUpdates call */
 	mutable TArray<uint8> QueuedUpdates;
@@ -436,7 +436,7 @@ public:
 	FBBKeyCachedAccessor() : BBKey(FBlackboard::InvalidKey), CachedValue(TBlackboardKey::InvalidValue)
 	{}
 
-	FBBKeyCachedAccessor(UBlackboardComponent& BBComponent, FBlackboard::FKey InBBKey)
+	FBBKeyCachedAccessor(const UBlackboardComponent& BBComponent, FBlackboard::FKey InBBKey)
 	{
 		ensure(InBBKey != FBlackboard::InvalidKey);
 		if (ensure(BBComponent.IsKeyOfType<TBlackboardKey>(InBBKey)))
diff --git a/Engine/Source/Runtime/AIModule/Private/BehaviorTree/BlackboardComponent.cpp b/Engine/Source/Runtime/AIModule/Private/BehaviorTree/BlackboardComponent.cpp
--- a/Engine/Source/Runtime/AIModule/Private/BehaviorTree/BlackboardComponent.cpp
+++ b/Engine/Source/Runtime/AIModule/Private/BehaviorTree/BlackboardComponent.cpp
@@ -294,7 +294,7 @@ FBlackboard::FKey UBlackboardComponent::GetKeyID(const FName& KeyName) const
 
 TSubclassOf<UBlackboardKeyType> UBlackboardComponent::GetKeyType(FBlackboard::FKey KeyID) const
 {
-	return BlackboardAsset ? BlackboardAsset->GetKeyType(KeyID) : NULL;
+	return BlackboardAsset ? BlackboardAsset->GetKeyType(KeyID) : nullptr;
 }
 
 bool UBlackboardComponent::IsKeyInstanceSynced(FBlackboard::FKey KeyID) const
@@ -307,7 +307,7 @@ int32 UBlackboardComponent::GetNumKeys() const
 	return BlackboardAsset ? BlackboardAsset->GetNumKeys() : 0;
 }
 
-FDelegateHandle UBlackboardComponent::RegisterObserver(FBlackboard::FKey KeyID, UObject* NotifyOwner, FOnBlackboardChangeNotification ObserverDelegate)
+FDelegateHandle UBlackboardComponent::RegisterObserver(FBlackboard::FKey KeyID, const UObject* NotifyOwner, FOnBlackboardChangeNotification ObserverDelegate)
 {
 	for (auto It = Observers.CreateConstKeyIterator(KeyID); It; ++It)
 	{
@@ -354,7 +354,7 @@ void UBlackboardComponent::UnregisterObserver(FBlackboard::FKey KeyID, FDelegate
 	}
 }
 
-void UBlackboardComponent::UnregisterObserversFrom(UObject* NotifyOwner)
+void UBlackboardComponent::UnregisterObserversFrom(const UObject* NotifyOwner)
 {
 	for (auto It = ObserverHandles.CreateKeyIterator(NotifyOwner); It; ++It)
 	{
@@ -482,7 +482,7 @@ void UBlackboardComponent::NotifyObservers(FBlackboard::FKey KeyID) const
 	}
 }
 
-bool UBlackboardComponent::IsCompatibleWith(UBlackboardData* TestAsset) const
+bool UBlackboardComponent::IsCompatibleWith(const UBlackboardData* TestAsset) const
 {
 	for (UBlackboardData* It = BlackboardAsset; It; It = It->Parent)
 	{
