AI Truncation warning fixes.

This is the majority of the fixes however a few systems have been left for individual CLs.

diff --git a/Engine/Plugins/Runtime/MassGameplay/Source/MassReplication/Private/MassReplicationTransformHandlers.cpp b/Engine/Plugins/Runtime/MassGameplay/Source/MassReplication/Private/MassReplicationTransformHandlers.cpp
--- a/Engine/Plugins/Runtime/MassGameplay/Source/MassReplication/Private/MassReplicationTransformHandlers.cpp
+++ b/Engine/Plugins/Runtime/MassGameplay/Source/MassReplication/Private/MassReplicationTransformHandlers.cpp
@@ -20,8 +20,6 @@ void FMassReplicationProcessorPositionYawHandler::AddEntity(const int32 EntityId
 	const FTransformFragment& TransformFragment = TransformList[EntityIdx];
 	InOutReplicatedPositionYawData.SetPosition(TransformFragment.GetTransform().GetLocation());
 
-	if (const TOptional<float> Yaw = UE::AI::GetYawFromQuaternion(TransformFragment.GetTransform().GetRotation()))
-	{
-		InOutReplicatedPositionYawData.SetYaw(Yaw.GetValue());
-	}
+	const FRotator::FReal Yaw = FMath::DegreesToRadians(TransformFragment.GetTransform().Rotator().Yaw);
+	InOutReplicatedPositionYawData.SetYaw(static_cast<float>(Yaw));
 }
diff --git a/Engine/Plugins/Runtime/MassGameplay/Source/MassReplication/Public/MassReplicationTransformHandlers.h b/Engine/Plugins/Runtime/MassGameplay/Source/MassReplication/Public/MassReplicationTransformHandlers.h
--- a/Engine/Plugins/Runtime/MassGameplay/Source/MassReplication/Public/MassReplicationTransformHandlers.h
+++ b/Engine/Plugins/Runtime/MassGameplay/Source/MassReplication/Public/MassReplicationTransformHandlers.h
@@ -122,14 +122,13 @@ void TMassClientBubbleTransformHandler<AgentArrayItem>::SetBubblePositionYawFrom
 		bMarkDirty = true;
 	}
 
-	if (const TOptional<float> Yaw = UE::AI::GetYawFromQuaternion(Transform.GetRotation()))
+	const float Yaw = static_cast<float>(FMath::DegreesToRadians(Transform.GetRotation().Rotator().Yaw));
+
+	// Only update the Yaw and mark the item as dirty if it has changed more than the tolerance
+	if (FMath::Abs(FMath::FindDeltaAngleRadians(Yaw, ReplicatedPositionYaw.GetYaw())) > UE::Mass::Replication::YawReplicateTolerance)
 	{
-		// Only update the Yaw and mark the item as dirty if it has changed more than the tolerance
-		if (FMath::Abs(FMath::FindDeltaAngleRadians(Yaw.GetValue(), ReplicatedPositionYaw.GetYaw())) > UE::Mass::Replication::YawReplicateTolerance)
-		{
-			ReplicatedPositionYaw.SetYaw(Yaw.GetValue());
-			bMarkDirty = true;
-		}
+		ReplicatedPositionYaw.SetYaw(Yaw);
+		bMarkDirty = true;
 	}
 
 	if (bMarkDirty)
diff --git a/Engine/Source/Developer/AITestSuite/Private/Tests/AITestsCommon.cpp b/Engine/Source/Developer/AITestSuite/Private/Tests/AITestsCommon.cpp
--- a/Engine/Source/Developer/AITestSuite/Private/Tests/AITestsCommon.cpp
+++ b/Engine/Source/Developer/AITestSuite/Private/Tests/AITestsCommon.cpp
@@ -27,7 +27,7 @@ namespace FAITestHelpers
 
 bool FAITestCommand_WaitSeconds::Update()
 {
-	float NewTime = FPlatformTime::Seconds();
+	const double NewTime = FPlatformTime::Seconds();
 	if (NewTime - StartTime >= Duration)
 	{
 		return true;
diff --git a/Engine/Source/Runtime/AIModule/Classes/BehaviorTree/BTCompositeNode.h b/Engine/Source/Runtime/AIModule/Classes/BehaviorTree/BTCompositeNode.h
--- a/Engine/Source/Runtime/AIModule/Classes/BehaviorTree/BTCompositeNode.h
+++ b/Engine/Source/Runtime/AIModule/Classes/BehaviorTree/BTCompositeNode.h
@@ -183,7 +183,7 @@ class AIMODULE_API UBTCompositeNode : public UBTNode
 	uint16 GetBranchExecutionIndex(uint16 NodeInBranchIdx) const;
 
 	/** is child execution allowed by decorators? */
-	bool DoDecoratorsAllowExecution(UBehaviorTreeComponent& OwnerComp, int32 InstanceIdx, int32 ChildIdx) const;
+	bool DoDecoratorsAllowExecution(UBehaviorTreeComponent& OwnerComp, const int32 InstanceIdx, const int32 ChildIdx) const;
 
 	bool IsApplyingDecoratorScope() const;
 
diff --git a/Engine/Source/Runtime/AIModule/Classes/BehaviorTree/BehaviorTreeComponent.h b/Engine/Source/Runtime/AIModule/Classes/BehaviorTree/BehaviorTreeComponent.h
--- a/Engine/Source/Runtime/AIModule/Classes/BehaviorTree/BehaviorTreeComponent.h
+++ b/Engine/Source/Runtime/AIModule/Classes/BehaviorTree/BehaviorTreeComponent.h
@@ -143,9 +143,9 @@ public:
 	void RestartTree();
 
 	/** request execution change */
-	void RequestExecution(const UBTCompositeNode* RequestedOn, int32 InstanceIdx, 
-		const UBTNode* RequestedBy, int32 RequestedByChildIndex,
-		EBTNodeResult::Type ContinueWithResult, bool bStoreForDebugger = true);
+	void RequestExecution(const UBTCompositeNode* RequestedOn, const int32 InstanceIdx, 
+		const UBTNode* RequestedBy, const int32 RequestedByChildIndex,
+		const EBTNodeResult::Type ContinueWithResult, bool bStoreForDebugger = true);
 
 	/** replaced by the RequestBranchEvaluation from decorator*/
 	void RequestExecution(const UBTDecorator* RequestedBy) { check(RequestedBy); RequestBranchEvaluation(*RequestedBy); }
@@ -254,7 +254,7 @@ public:
 
 	/** @return the cooldown tag end time, 0.0f if CooldownTag is not found */
 	UFUNCTION(BlueprintCallable, Category = "AI|Logic")
-	float GetTagCooldownEndTime(FGameplayTag CooldownTag) const;
+	double GetTagCooldownEndTime(FGameplayTag CooldownTag) const;
 
 	/** add to the cooldown tag's duration */
 	UFUNCTION(BlueprintCallable, Category = "AI|Logic")
@@ -330,7 +330,7 @@ protected:
 	TMultiMap<FBTNodeIndex,FAIMessageObserverHandle> TaskMessageObservers;
 
 	/** behavior cooldowns mapped by tag to last time it was set */
-	TMap<FGameplayTag, float> CooldownTagsMap;
+	TMap<FGameplayTag, double> CooldownTagsMap;
 
 #if USE_BEHAVIORTREE_DEBUGGER
 	/** search flow for debugger */
@@ -390,9 +390,6 @@ protected:
 	/** copy memory block from persistent memory to running instances (rollback) */
 	void CopyInstanceMemoryFromPersistent();
 
-	/** find next task to execute */
-	UBTTaskNode* FindNextTask(UBTCompositeNode* ParentNode, uint16 ParentInstanceIdx, EBTNodeResult::Type LastResult);
-
 	/** called when tree runs out of nodes to execute */
 	virtual void OnTreeFinished();
 
@@ -518,11 +515,11 @@ protected:
 	/** Used to tell tickmanager that we want interval ticking */
 	bool bTickedOnce = false;
 	/** Predicted next DeltaTime*/
-	float NextTickDeltaTime = 0.0f;
+	float NextTickDeltaTime = 0.;
 	/** Accumulated DeltaTime if ticked more than predicted next delta time */
 	float AccumulatedTickDeltaTime = 0.0f;
 	/** GameTime of the last DeltaTime request, used for debugging to output warnings about ticking */
-	float LastRequestedDeltaTimeGameTime = 0;
+	double LastRequestedDeltaTimeGameTime = 0.;
 
 #if CSV_PROFILER
 	/** CSV tick stat name. Can be changed but must point to a static string */
diff --git a/Engine/Source/Runtime/AIModule/Classes/BehaviorTree/BehaviorTreeManager.h b/Engine/Source/Runtime/AIModule/Classes/BehaviorTree/BehaviorTreeManager.h
--- a/Engine/Source/Runtime/AIModule/Classes/BehaviorTree/BehaviorTreeManager.h
+++ b/Engine/Source/Runtime/AIModule/Classes/BehaviorTree/BehaviorTreeManager.h
@@ -42,7 +42,7 @@ class AIMODULE_API UBehaviorTreeManager : public UObject
 	bool LoadTree(UBehaviorTree& Asset, UBTCompositeNode*& Root, uint16& InstanceMemorySize);
 
 	/** get aligned memory size */
-	static int32 GetAlignedDataSize(int32 Size);
+	static uint16 GetAlignedDataSize(uint16 Size);
 
 	/** helper function for sorting and aligning node memory */
 	static void InitializeMemoryHelper(const TArray<UBTDecorator*>& Nodes, TArray<uint16>& MemoryOffsets, int32& MemorySize, bool bForceInstancing = false);
diff --git a/Engine/Source/Runtime/AIModule/Classes/BehaviorTree/BehaviorTreeTypes.h b/Engine/Source/Runtime/AIModule/Classes/BehaviorTree/BehaviorTreeTypes.h
--- a/Engine/Source/Runtime/AIModule/Classes/BehaviorTree/BehaviorTreeTypes.h
+++ b/Engine/Source/Runtime/AIModule/Classes/BehaviorTree/BehaviorTreeTypes.h
@@ -257,7 +257,7 @@ struct FBehaviorTreeDebuggerInstance
 /** debugger data about current execution step */
 struct FBehaviorTreeExecutionStep
 {
-	FBehaviorTreeExecutionStep() : TimeStamp(0.f), ExecutionStepId(InvalidExecutionId) {}
+	FBehaviorTreeExecutionStep() : TimeStamp(0.), ExecutionStepId(InvalidExecutionId) {}
 
 	/** subtree instance stack */
 	TArray<FBehaviorTreeDebuggerInstance> InstanceStack;
@@ -266,7 +266,7 @@ struct FBehaviorTreeExecutionStep
 	TMap<FName, FString> BlackboardValues;
 
 	/** Game world's time stamp of this step */
-	float TimeStamp;
+	double TimeStamp;
 
 	static constexpr int32 InvalidExecutionId = -1;
 
@@ -417,17 +417,19 @@ private:
 
 struct FBTNodeIndex
 {
+	static constexpr uint16 InvalidIndex = TNumericLimits<uint16>::Max(); // (This is also the same as INDEX_NONE assigned to IndexType!)
+
 	/** index of instance of stack */
 	uint16 InstanceIndex;
 
 	/** execution index within instance */
 	uint16 ExecutionIndex;
 
-	FBTNodeIndex() : InstanceIndex(MAX_uint16), ExecutionIndex(MAX_uint16) {}
+	FBTNodeIndex() : InstanceIndex(InvalidIndex), ExecutionIndex(InvalidIndex) {}
 	FBTNodeIndex(uint16 InInstanceIndex, uint16 InExecutionIndex) : InstanceIndex(InInstanceIndex), ExecutionIndex(InExecutionIndex) {}
 
 	bool TakesPriorityOver(const FBTNodeIndex& Other) const;
-	bool IsSet() const { return InstanceIndex < MAX_uint16; }
+	bool IsSet() const { return InstanceIndex < InvalidIndex; }
 
 	FORCEINLINE bool operator==(const FBTNodeIndex& Other) const { return Other.ExecutionIndex == ExecutionIndex && Other.InstanceIndex == InstanceIndex; }
 	FORCEINLINE bool operator!=(const FBTNodeIndex& Other) const { return !operator==(Other); }
diff --git a/Engine/Source/Runtime/AIModule/Classes/BehaviorTree/Decorators/BTDecorator_Cooldown.h b/Engine/Source/Runtime/AIModule/Classes/BehaviorTree/Decorators/BTDecorator_Cooldown.h
--- a/Engine/Source/Runtime/AIModule/Classes/BehaviorTree/Decorators/BTDecorator_Cooldown.h
+++ b/Engine/Source/Runtime/AIModule/Classes/BehaviorTree/Decorators/BTDecorator_Cooldown.h
@@ -9,7 +9,7 @@
 
 struct FBTCooldownDecoratorMemory
 {
-	float LastUseTimestamp;
+	double LastUseTimestamp;
 	uint8 bRequestedRestart : 1;
 };
 
diff --git a/Engine/Source/Runtime/AIModule/Classes/BehaviorTree/Decorators/BTDecorator_IsAtLocation.h b/Engine/Source/Runtime/AIModule/Classes/BehaviorTree/Decorators/BTDecorator_IsAtLocation.h
--- a/Engine/Source/Runtime/AIModule/Classes/BehaviorTree/Decorators/BTDecorator_IsAtLocation.h
+++ b/Engine/Source/Runtime/AIModule/Classes/BehaviorTree/Decorators/BTDecorator_IsAtLocation.h
@@ -49,5 +49,5 @@ class AIMODULE_API UBTDecorator_IsAtLocation : public UBTDecorator_BlackboardBas
 #endif // WITH_EDITOR
 
 protected:
-	float GetGeometricDistanceSquared(const FVector& A, const FVector& B) const;
+	FVector::FReal GetGeometricDistanceSquared(const FVector& A, const FVector& B) const;
 };
diff --git a/Engine/Source/Runtime/AIModule/Classes/BehaviorTree/Decorators/BTDecorator_Loop.h b/Engine/Source/Runtime/AIModule/Classes/BehaviorTree/Decorators/BTDecorator_Loop.h
--- a/Engine/Source/Runtime/AIModule/Classes/BehaviorTree/Decorators/BTDecorator_Loop.h
+++ b/Engine/Source/Runtime/AIModule/Classes/BehaviorTree/Decorators/BTDecorator_Loop.h
@@ -10,7 +10,7 @@ struct FBTLoopDecoratorMemory
 {
 	int32 SearchId;
 	uint8 RemainingExecutions;
-	float TimeStarted;
+	double TimeStarted;
 };
 
 /**
@@ -23,7 +23,7 @@ class AIMODULE_API UBTDecorator_Loop : public UBTDecorator
 	GENERATED_UCLASS_BODY()
 
 	/** number of executions */
-	UPROPERTY(Category=Decorator, EditAnywhere, meta=(EditCondition="!bInfiniteLoop", ClampMin="1"))
+	UPROPERTY(Category=Decorator, EditAnywhere, meta=(EditCondition="!bInfiniteLoop", ClampMin="1", ClampMax="255"))
 	int32 NumLoops;
 
 	/** infinite loop */
diff --git a/Engine/Source/Runtime/AIModule/Classes/EnvironmentQuery/EQSRenderingComponent.h b/Engine/Source/Runtime/AIModule/Classes/EnvironmentQuery/EQSRenderingComponent.h
--- a/Engine/Source/Runtime/AIModule/Classes/EnvironmentQuery/EQSRenderingComponent.h
+++ b/Engine/Source/Runtime/AIModule/Classes/EnvironmentQuery/EQSRenderingComponent.h
@@ -35,7 +35,7 @@ private:
 	const IEQSQueryResultSourceInterface* QueryDataSource;
 	uint32 bDrawOnlyWhenSelected : 1;
 
-	static const FVector ItemDrawRadius;
+	static const FVector3f ItemDrawRadius;
 
 	bool SafeIsActorSelected() const;
 };
diff --git a/Engine/Source/Runtime/AIModule/Classes/EnvironmentQuery/EnvQueryManager.h b/Engine/Source/Runtime/AIModule/Classes/EnvironmentQuery/EnvQueryManager.h
--- a/Engine/Source/Runtime/AIModule/Classes/EnvironmentQuery/EnvQueryManager.h
+++ b/Engine/Source/Runtime/AIModule/Classes/EnvironmentQuery/EnvQueryManager.h
@@ -148,7 +148,7 @@ struct AIMODULE_API FEQSDebugger
 	struct FEnvQueryInfo
 	{
 		TSharedPtr<FEnvQueryInstance> Instance;
-		float Timestamp;
+		double Timestamp;
 	};
 
 	struct FStatsInfo
diff --git a/Engine/Source/Runtime/AIModule/Classes/EnvironmentQuery/EnvQueryTypes.h b/Engine/Source/Runtime/AIModule/Classes/EnvironmentQuery/EnvQueryTypes.h
--- a/Engine/Source/Runtime/AIModule/Classes/EnvironmentQuery/EnvQueryTypes.h
+++ b/Engine/Source/Runtime/AIModule/Classes/EnvironmentQuery/EnvQueryTypes.h
@@ -820,7 +820,7 @@ struct AIMODULE_API FEnvQueryInstance : public FEnvQueryResult
 	double StartTime;
 
 	/** time spent executing this query */
-	float TotalExecutionTime;
+	double TotalExecutionTime;
 
 	/** if > 0 then it's how much time query has for performing current step */
 	float CurrentStepTimeLimit;
diff --git a/Engine/Source/Runtime/AIModule/Classes/EnvironmentQuery/Generators/EnvQueryGenerator_ProjectedPoints.h b/Engine/Source/Runtime/AIModule/Classes/EnvironmentQuery/Generators/EnvQueryGenerator_ProjectedPoints.h
--- a/Engine/Source/Runtime/AIModule/Classes/EnvironmentQuery/Generators/EnvQueryGenerator_ProjectedPoints.h
+++ b/Engine/Source/Runtime/AIModule/Classes/EnvironmentQuery/Generators/EnvQueryGenerator_ProjectedPoints.h
@@ -20,7 +20,7 @@ class AIMODULE_API UEnvQueryGenerator_ProjectedPoints : public UEnvQueryGenerato
 
 	struct FSortByHeight
 	{
-		float OriginalZ;
+		FVector::FReal OriginalZ;
 
 		FSortByHeight(const FVector& OriginalPt) : OriginalZ(OriginalPt.Z) {}
 
diff --git a/Engine/Source/Runtime/AIModule/Classes/Navigation/NavLocalGridData.h b/Engine/Source/Runtime/AIModule/Classes/Navigation/NavLocalGridData.h
--- a/Engine/Source/Runtime/AIModule/Classes/Navigation/NavLocalGridData.h
+++ b/Engine/Source/Runtime/AIModule/Classes/Navigation/NavLocalGridData.h
@@ -99,8 +99,8 @@ struct AIMODULE_API FNavLocalGridData : public TSimpleCellGrid<uint8, MAX_uint8>
 
 protected:
 	friend class UNavLocalGridManager;
-	TArray<float> CellZ;
-	float LastAccessTime;
+	TArray<FVector::FReal> CellZ;
+	double LastAccessTime;
 
 	/** convert PathIndices into pruned PathCoords */
 	void PostProcessPath(const FIntVector& StartCoords, const FIntVector& EndCoords, const TArray<int32>& PathIndices, TArray<FIntVector>& PathCoords) const;
diff --git a/Engine/Source/Runtime/AIModule/Classes/Navigation/PathFollowingComponent.h b/Engine/Source/Runtime/AIModule/Classes/Navigation/PathFollowingComponent.h
--- a/Engine/Source/Runtime/AIModule/Classes/Navigation/PathFollowingComponent.h
+++ b/Engine/Source/Runtime/AIModule/Classes/Navigation/PathFollowingComponent.h
@@ -324,7 +324,7 @@ class AIMODULE_API UPathFollowingComponent : public UActorComponent, public IAIR
 
 	/** @returns estimated cost of unprocessed path segments
 	 *	@NOTE 0 means, that component is following final path segment or doesn't move */
-	float GetRemainingPathCost() const;
+	FVector::FReal GetRemainingPathCost() const;
 	
 	/** Returns current location on navigation data */
 	FNavLocation GetCurrentNavLocation() const;
diff --git a/Engine/Source/Runtime/AIModule/Classes/Perception/AIPerceptionSystem.h b/Engine/Source/Runtime/AIModule/Classes/Perception/AIPerceptionSystem.h
--- a/Engine/Source/Runtime/AIModule/Classes/Perception/AIPerceptionSystem.h
+++ b/Engine/Source/Runtime/AIModule/Classes/Perception/AIPerceptionSystem.h
@@ -62,7 +62,7 @@ protected:
 
 	struct FDelayedStimulus
 	{
-		float DeliveryTimestamp;
+		double DeliveryTimestamp;
 		FPerceptionListenerID ListenerId;
 		TWeakObjectPtr<AActor> Instigator;
 		FAIStimulus Stimulus;
@@ -209,10 +209,10 @@ protected:
 	virtual void StartPlay();
 
 	/** Timestamp of the next stimuli aging */
-	float NextStimuliAgingTick;
+	double NextStimuliAgingTick;
 private:
 	/** cached world's timestamp */
-	float CurrentTime;
+	double CurrentTime;
 };
 
 //////////////////////////////////////////////////////////////////////////
diff --git a/Engine/Source/Runtime/AIModule/Classes/Perception/AIPerceptionTypes.h b/Engine/Source/Runtime/AIModule/Classes/Perception/AIPerceptionTypes.h
--- a/Engine/Source/Runtime/AIModule/Classes/Perception/AIPerceptionTypes.h
+++ b/Engine/Source/Runtime/AIModule/Classes/Perception/AIPerceptionTypes.h
@@ -221,7 +221,7 @@ struct AIMODULE_API FAISenseAffiliationFilter
 	UPROPERTY(EditAnywhere, BlueprintReadOnly, Category = "Sense")
 	uint32 bDetectFriendlies : 1;
 	
-	uint8 GetAsFlags() const { return (bDetectEnemies << ETeamAttitude::Hostile) | (bDetectNeutrals << ETeamAttitude::Neutral) | (bDetectFriendlies << ETeamAttitude::Friendly); }
+	uint8 GetAsFlags() const { return IntCastChecked<uint8>((bDetectEnemies << ETeamAttitude::Hostile) | (bDetectNeutrals << ETeamAttitude::Neutral) | (bDetectFriendlies << ETeamAttitude::Friendly)); }
 	FORCEINLINE bool ShouldDetectAll() const { return (bDetectEnemies && bDetectNeutrals && bDetectFriendlies); }
 
 	static FORCEINLINE uint8 DetectAllFlags() { return (1 << ETeamAttitude::Hostile) | (1 << ETeamAttitude::Neutral) | (1 << ETeamAttitude::Friendly); }
diff --git a/Engine/Source/Runtime/AIModule/Private/AIController.cpp b/Engine/Source/Runtime/AIModule/Private/AIController.cpp
--- a/Engine/Source/Runtime/AIModule/Private/AIController.cpp
+++ b/Engine/Source/Runtime/AIModule/Private/AIController.cpp
@@ -347,7 +347,7 @@ bool AAIController::LineOfSightTo(const AActor* Other, FVector ViewPoint, bool b
 	}
 
 	const FVector OtherActorLocation = Other->GetActorLocation();
-	const float DistSq = (OtherActorLocation - ViewPoint).SizeSquared();
+	const FVector::FReal DistSq = (OtherActorLocation - ViewPoint).SizeSquared();
 	if (DistSq > FARSIGHTTHRESHOLDSQUARED)
 	{
 		return false;
@@ -386,11 +386,11 @@ bool AAIController::LineOfSightTo(const AActor* Other, FVector ViewPoint, bool b
 		Points[3] = OtherActorLocation + FVector(OtherRadius, -1 * OtherRadius, 0);
 		int32 IndexMin = 0;
 		int32 IndexMax = 0;
-		float CurrentMax = (Points[0] - ViewPoint).SizeSquared();
-		float CurrentMin = CurrentMax;
+		FVector::FReal CurrentMax = (Points[0] - ViewPoint).SizeSquared();
+		FVector::FReal CurrentMin = CurrentMax;
 		for (int32 PointIndex = 1; PointIndex<4; PointIndex++)
 		{
-			const float NextSize = (Points[PointIndex] - ViewPoint).SizeSquared();
+			const FVector::FReal NextSize = (Points[PointIndex] - ViewPoint).SizeSquared();
 			if (NextSize > CurrentMin)
 			{
 				CurrentMin = NextSize;
diff --git a/Engine/Source/Runtime/AIModule/Private/AIHelpers.cpp b/Engine/Source/Runtime/AIModule/Private/AIHelpers.cpp
--- a/Engine/Source/Runtime/AIModule/Private/AIHelpers.cpp
+++ b/Engine/Source/Runtime/AIModule/Private/AIHelpers.cpp
@@ -10,14 +10,14 @@ namespace FAISystem
 	FVector FindClosestLocation(const FVector& Origin, const TArray<FVector>& Locations)
 	{
 		FVector BestLocation = FAISystem::InvalidLocation;
-		float BestDistance = FLT_MAX;
+		FVector::FReal BestDistanceSq = TNumericLimits<FVector::FReal>::Max();
 
 		for (const FVector& Candidate : Locations)
 		{
-			const float CurrentDistanceSq = FVector::DistSquared(Origin, Candidate);
-			if (CurrentDistanceSq < BestDistance)
+			const FVector::FReal CurrentDistanceSq = FVector::DistSquared(Origin, Candidate);
+			if (CurrentDistanceSq < BestDistanceSq)
 			{
-				BestDistance = CurrentDistanceSq;
+				BestDistanceSq = CurrentDistanceSq;
 				BestLocation = Candidate;
 			}
 		}
@@ -61,7 +61,7 @@ namespace FAISystem
 	{
 		const FVector BaseLocation = FMath::IsNearlyZero(ConeDirectionBackwardOffset) ? StartLocation : StartLocation - ConeDirectionNormal * ConeDirectionBackwardOffset;
 		const FVector ActorToTarget = TargetLocation - BaseLocation;
-		const float DistToTargetSq = ActorToTarget.SizeSquared();
+		const FVector::FReal DistToTargetSq = ActorToTarget.SizeSquared();
 		if (DistToTargetSq <= FarClippingRadiusSq && DistToTargetSq >= NearClippingRadiusSq)
 		{
 			// Will return true if squared distance to Target is smaller than SMALL_NUMBER
@@ -86,30 +86,34 @@ namespace UE::AI
 TOptional<float> GetYawFromVector(const FVector& Vector)
 {
 	const FVector2D Vector2D(Vector);
-	const float SizeSquared = Vector2D.SizeSquared();
+	const FVector::FReal SizeSquared = Vector2D.SizeSquared();
 	if (SizeSquared <= SMALL_NUMBER)
 	{
 		return TOptional<float>();
 	}
 
-	const FVector2D ForwardVector(1.0f, 0.0f);
+	const FVector2D ForwardVector(1., 0.);
 	const FVector2D NormVector2D = Vector2D * FMath::InvSqrt(SizeSquared);
-	const float DotVal = ForwardVector | NormVector2D;
-	const float ClampedDot = FMath::Clamp(DotVal, -1.0f, 1.0f);
-	const float Direction = ForwardVector ^ NormVector2D;
-	const float Sign = Direction < 0.0f ? -1.0f : 1.0f;
-	const float Yaw = Sign * FMath::Acos(ClampedDot);
-	return Yaw;
+	const FVector::FReal DotVal = ForwardVector | NormVector2D;
+	const FVector::FReal ClampedDot = FMath::Clamp(DotVal, -1., 1.);
+	const FVector::FReal Direction = ForwardVector ^ NormVector2D;
+	const FVector::FReal Sign = Direction < 0. ? -1. : 1.;
+	const FVector::FReal Yaw = Sign * FMath::Acos(ClampedDot);
+	return static_cast<float>(Yaw);
 }
 
 TOptional<float> GetYawFromRotator(const FRotator& Rotator)
 {
+	PRAGMA_DISABLE_DEPRECATION_WARNINGS
 	return GetYawFromVector(Rotator.Vector());
+	PRAGMA_ENABLE_DEPRECATION_WARNINGS
 }
 
 TOptional<float> GetYawFromQuaternion(const FQuat& Quaternion)
 {
+	PRAGMA_DISABLE_DEPRECATION_WARNINGS
 	return GetYawFromVector(Quaternion.Vector());
+	PRAGMA_ENABLE_DEPRECATION_WARNINGS
 }
 
 void GetActorClassDefaultComponents(const TSubclassOf<AActor>& ActorClass, TArray<UActorComponent*>& OutComponents, const TSubclassOf<UActorComponent>& InComponentClass)
diff --git a/Engine/Source/Runtime/AIModule/Private/BehaviorTree/BTCompositeNode.cpp b/Engine/Source/Runtime/AIModule/Private/BehaviorTree/BTCompositeNode.cpp
--- a/Engine/Source/Runtime/AIModule/Private/BehaviorTree/BTCompositeNode.cpp
+++ b/Engine/Source/Runtime/AIModule/Private/BehaviorTree/BTCompositeNode.cpp
@@ -120,7 +120,7 @@ void UBTCompositeNode::OnChildActivation(FBehaviorTreeSearchData& SearchData, in
 	}
 
 	// update active node in current context: child node
-	NodeMemory->CurrentChild = ChildIndex;
+	NodeMemory->CurrentChild = IntCastChecked<int8>(ChildIndex);
 }
 
 void UBTCompositeNode::OnChildDeactivation(FBehaviorTreeSearchData& SearchData, const UBTNode& ChildNode, EBTNodeResult::Type& NodeResult, const bool bRequestedFromValidInstance) const
@@ -429,7 +429,7 @@ static bool UpdateOperationStack(const UBehaviorTreeComponent& OwnerComp, FStrin
 	return bTestResult;
 }
 
-bool UBTCompositeNode::DoDecoratorsAllowExecution(UBehaviorTreeComponent& OwnerComp, int32 InstanceIdx, int32 ChildIdx) const
+bool UBTCompositeNode::DoDecoratorsAllowExecution(UBehaviorTreeComponent& OwnerComp, const int32 InstanceIdx, const int32 ChildIdx) const
 {
 	ensure(Children.IsValidIndex(ChildIdx));
 	if (Children.IsValidIndex(ChildIdx) == false)
@@ -449,6 +449,8 @@ bool UBTCompositeNode::DoDecoratorsAllowExecution(UBehaviorTreeComponent& OwnerC
 
 	FBehaviorTreeInstance& MyInstance = OwnerComp.InstanceStack[InstanceIdx];
 
+	const uint16 InstanceIdxUint16 = IntCastChecked<uint16>(InstanceIdx);
+
 	if (ChildInfo.DecoratorOps.Num() == 0)
 	{
 		// simple check: all decorators must agree
@@ -456,7 +458,7 @@ bool UBTCompositeNode::DoDecoratorsAllowExecution(UBehaviorTreeComponent& OwnerC
 		{
 			const UBTDecorator* TestDecorator = ChildInfo.Decorators[DecoratorIndex];
 			const bool bIsAllowed = TestDecorator ? TestDecorator->WrappedCanExecute(OwnerComp, TestDecorator->GetNodeMemory<uint8>(MyInstance)) : false;
-			OwnerComp.StoreDebuggerSearchStep(TestDecorator, InstanceIdx, bIsAllowed);
+			OwnerComp.StoreDebuggerSearchStep(TestDecorator, InstanceIdxUint16, bIsAllowed);
 
 			const UBTNode* ChildNode = GetChildNode(ChildIdx);
 
@@ -527,7 +529,7 @@ bool UBTCompositeNode::DoDecoratorsAllowExecution(UBehaviorTreeComponent& OwnerC
 					UE_VLOG(OwnerComp.GetOwner(), LogBehaviorTree, Verbose, TEXT("finished execution test: %s"),
 						bResult ? TEXT("allowed") : TEXT("forbidden"));
 
-					OwnerComp.StoreDebuggerSearchStep(ChildInfo.Decorators[FMath::Max(0, FailedDecoratorIdx)], InstanceIdx, bResult);
+					OwnerComp.StoreDebuggerSearchStep(ChildInfo.Decorators[FMath::Max(0, FailedDecoratorIdx)], InstanceIdxUint16, bResult);
 					break;
 				}
 			}
@@ -645,7 +647,7 @@ uint16 UBTCompositeNode::GetChildExecutionIndex(int32 Index, EBTChildIndex Child
 			}
 		}
 
-		return ChildNode->GetExecutionIndex() - Offset;
+		return IntCastChecked<uint16>(ChildNode->GetExecutionIndex() - Offset);
 	}
 
 	return (LastExecutionIndex + 1);
diff --git a/Engine/Source/Runtime/AIModule/Private/BehaviorTree/BTTaskNode.cpp b/Engine/Source/Runtime/AIModule/Private/BehaviorTree/BTTaskNode.cpp
--- a/Engine/Source/Runtime/AIModule/Private/BehaviorTree/BTTaskNode.cpp
+++ b/Engine/Source/Runtime/AIModule/Private/BehaviorTree/BTTaskNode.cpp
@@ -120,7 +120,7 @@ void UBTTaskNode::ReceivedMessage(UBrainComponent* BrainComp, const FAIMessage&
 	UBehaviorTreeComponent* OwnerComp = static_cast<UBehaviorTreeComponent*>(BrainComp);
 	check(OwnerComp);
 	
-	const uint16 InstanceIdx = OwnerComp->FindInstanceContainingNode(this);
+	const uint16 InstanceIdx = IntCastChecked<uint16>(OwnerComp->FindInstanceContainingNode(this));
 	if (OwnerComp->InstanceStack.IsValidIndex(InstanceIdx))
 	{
 		uint8* NodeMemory = GetNodeMemory<uint8>(OwnerComp->InstanceStack[InstanceIdx]);
diff --git a/Engine/Source/Runtime/AIModule/Private/BehaviorTree/BehaviorTreeComponent.cpp b/Engine/Source/Runtime/AIModule/Private/BehaviorTree/BehaviorTreeComponent.cpp
--- a/Engine/Source/Runtime/AIModule/Private/BehaviorTree/BehaviorTreeComponent.cpp
+++ b/Engine/Source/Runtime/AIModule/Private/BehaviorTree/BehaviorTreeComponent.cpp
@@ -493,7 +493,7 @@ void UBehaviorTreeComponent::OnTaskFinished(const UBTTaskNode* TaskNode, EBTNode
 
 	if (TaskResult != EBTNodeResult::InProgress)
 	{
-		StoreDebuggerSearchStep(TaskNode, TaskInstanceIdx, TaskResult);
+		StoreDebuggerSearchStep(TaskNode, IntCastChecked<uint16>(TaskInstanceIdx), TaskResult);
 
 		// cleanup task observers
 		UnregisterMessageObserversFrom(TaskNode);
@@ -937,7 +937,9 @@ void UBehaviorTreeComponent::ActivateBranch(const UBTDecorator& RequestedBy, con
 
 	const bool bIsExecutingBranch = IsExecutingBranch(&RequestedBy, RequestedBy.GetChildIndex());
 	const bool bAbortPending = IsAbortPending();
-	const bool bIsDeactivatingBranchRoot = ExecutionRequest.ContinueWithResult == EBTNodeResult::Failed && ExecutionRequest.SearchStart == FBTNodeIndex(InstanceIdx, RequestedBy.GetExecutionIndex());
+
+	checkf(InstanceIdx != FBTNodeIndex::InvalidIndex, TEXT("Index has used the InvalidIndex value!"));
+	const bool bIsDeactivatingBranchRoot = ExecutionRequest.ContinueWithResult == EBTNodeResult::Failed && ExecutionRequest.SearchStart == FBTNodeIndex(IntCastChecked<uint16>(InstanceIdx), RequestedBy.GetExecutionIndex());
 
 	const bool bLogRequestExecution = !bIsExecutingBranch || (bForceResquestEvenIfExecuting || bAbortPending || bIsDeactivatingBranchRoot);
 	UE_VLOG(GetOwner(), LogBehaviorTree, Verbose, TEXT("%s, ActivateBranch(%s) executingBranch:%d abortPending:%d deactivatingBranchRoot:%d => %s"),
@@ -1085,8 +1087,8 @@ void UBehaviorTreeComponent::ScheduleExecutionUpdate()
 	bRequestedFlowUpdate = true;
 }
 
-void UBehaviorTreeComponent::RequestExecution(const UBTCompositeNode* RequestedOn, int32 InstanceIdx, const UBTNode* RequestedBy,
-											  int32 RequestedByChildIndex, EBTNodeResult::Type ContinueWithResult, bool bStoreForDebugger)
+void UBehaviorTreeComponent::RequestExecution(const UBTCompositeNode* RequestedOn, const int32 InstanceIdx, const UBTNode* RequestedBy,
+											  const int32 RequestedByChildIndex, const EBTNodeResult::Type ContinueWithResult, const bool bStoreForDebugger)
 {
 	SCOPE_CYCLE_COUNTER(STAT_AI_BehaviorTree_SearchTime);
 #if !UE_BUILD_SHIPPING // Disable in shipping builds
@@ -1113,12 +1115,15 @@ void UBehaviorTreeComponent::RequestExecution(const UBTCompositeNode* RequestedO
 		return;
 	}
 
+	checkf(InstanceIdx != FBTNodeIndex::InvalidIndex, TEXT("Index has used the InvalidIndex value!"));
+	const uint16 InstanceIdxUint16 = IntCastChecked<uint16>(InstanceIdx);
+
 	const bool bSwitchToHigherPriority = (ContinueWithResult == EBTNodeResult::Aborted);
 	const bool bAlreadyHasRequest = (ExecutionRequest.ExecuteNode != NULL);
 	const UBTNode* DebuggerNode = bStoreForDebugger ? RequestedBy : NULL;
 
 	FBTNodeIndex ExecutionIdx;
-	ExecutionIdx.InstanceIndex = InstanceIdx;
+	ExecutionIdx.InstanceIndex = InstanceIdxUint16;
 	ExecutionIdx.ExecutionIndex = RequestedBy->GetExecutionIndex();
 	uint16 LastExecutionIndex = MAX_uint16;
 
@@ -1158,7 +1163,10 @@ void UBehaviorTreeComponent::RequestExecution(const UBTCompositeNode* RequestedO
 				const int32 BranchRootInstanceIdx = FindInstanceContainingNode(BranchRoot);
 				if (BranchRootInstanceIdx != INDEX_NONE)
 				{
-					FBTNodeIndexRange Range( FBTNodeIndex(BranchRootInstanceIdx, BranchRoot->GetExecutionIndex()), FBTNodeIndex(BranchRootInstanceIdx, BranchRoot->GetLastExecutionIndex()));
+					checkf(BranchRootInstanceIdx != FBTNodeIndex::InvalidIndex, TEXT("Index has used the InvalidIndex value!"));
+					const uint16 BranchRootInstanceIdxUint16 = IntCastChecked<uint16>(BranchRootInstanceIdx);
+
+					FBTNodeIndexRange Range( FBTNodeIndex(BranchRootInstanceIdxUint16, BranchRoot->GetExecutionIndex()), FBTNodeIndex(BranchRootInstanceIdxUint16, BranchRoot->GetLastExecutionIndex()));
 					if (Range.Contains(ExecutionIdx))
 					{
 						UE_VLOG(GetOwner(), LogBehaviorTree, Log, TEXT("> skip: request by %s(%s) is in deactivated branch %s(%s) and was deactivated by %s"), *UBehaviorTreeTypes::DescribeNodeHelper(RequestedBy), *ExecutionIdx.Describe(), *UBehaviorTreeTypes::DescribeNodeHelper(BranchRoot), *Range.Describe(), *UBehaviorTreeTypes::DescribeNodeHelper(Info.Node));
@@ -1181,13 +1189,13 @@ void UBehaviorTreeComponent::RequestExecution(const UBTCompositeNode* RequestedO
 		LastExecutionIndex = RequestedOn->GetChildExecutionIndex(RequestedByChildIndex + 1, EBTChildIndex::FirstNode);
 	}
 
-	const FBTNodeIndex SearchEnd(InstanceIdx, LastExecutionIndex);
+	const FBTNodeIndex SearchEnd(InstanceIdxUint16, LastExecutionIndex);
 
 	// check if it's more important than currently requested
 	if (bAlreadyHasRequest && ExecutionRequest.SearchStart.TakesPriorityOver(ExecutionIdx))
 	{
 		UE_VLOG(GetOwner(), LogBehaviorTree, Log, TEXT("> skip: already has request with higher priority"));
-		StoreDebuggerRestart(DebuggerNode, InstanceIdx, true);
+		StoreDebuggerRestart(DebuggerNode, InstanceIdxUint16, true);
 
 		// make sure to update end of search range
 		if (bSwitchToHigherPriority)
@@ -1222,7 +1230,7 @@ void UBehaviorTreeComponent::RequestExecution(const UBTCompositeNode* RequestedO
 			{
 				UE_VLOG(GetOwner(), LogBehaviorTree, Log, TEXT("> skip: node index %s in a deactivated instance [%s..%s[ (applying search data for %s)"),
 					*ExecutionIdx.Describe(), *SearchData.DeactivatedBranchStart.Describe(), *SearchData.DeactivatedBranchEnd.Describe(), *SearchData.SearchRootNode.Describe());
-				StoreDebuggerRestart(DebuggerNode, InstanceIdx, false);
+				StoreDebuggerRestart(DebuggerNode, InstanceIdxUint16, false);
 				return;
 			}
 			else if (ExecutionIdx.InstanceIndex == SearchData.DeactivatedBranchStart.InstanceIndex && 
@@ -1231,7 +1239,7 @@ void UBehaviorTreeComponent::RequestExecution(const UBTCompositeNode* RequestedO
 			{
 				UE_VLOG(GetOwner(), LogBehaviorTree, Log, TEXT("> skip: node index %s in a deactivated branch [%s..%s[ (applying search data for %s)"),
 					*ExecutionIdx.Describe(), *SearchData.DeactivatedBranchStart.Describe(), *SearchData.DeactivatedBranchEnd.Describe(), *SearchData.SearchRootNode.Describe());
-				StoreDebuggerRestart(DebuggerNode, InstanceIdx, false);
+				StoreDebuggerRestart(DebuggerNode, InstanceIdxUint16, false);
 				return;
 			}
 		}
@@ -1247,7 +1255,7 @@ void UBehaviorTreeComponent::RequestExecution(const UBTCompositeNode* RequestedO
 		if (!bCanExecute)
 		{
 			UE_VLOG(GetOwner(), LogBehaviorTree, Log, TEXT("> skip: decorators are not allowing execution"));
-			StoreDebuggerRestart(DebuggerNode, InstanceIdx, false);
+			StoreDebuggerRestart(DebuggerNode, InstanceIdxUint16, false);
 			return;
 		}
 
@@ -1269,7 +1277,7 @@ void UBehaviorTreeComponent::RequestExecution(const UBTCompositeNode* RequestedO
 			const UBTCompositeNode* CommonParent = NULL;
 			uint16 CommonInstanceIdx = MAX_uint16;
 
-			FindCommonParent(InstanceStack, KnownInstances, RequestedOn, InstanceIdx, CurrentNode, CurrentInstanceIdx, CommonParent, CommonInstanceIdx);
+			FindCommonParent(InstanceStack, KnownInstances, RequestedOn, InstanceIdxUint16, CurrentNode, CurrentInstanceIdx, CommonParent, CommonInstanceIdx);
 
 			// check decorators between common parent and restart parent
 			int32 ItInstanceIdx = InstanceIdx;
@@ -1303,7 +1311,7 @@ void UBehaviorTreeComponent::RequestExecution(const UBTCompositeNode* RequestedO
 				if (!bCanExecuteTest)
 				{
 					UE_VLOG(GetOwner(), LogBehaviorTree, Log, TEXT("> skip: decorators are not allowing execution"));
-					StoreDebuggerRestart(DebuggerNode, InstanceIdx, false);
+					StoreDebuggerRestart(DebuggerNode, InstanceIdxUint16, false);
 					return;
 				}
 
@@ -1325,16 +1333,16 @@ void UBehaviorTreeComponent::RequestExecution(const UBTCompositeNode* RequestedO
 		if (bCanExecute)
 		{
 			UE_VLOG(GetOwner(), LogBehaviorTree, Log, TEXT("> skip: decorators are still allowing execution"));
-			StoreDebuggerRestart(DebuggerNode, InstanceIdx, false);
+			StoreDebuggerRestart(DebuggerNode, InstanceIdxUint16, false);
 			return;
 		}
 
 		ExecutionRequest.ExecuteNode = RequestedOn;
-		ExecutionRequest.ExecuteInstanceIdx = InstanceIdx;
+		ExecutionRequest.ExecuteInstanceIdx = InstanceIdxUint16;
 	}
 
 	// store it
-	StoreDebuggerRestart(DebuggerNode, InstanceIdx, true);
+	StoreDebuggerRestart(DebuggerNode, InstanceIdxUint16, true);
 
 	// search end can be set only when switching to high priority
 	// or previous request was limited and current limit is wider
@@ -1573,7 +1581,7 @@ void UBehaviorTreeComponent::TickComponent(float DeltaTime, enum ELevelTick Tick
 		UWorld* MyWorld = GetWorld();
 		if (MyWorld)
 		{
-			const float CurrentGameTime = MyWorld->GetTimeSeconds();
+			const double CurrentGameTime = MyWorld->GetTimeSeconds();
 			const float CurrentDeltaTime = MyWorld->GetDeltaSeconds();
 			if (CurrentGameTime - LastRequestedDeltaTimeGameTime - CurrentDeltaTime > KINDA_SMALL_NUMBER)
 			{
@@ -1807,7 +1815,7 @@ void UBehaviorTreeComponent::ScheduleNextTick(const float NextNeededDeltaTime)
 		SetComponentTickIntervalAndCooldown(!bTickedOnce && NextTickDeltaTime < FORCE_TICK_INTERVAL_DT ? FORCE_TICK_INTERVAL_DT : NextTickDeltaTime);
 	}
 	UWorld* MyWorld = GetWorld();
-	LastRequestedDeltaTimeGameTime = MyWorld ? MyWorld->GetTimeSeconds() : 0.0f;
+	LastRequestedDeltaTimeGameTime = MyWorld ? MyWorld->GetTimeSeconds() : 0.;
 }
 
 void UBehaviorTreeComponent::ProcessExecutionRequest()
@@ -2130,7 +2138,8 @@ void UBehaviorTreeComponent::RollbackSearchChanges()
 {
 	if (SearchData.RollbackInstanceIdx >= 0)
 	{
-		ActiveInstanceIdx = SearchData.RollbackInstanceIdx;
+		checkf(SearchData.RollbackInstanceIdx != FBTNodeIndex::InvalidIndex, TEXT("Index has used the InvalidIndex value!"));
+		ActiveInstanceIdx = IntCastChecked<uint16>(SearchData.RollbackInstanceIdx);
 		SearchData.DeactivatedBranchStart = SearchData.RollbackDeactivatedBranchStart;
 		SearchData.DeactivatedBranchEnd = SearchData.RollbackDeactivatedBranchEnd;
 
@@ -2241,10 +2250,12 @@ void UBehaviorTreeComponent::UnregisterAuxNodesUpTo(const FBTNodeIndex& Index)
 		FBehaviorTreeInstance& InstanceInfo = InstanceStack[InstanceIndex];
 		for (const UBTAuxiliaryNode* AuxNode : InstanceInfo.GetActiveAuxNodes())
 		{
-			FBTNodeIndex AuxIdx(InstanceIndex, AuxNode->GetExecutionIndex());
+			// Could safely use static_cast() here but its a bit safer using IntCastChecked() for future code changes.
+			const uint16 InstanceIndexUint16 = IntCastChecked<uint16>(InstanceIndex);
+			FBTNodeIndex AuxIdx(InstanceIndexUint16, AuxNode->GetExecutionIndex());
 			if (Index.TakesPriorityOver(AuxIdx))
 			{
-				SearchData.AddUniqueUpdate(FBehaviorTreeSearchUpdate(AuxNode, InstanceIndex, EBTNodeUpdateMode::Remove));
+				SearchData.AddUniqueUpdate(FBehaviorTreeSearchUpdate(AuxNode, InstanceIndexUint16, EBTNodeUpdateMode::Remove));
 			}
 		}
 	}
@@ -2257,10 +2268,12 @@ void UBehaviorTreeComponent::UnregisterAuxNodesInRange(const FBTNodeIndex& FromI
 		FBehaviorTreeInstance& InstanceInfo = InstanceStack[InstanceIndex];
 		for (const UBTAuxiliaryNode* AuxNode : InstanceInfo.GetActiveAuxNodes())
 		{
-			FBTNodeIndex AuxIdx(InstanceIndex, AuxNode->GetExecutionIndex());
+			// Could safely use static_cast() here but its a bit safer using IntCastChecked() for future code changes.
+			const uint16 InstanceIndexUint16 = IntCastChecked<uint16>(InstanceIndex);
+			FBTNodeIndex AuxIdx(InstanceIndexUint16, AuxNode->GetExecutionIndex());
 			if (FromIndex.TakesPriorityOver(AuxIdx) && AuxIdx.TakesPriorityOver(ToIndex))
 			{
-				SearchData.AddUniqueUpdate(FBehaviorTreeSearchUpdate(AuxNode, InstanceIndex, EBTNodeUpdateMode::Remove));
+				SearchData.AddUniqueUpdate(FBehaviorTreeSearchUpdate(AuxNode, InstanceIndexUint16, EBTNodeUpdateMode::Remove));
 			}
 		}
 	}
@@ -2273,6 +2286,9 @@ void UBehaviorTreeComponent::UnregisterAuxNodesInBranch(const UBTCompositeNode*
 	{
 		check(Node);
 
+		checkf(InstanceIdx != FBTNodeIndex::InvalidIndex, TEXT("Index has used the InvalidIndex value!"));
+		const uint16 InstanceIdxUint16 = IntCastChecked<uint16>(InstanceIdx);
+
 		TArray<FBehaviorTreeSearchUpdate> UpdateListCopy;
 		if (bApplyImmediately)
 		{
@@ -2280,8 +2296,8 @@ void UBehaviorTreeComponent::UnregisterAuxNodesInBranch(const UBTCompositeNode*
 			SearchData.PendingUpdates.Reset();
 		}
 
-		const FBTNodeIndex FromIndex(InstanceIdx, Node->GetExecutionIndex());
-		const FBTNodeIndex ToIndex(InstanceIdx, Node->GetLastExecutionIndex());
+		const FBTNodeIndex FromIndex(InstanceIdxUint16, Node->GetExecutionIndex());
+		const FBTNodeIndex ToIndex(InstanceIdxUint16, Node->GetLastExecutionIndex());
 		UnregisterAuxNodesInRange(FromIndex, ToIndex);
 
 		if (bApplyImmediately)
@@ -2398,7 +2414,9 @@ void UBehaviorTreeComponent::RegisterMessageObserver(const UBTTaskNode* TaskNode
 	{
 		FBTNodeIndex NodeIdx;
 		NodeIdx.ExecutionIndex = TaskNode->GetExecutionIndex();
-		NodeIdx.InstanceIndex = InstanceStack.Num() - 1;
+
+		const int32 InstanceIndex = InstanceStack.Num() - 1;
+		NodeIdx.InstanceIndex = IntCastChecked<uint16>(InstanceIndex);
 
 		TaskMessageObservers.Add(NodeIdx,
 			FAIMessageObserver::Create(this, MessageType, FOnAIMessage::CreateUObject(const_cast<UBTTaskNode*>(TaskNode), &UBTTaskNode::ReceivedMessage))
@@ -2415,7 +2433,9 @@ void UBehaviorTreeComponent::RegisterMessageObserver(const UBTTaskNode* TaskNode
 	{
 		FBTNodeIndex NodeIdx;
 		NodeIdx.ExecutionIndex = TaskNode->GetExecutionIndex();
-		NodeIdx.InstanceIndex = InstanceStack.Num() - 1;
+
+		const int32 InstanceIndex = InstanceStack.Num() - 1;
+		NodeIdx.InstanceIndex = IntCastChecked<uint16>(InstanceIndex);
 
 		TaskMessageObservers.Add(NodeIdx,
 			FAIMessageObserver::Create(this, MessageType, RequestID, FOnAIMessage::CreateUObject(const_cast<UBTTaskNode*>(TaskNode), &UBTTaskNode::ReceivedMessage))
@@ -2444,7 +2464,14 @@ void UBehaviorTreeComponent::UnregisterMessageObserversFrom(const UBTTaskNode* T
 
 		FBTNodeIndex NodeIdx;
 		NodeIdx.ExecutionIndex = TaskNode->GetExecutionIndex();
-		NodeIdx.InstanceIndex = FindInstanceContainingNode(TaskNode);
+
+		const int32 InstanceContainingNodeIdx = FindInstanceContainingNode(TaskNode);
+		
+		// InstanceContainingNodeIdx could be INDEX_NONE which means we can't use IntCastChecked() here
+		checkf(InstanceContainingNodeIdx != FBTNodeIndex::InvalidIndex, TEXT("Index has used the InvalidIndex value!"));
+		checkf(InstanceContainingNodeIdx <= MAX_uint16, TEXT("Narrowing conversion causing loss of data"));
+
+		NodeIdx.InstanceIndex = static_cast<uint16>(InstanceContainingNodeIdx);
 		
 		UnregisterMessageObserversFrom(NodeIdx);
 	}
@@ -2596,7 +2623,9 @@ bool UBehaviorTreeComponent::PushInstance(UBehaviorTree& TreeAsset)
 		NewInstance.Initialize(*this, *RootNode, NodeInstanceIndex, bFirstTime ? EBTMemoryInit::Initialize : EBTMemoryInit::RestoreSubtree);
 
 		InstanceStack.Push(NewInstance);
-		ActiveInstanceIdx = InstanceStack.Num() - 1;
+
+		const int32 InstanceIndex = InstanceStack.Num() - 1;
+		ActiveInstanceIdx = IntCastChecked<uint16>(InstanceIndex);
 
 		// start root level services now (they won't be removed on looping tree anyway)
 		for (int32 ServiceIndex = 0; ServiceIndex < RootNode->Services.Num(); ServiceIndex++)
@@ -2650,7 +2679,7 @@ uint8 UBehaviorTreeComponent::UpdateInstanceId(UBehaviorTree* TreeAsset, const U
 	{
 		if (KnownInstances[InstanceIndex] == InstanceId)
 		{
-			return InstanceIndex;
+			return IntCastChecked<uint8>(InstanceIndex);
 		}
 	}
 
@@ -2658,8 +2687,7 @@ uint8 UBehaviorTreeComponent::UpdateInstanceId(UBehaviorTree* TreeAsset, const U
 	InstanceId.FirstNodeInstance = NodeInstances.Num();
 
 	const int32 NewIndex = KnownInstances.Add(InstanceId);
-	check(NewIndex < MAX_uint8);
-	return NewIndex;
+	return IntCastChecked<uint8>(NewIndex);
 }
 
 int32 UBehaviorTreeComponent::FindInstanceContainingNode(const UBTNode* Node) const
@@ -2771,7 +2799,7 @@ void UBehaviorTreeComponent::RemoveAllInstances()
 			// instance memory will be removed on Cleanup in EBTMemoryClear::Destroy mode
 			// prevent from calling it multiple times - StopTree does it for current InstanceStack
 			DummyInstance.SetInstanceMemory(Info.InstanceMemory);
-			DummyInstance.InstanceIdIndex = Idx;
+			DummyInstance.InstanceIdIndex = IntCastChecked<uint8>(Idx);
 			DummyInstance.RootNode = Info.RootNode;
 
 			DummyInstance.Cleanup(*this, EBTMemoryClear::Destroy);
@@ -2878,9 +2906,9 @@ FString UBehaviorTreeComponent::DescribeActiveTrees() const
 	return Assets.Len() ? Assets.LeftChop(2) : TEXT("None");
 }
 
-float UBehaviorTreeComponent::GetTagCooldownEndTime(FGameplayTag CooldownTag) const
+double UBehaviorTreeComponent::GetTagCooldownEndTime(FGameplayTag CooldownTag) const
 {
-	const float CooldownEndTime = CooldownTagsMap.FindRef(CooldownTag);
+	const double CooldownEndTime = CooldownTagsMap.FindRef(CooldownTag);
 	return CooldownEndTime;
 }
 
@@ -2888,7 +2916,7 @@ void UBehaviorTreeComponent::AddCooldownTagDuration(FGameplayTag CooldownTag, fl
 {
 	if (CooldownTag.IsValid())
 	{
-		float* CurrentEndTime = CooldownTagsMap.Find(CooldownTag);
+		double* CurrentEndTime = CooldownTagsMap.Find(CooldownTag);
 
 		// If we are supposed to add to an existing duration, do that, otherwise we set a new value.
 		if (bAddToExistingDuration && (CurrentEndTime != nullptr))
@@ -3072,7 +3100,7 @@ void UBehaviorTreeComponent::StoreDebuggerExecutionStep(EBTExecutionSnap::Type S
 		const FBehaviorTreeInstance& ActiveInstance = InstanceStack[InstanceIndex];
 		
 		FBehaviorTreeDebuggerInstance StoreInfo;
-		StoreDebuggerInstance(StoreInfo, InstanceIndex, SnapType);
+		StoreDebuggerInstance(StoreInfo, IntCastChecked<uint16>(InstanceIndex), SnapType);
 		CurrentStep.InstanceStack.Add(StoreInfo);
 	}
 
@@ -3335,13 +3363,15 @@ void UBehaviorTreeComponent::StoreDebuggerBlackboard(TMap<FName, FString>& Black
 
 		for (int32 KeyIndex = 0; KeyIndex < NumKeys; KeyIndex++)
 		{
-			FString Value = BlackboardComp->DescribeKeyValue(KeyIndex, EBlackboardDescription::OnlyValue);
+			const FBlackboard::FKey Key = IntCastChecked<FBlackboard::FKey>(KeyIndex);
+
+			FString Value = BlackboardComp->DescribeKeyValue(Key, EBlackboardDescription::OnlyValue);
 			if (Value.Len() == 0)
 			{
 				Value = TEXT("n/a");
 			}
 
-			BlackboardValueDesc.Add(BlackboardComp->GetKeyName(KeyIndex), Value);
+			BlackboardValueDesc.Add(BlackboardComp->GetKeyName(Key), Value);
 		}
 	}
 #endif
diff --git a/Engine/Source/Runtime/AIModule/Private/BehaviorTree/BehaviorTreeManager.cpp b/Engine/Source/Runtime/AIModule/Private/BehaviorTree/BehaviorTreeManager.cpp
--- a/Engine/Source/Runtime/AIModule/Private/BehaviorTree/BehaviorTreeManager.cpp
+++ b/Engine/Source/Runtime/AIModule/Private/BehaviorTree/BehaviorTreeManager.cpp
@@ -54,7 +54,7 @@ void UBehaviorTreeManager::FinishDestroy()
 	Super::FinishDestroy();
 }
 
-int32 UBehaviorTreeManager::GetAlignedDataSize(int32 Size)
+uint16 UBehaviorTreeManager::GetAlignedDataSize(uint16 Size)
 {
 	// round to 4 bytes
 	return ((Size + 3) & ~3);
@@ -100,7 +100,7 @@ static void MergeDecoratorOpsHelper(TArray<FBTDecoratorLogic>& LinkOps, const TA
 	if (NumOriginalDecorators > 0)
 	{
 		// and operator for two groups of composites: original and new one
-		FBTDecoratorLogic MainAndOp(EBTDecoratorLogic::And, LinkOps.Num() ? 2 : (NumOriginalDecorators + 1));
+		FBTDecoratorLogic MainAndOp(EBTDecoratorLogic::And, LinkOps.Num() ? 2 : IntCastChecked<uint16>(NumOriginalDecorators + 1));
 		LinkOps.Insert(MainAndOp, 0);
 
 		if (NumOriginalOps == 0)
@@ -108,7 +108,7 @@ static void MergeDecoratorOpsHelper(TArray<FBTDecoratorLogic>& LinkOps, const TA
 			// add Test operations, original link didn't have composite operators
 			for (int32 Idx = 0; Idx < NumOriginalDecorators; Idx++)
 			{
-				FBTDecoratorLogic TestOp(EBTDecoratorLogic::Test, Idx);
+				FBTDecoratorLogic TestOp(EBTDecoratorLogic::Test, IntCastChecked<uint16>(Idx));
 				LinkOps.Add(TestOp);
 			}
 		}
@@ -117,12 +117,12 @@ static void MergeDecoratorOpsHelper(TArray<FBTDecoratorLogic>& LinkOps, const TA
 	// add injected operators
 	if (InjectedOps.Num() == 0)
 	{
-		FBTDecoratorLogic InjectedAndOp(EBTDecoratorLogic::And, NumInjectedDecorators);
+		FBTDecoratorLogic InjectedAndOp(EBTDecoratorLogic::And, IntCastChecked<uint16>(NumInjectedDecorators));
 		LinkOps.Add(InjectedAndOp);
 
 		for (int32 Idx = 0; Idx < NumInjectedDecorators; Idx++)
 		{
-			FBTDecoratorLogic TestOp(EBTDecoratorLogic::Test, NumOriginalDecorators + Idx);
+			FBTDecoratorLogic TestOp(EBTDecoratorLogic::Test, IntCastChecked<uint16>(NumOriginalDecorators + Idx));
 			LinkOps.Add(TestOp);
 		}
 	}
@@ -133,7 +133,7 @@ static void MergeDecoratorOpsHelper(TArray<FBTDecoratorLogic>& LinkOps, const TA
 			FBTDecoratorLogic InjectedOpCopy = InjectedOps[Idx];
 			if (InjectedOpCopy.Operation == EBTDecoratorLogic::Test)
 			{
-				InjectedOpCopy.Number += NumOriginalDecorators;
+				InjectedOpCopy.Number = IntCastChecked<uint16>((int32)InjectedOpCopy.Number + NumOriginalDecorators);
 			}
 
 			LinkOps.Add(InjectedOpCopy);
@@ -202,7 +202,7 @@ static void InitializeNodeHelper(UBTCompositeNode* ParentNode, UBTNode* NodeOb,
 
 				FScopedBTLoggingContext LogContext(Decorator);
 				Decorator->InitializeFromAsset(TreeAsset);
-				Decorator->InitializeParentLink(ChildIndex);
+				Decorator->InitializeParentLink(IntCastChecked<uint8>(ChildIndex));
 				ExecutionIndex++;
 			}
 
@@ -233,7 +233,7 @@ static void InitializeNodeHelper(UBTCompositeNode* ParentNode, UBTNode* NodeOb,
 					Decorator->MarkInjectedNode();
 					FScopedBTLoggingContext LogContext(Decorator);
 					Decorator->InitializeFromAsset(TreeAsset);
-					Decorator->InitializeParentLink(ChildIndex);
+					Decorator->InitializeParentLink(IntCastChecked<uint8>(ChildIndex));
 					ExecutionIndex++;
 				}
 
@@ -273,7 +273,7 @@ static void InitializeNodeHelper(UBTCompositeNode* ParentNode, UBTNode* NodeOb,
 
 					FScopedBTLoggingContext LogContext(Service);
 					Service->InitializeFromAsset(TreeAsset);
-					Service->InitializeParentLink(ChildIndex);
+					Service->InitializeParentLink(IntCastChecked<uint8>(ChildIndex));
 					ExecutionIndex++;
 				}
 			}
diff --git a/Engine/Source/Runtime/AIModule/Private/BehaviorTree/BlackboardComponent.cpp b/Engine/Source/Runtime/AIModule/Private/BehaviorTree/BlackboardComponent.cpp
--- a/Engine/Source/Runtime/AIModule/Private/BehaviorTree/BlackboardComponent.cpp
+++ b/Engine/Source/Runtime/AIModule/Private/BehaviorTree/BlackboardComponent.cpp
@@ -155,7 +155,10 @@ bool UBlackboardComponent::InitializeBlackboard(UBlackboardData& NewAsset)
 					KeyType->PreInitialize(*this);
 
 					const uint16 KeyMemory = KeyType->GetValueSize() + (KeyType->HasInstance() ? sizeof(FBlackboardInstancedKeyMemory) : 0);
-					InitList.Add(FBlackboardInitializationData(KeyIndex + It->GetFirstKeyID(), KeyMemory));
+					const int32 OffsetKey = KeyIndex + It->GetFirstKeyID();
+					check(OffsetKey != FBlackboard::InvalidKey);
+
+					InitList.Add(FBlackboardInitializationData(IntCastChecked<FBlackboard::FKey>(OffsetKey), KeyMemory));
 				}
 			}
 		}
@@ -208,8 +211,10 @@ void UBlackboardComponent::DestroyValues()
 			UBlackboardKeyType* KeyType = It->Keys[KeyIndex].KeyType;
 			if (KeyType)
 			{
-				const int32 UseIdx = KeyIndex + It->GetFirstKeyID();
-				uint8* KeyMemory = GetKeyRawData(UseIdx);
+				const int32 OffsetKey = KeyIndex + It->GetFirstKeyID();
+				check(OffsetKey != FBlackboard::InvalidKey);
+
+				uint8* KeyMemory = GetKeyRawData(IntCastChecked<FBlackboard::FKey>(OffsetKey));
 				KeyType->WrappedFree(*this, KeyMemory);
 			}
 		}
@@ -237,7 +242,7 @@ void UBlackboardComponent::PopulateSynchronizedKeys()
 				if (Key.bInstanceSynced)
 				{
 					UBlackboardData* const OtherBlackboardAsset = OtherBlackboard->GetBlackboardAsset();
-					const int32 OtherKeyID = OtherBlackboardAsset ? OtherBlackboardAsset->GetKeyID(Key.EntryName) : FBlackboard::InvalidKey;
+					const FBlackboard::FKey OtherKeyID = OtherBlackboardAsset ? OtherBlackboardAsset->GetKeyID(Key.EntryName) : FBlackboard::InvalidKey;
 					if (OtherKeyID != FBlackboard::InvalidKey)
 					{
 						const FBlackboardEntry* const OtherKey = OtherBlackboard->GetBlackboardAsset()->GetKey(OtherKeyID);
@@ -247,7 +252,7 @@ void UBlackboardComponent::PopulateSynchronizedKeys()
 						const bool bKeyHasInstance = Key.KeyType->HasInstance();
 						const uint16 DataOffset = bKeyHasInstance ? sizeof(FBlackboardInstancedKeyMemory) : 0;
 
-						const int32 KeyID = BlackboardAsset->GetKeyID(Key.EntryName);
+						const FBlackboard::FKey KeyID = BlackboardAsset->GetKeyID(Key.EntryName);
 						uint8* RawData = GetKeyRawData(KeyID) + DataOffset;
 						uint8* RawSource = OtherBlackboard->GetKeyRawData(OtherKeyID) + DataOffset;
 
@@ -514,12 +519,15 @@ FString UBlackboardComponent::GetDebugInfoString(EBlackboardDescription::Type Mo
 	FString DebugString = FString::Printf(TEXT("Blackboard (asset: %s)\n"), *GetNameSafe(BlackboardAsset));
 
 	TArray<FString> KeyDesc;
-	uint8 Offset = 0;
+	int32 Offset = 0;
 	for (UBlackboardData* It = BlackboardAsset; It; It = It->Parent)
 	{
 		for (int32 KeyIndex = 0; KeyIndex < It->Keys.Num(); KeyIndex++)
 		{
-			KeyDesc.Add(DescribeKeyValue(KeyIndex + Offset, Mode));
+			const int32 OffsetKey = KeyIndex + Offset;
+			check(OffsetKey != FBlackboard::InvalidKey);
+
+			KeyDesc.Add(DescribeKeyValue(IntCastChecked<FBlackboard::FKey>(OffsetKey), Mode));
 		}
 		Offset += It->Keys.Num();
 	}
@@ -609,7 +617,10 @@ void UBlackboardComponent::DescribeSelfToVisLog(FVisualLogEntry* Snapshot) const
 			{
 				const FBlackboardEntry& Key = It->Keys[KeyIndex];
 
-				const uint8* ValueData = GetKeyRawData(It->GetFirstKeyID() + KeyIndex);
+				const int32 OffsetKey = It->GetFirstKeyID() + KeyIndex;
+				check(OffsetKey != FBlackboard::InvalidKey);
+
+				const uint8* ValueData = GetKeyRawData(IntCastChecked<FBlackboard::FKey>(OffsetKey));
 				FString ValueDesc = Key.KeyType ? *(Key.KeyType->WrappedDescribeValue(*this, ValueData)) : TEXT("empty");
 
 				Category.Add(Key.EntryName.ToString(), ValueDesc);
@@ -787,7 +798,7 @@ void UBlackboardComponent::ClearValue(FBlackboard::FKey KeyID)
 					if (OtherBlackboard != nullptr && ShouldSyncWithBlackboard(*OtherBlackboard))
 					{
 						UBlackboardData* const OtherBlackboardAsset = OtherBlackboard->GetBlackboardAsset();
-						const int32 OtherKeyID = OtherBlackboardAsset ? OtherBlackboardAsset->GetKeyID(EntryInfo->EntryName) : FBlackboard::InvalidKey;
+						const FBlackboard::FKey OtherKeyID = OtherBlackboardAsset ? OtherBlackboardAsset->GetKeyID(EntryInfo->EntryName) : FBlackboard::InvalidKey;
 						if (OtherKeyID != FBlackboard::InvalidKey)
 						{
 							const FBlackboardEntry* OtherEntryInfo = OtherBlackboard->BlackboardAsset->GetKey(OtherKeyID);
diff --git a/Engine/Source/Runtime/AIModule/Private/BehaviorTree/BlackboardData.cpp b/Engine/Source/Runtime/AIModule/Private/BehaviorTree/BlackboardData.cpp
--- a/Engine/Source/Runtime/AIModule/Private/BehaviorTree/BlackboardData.cpp
+++ b/Engine/Source/Runtime/AIModule/Private/BehaviorTree/BlackboardData.cpp
@@ -107,7 +107,10 @@ FBlackboard::FKey UBlackboardData::InternalGetKeyID(const FName& KeyName, EKeyLo
 	{
 		if (Keys[KeyIndex].EntryName == KeyName)
 		{
-			return KeyIndex + FirstKeyID;
+			const int32 OffsetKey = KeyIndex + FirstKeyID;
+			check(OffsetKey != FBlackboard::InvalidKey);
+
+			return IntCastChecked<FBlackboard::FKey>(OffsetKey);
 		}
 	}
 
@@ -280,7 +283,10 @@ void UBlackboardData::UpdateParentKeys()
 
 void UBlackboardData::UpdateKeyIDs()
 {
-	FirstKeyID = Parent ? Parent->GetNumKeys() : 0;
+	const int32 FirstKeyIDInt = Parent ? Parent->GetNumKeys() : 0;
+	check(FirstKeyIDInt != FBlackboard::InvalidKey);
+
+	FirstKeyID = Parent ? IntCastChecked<FBlackboard::FKey>(FirstKeyIDInt) : 0;
 }
 
 void UBlackboardData::UpdateDeprecatedKeys()
diff --git a/Engine/Source/Runtime/AIModule/Private/BehaviorTree/BlueprintNodeHelpers.cpp b/Engine/Source/Runtime/AIModule/Private/BehaviorTree/BlueprintNodeHelpers.cpp
--- a/Engine/Source/Runtime/AIModule/Private/BehaviorTree/BlueprintNodeHelpers.cpp
+++ b/Engine/Source/Runtime/AIModule/Private/BehaviorTree/BlueprintNodeHelpers.cpp
@@ -28,7 +28,7 @@ namespace BlueprintNodeHelpers
 			TotalMem = 0;
 		}
 
-		return TotalMem;
+		return IntCastChecked<uint16>(TotalMem);
 	}
 
 #define GET_STRUCT_NAME_CHECKED(StructName) \
diff --git a/Engine/Source/Runtime/AIModule/Private/BehaviorTree/Composites/BTComposite_SimpleParallel.cpp b/Engine/Source/Runtime/AIModule/Private/BehaviorTree/Composites/BTComposite_SimpleParallel.cpp
--- a/Engine/Source/Runtime/AIModule/Private/BehaviorTree/Composites/BTComposite_SimpleParallel.cpp
+++ b/Engine/Source/Runtime/AIModule/Private/BehaviorTree/Composites/BTComposite_SimpleParallel.cpp
@@ -81,7 +81,9 @@ void UBTComposite_SimpleParallel::NotifyChildExecution(UBehaviorTreeComponent& O
 
 			const int32 MyInstanceIdx = OwnerComp.FindInstanceContainingNode(this);
 
-			OwnerComp.UnregisterParallelTask(Children[EBTParallelChild::MainTask].ChildTask, MyInstanceIdx);
+			// use check() here not IntCastChecked() as MyInstanceIdx can be INDEX_NONE!
+			check(MyInstanceIdx < MAX_uint16);
+			OwnerComp.UnregisterParallelTask(Children[EBTParallelChild::MainTask].ChildTask, static_cast<uint16>(MyInstanceIdx));
 			if (NodeResult != EBTNodeResult::Aborted && !MyMemory->bRepeatMainTask)
 			{
 				// check if subtree should be aborted when task finished with success/failed result
diff --git a/Engine/Source/Runtime/AIModule/Private/BehaviorTree/Decorators/BTDecorator_Blackboard.cpp b/Engine/Source/Runtime/AIModule/Private/BehaviorTree/Decorators/BTDecorator_Blackboard.cpp
--- a/Engine/Source/Runtime/AIModule/Private/BehaviorTree/Decorators/BTDecorator_Blackboard.cpp
+++ b/Engine/Source/Runtime/AIModule/Private/BehaviorTree/Decorators/BTDecorator_Blackboard.cpp
@@ -196,7 +196,7 @@ void UBTDecorator_Blackboard::PostEditChangeProperty(FPropertyChangedEvent& Prop
 				// Set both properties to match the first valid value or invalidate them
 				if (EnumType != nullptr && EnumType->NumEnums())
 				{
-					IntValue = EnumType->GetValueByIndex(0);
+					IntValue = IntCastChecked<int32>(EnumType->GetValueByIndex(0));
 					StringValue = EnumType->GetNameStringByIndex(0);
 				}
 				else
@@ -211,7 +211,7 @@ void UBTDecorator_Blackboard::PostEditChangeProperty(FPropertyChangedEvent& Prop
 			}
 			else if (ChangedPropName == StringValueName)
 			{
-				IntValue = (EnumType != nullptr) ? EnumType->GetValueByNameString(StringValue) : INDEX_NONE;
+				IntValue = (EnumType != nullptr) ? IntCastChecked<int32>(EnumType->GetValueByNameString(StringValue)) : INDEX_NONE;
 			}
 		}
 	}
@@ -221,21 +221,21 @@ void UBTDecorator_Blackboard::PostEditChangeProperty(FPropertyChangedEvent& Prop
 	{
 		if (KeyCDO && KeyCDO->GetTestOperation() == EBlackboardKeyOperation::Basic)
 		{
-			OperationType = BasicOperation;
+			OperationType = BasicOperation.GetIntValue();
 		}
 	}
 	else if (ChangedPropName == GET_MEMBER_NAME_CHECKED(UBTDecorator_Blackboard, ArithmeticOperation))
 	{
 		if (KeyCDO && KeyCDO->GetTestOperation() == EBlackboardKeyOperation::Arithmetic)
 		{
-			OperationType = ArithmeticOperation;
+			OperationType = ArithmeticOperation.GetIntValue();
 		}
 	}
 	else if (ChangedPropName == GET_MEMBER_NAME_CHECKED(UBTDecorator_Blackboard, TextOperation))
 	{
 		if (KeyCDO && KeyCDO->GetTestOperation() == EBlackboardKeyOperation::Text)
 		{
-			OperationType = TextOperation;
+			OperationType = TextOperation.GetIntValue();
 		}
 	}
 
@@ -282,14 +282,14 @@ void UBTDecorator_Blackboard::RefreshEnumBasedDecorator(const FBlackboardEntry&
 	{
 		// Pass string as FName to use fast search path whenever possible 
 		const FName Name = *StringValue;
-		IntValue = Enum->GetValueByName(Name);
+		IntValue = IntCastChecked<int32>(Enum->GetValueByName(Name));
 	}
 	else if (IntValue >= 0 && IntValue < Enum->NumEnums())
 	{
 		// Previous implementation was using only 'IntValue' to store the enumeration
 		// index so we need to extract associated name and update `IntValue` to store the value.
 		StringValue = Enum->GetNameByIndex(IntValue).ToString();
-		IntValue = Enum->GetValueByIndex(IntValue);
+		IntValue = IntCastChecked<int32>(Enum->GetValueByIndex(IntValue));
 	}
 }
 
diff --git a/Engine/Source/Runtime/AIModule/Private/BehaviorTree/Decorators/BTDecorator_ConeCheck.cpp b/Engine/Source/Runtime/AIModule/Private/BehaviorTree/Decorators/BTDecorator_ConeCheck.cpp
--- a/Engine/Source/Runtime/AIModule/Private/BehaviorTree/Decorators/BTDecorator_ConeCheck.cpp
+++ b/Engine/Source/Runtime/AIModule/Private/BehaviorTree/Decorators/BTDecorator_ConeCheck.cpp
@@ -125,8 +125,8 @@ void UBTDecorator_ConeCheck::DescribeRuntimeValues(const UBehaviorTreeComponent&
 	if (CalculateDirection(BBComponent, ConeOrigin, Observed, DirectionToObserved)
 		&& CalculateDirection(BBComponent, ConeOrigin, ConeDirection, ConeDir))
 	{
-		const float CurrentAngleDot = ConeDir.CosineAngle2D(DirectionToObserved);
-		const float CurrentAngleRad = FMath::Acos(CurrentAngleDot);
+		const FVector::FReal CurrentAngleDot = ConeDir.CosineAngle2D(DirectionToObserved);
+		const FVector::FReal CurrentAngleRad = FMath::Acos(CurrentAngleDot);
 
 		Values.Add(FString::Printf(TEXT("Angle: %.0f (%s cone)"),
 			FMath::RadiansToDegrees(CurrentAngleRad),
diff --git a/Engine/Source/Runtime/AIModule/Private/BehaviorTree/Decorators/BTDecorator_Cooldown.cpp b/Engine/Source/Runtime/AIModule/Private/BehaviorTree/Decorators/BTDecorator_Cooldown.cpp
--- a/Engine/Source/Runtime/AIModule/Private/BehaviorTree/Decorators/BTDecorator_Cooldown.cpp
+++ b/Engine/Source/Runtime/AIModule/Private/BehaviorTree/Decorators/BTDecorator_Cooldown.cpp
@@ -24,7 +24,7 @@ void UBTDecorator_Cooldown::PostLoad()
 bool UBTDecorator_Cooldown::CalculateRawConditionValue(UBehaviorTreeComponent& OwnerComp, uint8* NodeMemory) const
 {
 	FBTCooldownDecoratorMemory* DecoratorMemory = CastInstanceNodeMemory<FBTCooldownDecoratorMemory>(NodeMemory);
-	const float TimePassed = (OwnerComp.GetWorld()->GetTimeSeconds() - DecoratorMemory->LastUseTimestamp);
+	const double TimePassed = (OwnerComp.GetWorld()->GetTimeSeconds() - DecoratorMemory->LastUseTimestamp);
 	return TimePassed >= CoolDownTime;
 }
 
@@ -33,7 +33,7 @@ void UBTDecorator_Cooldown::InitializeMemory(UBehaviorTreeComponent& OwnerComp,
 	FBTCooldownDecoratorMemory* DecoratorMemory = CastInstanceNodeMemory<FBTCooldownDecoratorMemory>(NodeMemory);
 	if (InitType == EBTMemoryInit::Initialize)
 	{
-		DecoratorMemory->LastUseTimestamp = -FLT_MAX;
+		DecoratorMemory->LastUseTimestamp = DBL_MIN;
 	}
 
 	DecoratorMemory->bRequestedRestart = false;
@@ -51,7 +51,7 @@ void UBTDecorator_Cooldown::TickNode(UBehaviorTreeComponent& OwnerComp, uint8* N
 	FBTCooldownDecoratorMemory* DecoratorMemory = CastInstanceNodeMemory<FBTCooldownDecoratorMemory>(NodeMemory);
 	if (!DecoratorMemory->bRequestedRestart)
 	{
-		const float TimePassed = (OwnerComp.GetWorld()->GetTimeSeconds() - DecoratorMemory->LastUseTimestamp);
+		const double TimePassed = (OwnerComp.GetWorld()->GetTimeSeconds() - DecoratorMemory->LastUseTimestamp);
 		if (TimePassed >= CoolDownTime)
 		{
 			DecoratorMemory->bRequestedRestart = true;
@@ -72,7 +72,7 @@ void UBTDecorator_Cooldown::DescribeRuntimeValues(const UBehaviorTreeComponent&
 	Super::DescribeRuntimeValues(OwnerComp, NodeMemory, Verbosity, Values);
 
 	FBTCooldownDecoratorMemory* DecoratorMemory = CastInstanceNodeMemory<FBTCooldownDecoratorMemory>(NodeMemory);
-	const float TimePassed = OwnerComp.GetWorld()->GetTimeSeconds() - DecoratorMemory->LastUseTimestamp;
+	const double TimePassed = OwnerComp.GetWorld()->GetTimeSeconds() - DecoratorMemory->LastUseTimestamp;
 	
 	if (TimePassed < CoolDownTime)
 	{
diff --git a/Engine/Source/Runtime/AIModule/Private/BehaviorTree/Decorators/BTDecorator_IsAtLocation.cpp b/Engine/Source/Runtime/AIModule/Private/BehaviorTree/Decorators/BTDecorator_IsAtLocation.cpp
--- a/Engine/Source/Runtime/AIModule/Private/BehaviorTree/Decorators/BTDecorator_IsAtLocation.cpp
+++ b/Engine/Source/Runtime/AIModule/Private/BehaviorTree/Decorators/BTDecorator_IsAtLocation.cpp
@@ -30,9 +30,10 @@ UBTDecorator_IsAtLocation::UBTDecorator_IsAtLocation(const FObjectInitializer& O
 	FlowAbortMode = EBTFlowAbortMode::None;
 }
 
-float UBTDecorator_IsAtLocation::GetGeometricDistanceSquared(const FVector& A, const FVector& B) const
+FVector::FReal UBTDecorator_IsAtLocation::GetGeometricDistanceSquared(const FVector& A, const FVector& B) const
 {
-	float Result = MAX_flt;
+	FVector::FReal Result = TNumericLimits<FVector::FReal>::Max();
+
 	switch (GeometricDistanceType)
 	{
 	case FAIDistanceType::Distance3D:
@@ -90,7 +91,7 @@ bool UBTDecorator_IsAtLocation::CalculateRawConditionValue(UBehaviorTreeComponen
 					? PathFollowingComponent->HasReached(TargetLocation, EPathFollowingReachMode::OverlapAgent, Radius)
 					: (AIOwner->GetPawn()
 						? (GetGeometricDistanceSquared(AIOwner->GetPawn()->GetActorLocation(), TargetLocation) < FMath::Square(Radius))
-						: false);;
+						: false);
 			}
 		}
 	}
diff --git a/Engine/Source/Runtime/AIModule/Private/BehaviorTree/Decorators/BTDecorator_KeepInCone.cpp b/Engine/Source/Runtime/AIModule/Private/BehaviorTree/Decorators/BTDecorator_KeepInCone.cpp
--- a/Engine/Source/Runtime/AIModule/Private/BehaviorTree/Decorators/BTDecorator_KeepInCone.cpp
+++ b/Engine/Source/Runtime/AIModule/Private/BehaviorTree/Decorators/BTDecorator_KeepInCone.cpp
@@ -91,7 +91,7 @@ void UBTDecorator_KeepInCone::TickNode(UBehaviorTreeComponent& OwnerComp, uint8*
 	
 	if (CalculateCurrentDirection(OwnerComp, CurrentDir))
 	{
-		const float Angle = DecoratorMemory->InitialDirection.CosineAngle2D(CurrentDir);
+		const FVector::FReal Angle = DecoratorMemory->InitialDirection.CosineAngle2D(CurrentDir);
 		if (Angle < ConeHalfAngleDot || (IsInversed() && Angle > ConeHalfAngleDot))
 		{
 			OwnerComp.RequestExecution(this);
@@ -116,8 +116,8 @@ void UBTDecorator_KeepInCone::DescribeRuntimeValues(const UBehaviorTreeComponent
 	
 	if (CalculateCurrentDirection(OwnerComp, CurrentDir))
 	{
-		const float CurrentAngleDot = DecoratorMemory->InitialDirection.CosineAngle2D(CurrentDir);
-		const float CurrentAngleRad = FMath::Acos(CurrentAngleDot);
+		const FVector::FReal CurrentAngleDot = DecoratorMemory->InitialDirection.CosineAngle2D(CurrentDir);
+		const FVector::FReal CurrentAngleRad = FMath::Acos(CurrentAngleDot);
 
 		Values.Add(FString::Printf(TEXT("Angle: %.0f (%s cone)"),
 			FMath::RadiansToDegrees(CurrentAngleRad),
diff --git a/Engine/Source/Runtime/AIModule/Private/BehaviorTree/Decorators/BTDecorator_Loop.cpp b/Engine/Source/Runtime/AIModule/Private/BehaviorTree/Decorators/BTDecorator_Loop.cpp
--- a/Engine/Source/Runtime/AIModule/Private/BehaviorTree/Decorators/BTDecorator_Loop.cpp
+++ b/Engine/Source/Runtime/AIModule/Private/BehaviorTree/Decorators/BTDecorator_Loop.cpp
@@ -29,7 +29,7 @@ void UBTDecorator_Loop::OnNodeActivation(FBehaviorTreeSearchData& SearchData)
 		(!bIsSpecialNode && ParentMemory->CurrentChild != ChildIndex))
 	{
 		// initialize counter if it's first activation
-		DecoratorMemory->RemainingExecutions = NumLoops;
+		DecoratorMemory->RemainingExecutions = IntCastChecked<uint8>(NumLoops);
 		DecoratorMemory->TimeStarted = GetWorld()->GetTimeSeconds();
 	}
 
@@ -97,7 +97,7 @@ void UBTDecorator_Loop::DescribeRuntimeValues(const UBehaviorTreeComponent& Owne
 	{
 		FBTLoopDecoratorMemory* DecoratorMemory = (FBTLoopDecoratorMemory*)NodeMemory;
 
-		const float TimeRemaining = FMath::Max(InfiniteLoopTimeoutTime - (GetWorld()->GetTimeSeconds() - DecoratorMemory->TimeStarted), 0.f);
+		const double TimeRemaining = FMath::Max(InfiniteLoopTimeoutTime - (GetWorld()->GetTimeSeconds() - DecoratorMemory->TimeStarted), 0.f);
 		Values.Add(FString::Printf(TEXT("time remaining: %s"), *FString::SanitizeFloat(TimeRemaining)));
 	}
 }
diff --git a/Engine/Source/Runtime/AIModule/Private/BehaviorTree/Decorators/BTDecorator_TagCooldown.cpp b/Engine/Source/Runtime/AIModule/Private/BehaviorTree/Decorators/BTDecorator_TagCooldown.cpp
--- a/Engine/Source/Runtime/AIModule/Private/BehaviorTree/Decorators/BTDecorator_TagCooldown.cpp
+++ b/Engine/Source/Runtime/AIModule/Private/BehaviorTree/Decorators/BTDecorator_TagCooldown.cpp
@@ -25,9 +25,9 @@ void UBTDecorator_TagCooldown::PostLoad()
 
 bool UBTDecorator_TagCooldown::HasCooldownFinished(const UBehaviorTreeComponent& OwnerComp) const
 {
-	const float TagCooldownEndTime = OwnerComp.GetTagCooldownEndTime(CooldownTag);
+	const double TagCooldownEndTime = OwnerComp.GetTagCooldownEndTime(CooldownTag);
 
-	if (TagCooldownEndTime == 0.f)
+	if (TagCooldownEndTime == 0.)
 	{
 		// special case, we don't have an end time yet for this cooldown tag
 		return true;
@@ -81,12 +81,12 @@ void UBTDecorator_TagCooldown::DescribeRuntimeValues(const UBehaviorTreeComponen
 {
 	Super::DescribeRuntimeValues(OwnerComp, NodeMemory, Verbosity, Values);
 	
-	const float TagCooldownEndTime = OwnerComp.GetTagCooldownEndTime(CooldownTag);
+	const double TagCooldownEndTime = OwnerComp.GetTagCooldownEndTime(CooldownTag);
 
 	// if the tag cooldown end time is 0.f then it hasn't been set yet.
-	if (TagCooldownEndTime > 0.f)
+	if (TagCooldownEndTime > 0.)
 	{
-		const float TimePassed = (OwnerComp.GetWorld()->GetTimeSeconds() - TagCooldownEndTime);
+		const double TimePassed = (OwnerComp.GetWorld()->GetTimeSeconds() - TagCooldownEndTime);
 
 		if (TimePassed < CooldownDuration)
 		{
diff --git a/Engine/Source/Runtime/AIModule/Private/BehaviorTree/Tasks/BTTask_RotateToFaceBBEntry.cpp b/Engine/Source/Runtime/AIModule/Private/BehaviorTree/Tasks/BTTask_RotateToFaceBBEntry.cpp
--- a/Engine/Source/Runtime/AIModule/Private/BehaviorTree/Tasks/BTTask_RotateToFaceBBEntry.cpp
+++ b/Engine/Source/Runtime/AIModule/Private/BehaviorTree/Tasks/BTTask_RotateToFaceBBEntry.cpp
@@ -39,7 +39,7 @@ void UBTTask_RotateToFaceBBEntry::PostLoad()
 
 namespace
 {
-	FORCEINLINE_DEBUGGABLE float CalculateAngleDifferenceDot(const FVector& VectorA, const FVector& VectorB)
+	FORCEINLINE_DEBUGGABLE FVector::FReal CalculateAngleDifferenceDot(const FVector& VectorA, const FVector& VectorB)
 	{
 		return (VectorA.IsNearlyZero() || VectorB.IsNearlyZero())
             ? 1.f
@@ -73,7 +73,7 @@ EBTNodeResult::Type UBTTask_RotateToFaceBBEntry::ExecuteTask(UBehaviorTreeCompon
 
 		if (ActorValue != NULL)
 		{
-			const float AngleDifference = CalculateAngleDifferenceDot(Pawn->GetActorForwardVector()
+			const FVector::FReal AngleDifference = CalculateAngleDifferenceDot(Pawn->GetActorForwardVector()
 				, (ActorValue->GetActorLocation() - PawnLocation));
 			
 			if (AngleDifference >= PrecisionDot)
@@ -95,7 +95,7 @@ EBTNodeResult::Type UBTTask_RotateToFaceBBEntry::ExecuteTask(UBehaviorTreeCompon
 		
 		if (FAISystem::IsValidLocation(KeyValue))
 		{
-			const float AngleDifference = CalculateAngleDifferenceDot(Pawn->GetActorForwardVector()
+			const FVector::FReal AngleDifference = CalculateAngleDifferenceDot(Pawn->GetActorForwardVector()
 				, (KeyValue - PawnLocation));
 
 			if (AngleDifference >= PrecisionDot)
@@ -117,7 +117,7 @@ EBTNodeResult::Type UBTTask_RotateToFaceBBEntry::ExecuteTask(UBehaviorTreeCompon
 		if (FAISystem::IsValidRotation(KeyValue))
 		{
 			const FVector DirectionVector = KeyValue.Vector();
-			const float AngleDifference = CalculateAngleDifferenceDot(Pawn->GetActorForwardVector(), DirectionVector);
+			const FVector::FReal AngleDifference = CalculateAngleDifferenceDot(Pawn->GetActorForwardVector(), DirectionVector);
 
 			if (AngleDifference >= PrecisionDot)
 			{
@@ -213,7 +213,7 @@ void UBTTask_RotateToFaceBBEntry::DescribeRuntimeValues(const UBehaviorTreeCompo
 
 		if (FocalPoint != FAISystem::InvalidLocation)
 		{
-			const float CurrentAngleRadians = CalculateAngleDifferenceDot(PawnDirection, (FocalPoint - AIController->GetPawn()->GetActorLocation()));
+			const FVector::FReal CurrentAngleRadians = CalculateAngleDifferenceDot(PawnDirection, (FocalPoint - AIController->GetPawn()->GetActorLocation()));
 			Values.Add(FString::Printf(TEXT("Current angle: %.2f"), FMath::RadiansToDegrees(FMath::Acos(CurrentAngleRadians))));
 		}
 		else
diff --git a/Engine/Source/Runtime/AIModule/Private/EnvironmentQuery/EQSRenderingComponent.cpp b/Engine/Source/Runtime/AIModule/Private/EnvironmentQuery/EQSRenderingComponent.cpp
--- a/Engine/Source/Runtime/AIModule/Private/EnvironmentQuery/EQSRenderingComponent.cpp
+++ b/Engine/Source/Runtime/AIModule/Private/EnvironmentQuery/EQSRenderingComponent.cpp
@@ -26,7 +26,7 @@ namespace FEQSRenderingHelper
 //----------------------------------------------------------------------//
 // FEQSSceneProxy
 //----------------------------------------------------------------------//
-const FVector FEQSSceneProxy::ItemDrawRadius(30,30,30);
+const FVector3f FEQSSceneProxy::ItemDrawRadius(30.f,30.f,30.f);
 
 SIZE_T FEQSSceneProxy::GetTypeHash() const
 {
diff --git a/Engine/Source/Runtime/AIModule/Private/EnvironmentQuery/EQSTestingPawn.cpp b/Engine/Source/Runtime/AIModule/Private/EnvironmentQuery/EQSTestingPawn.cpp
--- a/Engine/Source/Runtime/AIModule/Private/EnvironmentQuery/EQSTestingPawn.cpp
+++ b/Engine/Source/Runtime/AIModule/Private/EnvironmentQuery/EQSTestingPawn.cpp
@@ -264,7 +264,7 @@ void AEQSTestingPawn::MakeOneStep()
 	// possible still not valid 
 	if (QueryInstance.IsValid() == true && QueryInstance->IsFinished() == false)
 	{
-		QueryInstance->ExecuteOneStep(double(TimeLimitPerStep));
+		QueryInstance->ExecuteOneStep(TimeLimitPerStep);
 		StepResults.Add(*(QueryInstance.Get()));
 
 		if (QueryInstance->IsFinished())
diff --git a/Engine/Source/Runtime/AIModule/Private/EnvironmentQuery/EnvQueryInstance.cpp b/Engine/Source/Runtime/AIModule/Private/EnvironmentQuery/EnvQueryInstance.cpp
--- a/Engine/Source/Runtime/AIModule/Private/EnvironmentQuery/EnvQueryInstance.cpp
+++ b/Engine/Source/Runtime/AIModule/Private/EnvironmentQuery/EnvQueryInstance.cpp
@@ -362,7 +362,7 @@ void FEnvQueryInstance::ExecuteOneStep(float TimeLimit)
 					}
 
 					const double GenTime = FPlatformTime::Seconds();
-					const float StepExecutionTime = GenTime - StepStartTime;
+					const double StepExecutionTime = GenTime - StepStartTime;
 					StepStartTime += StepExecutionTime;
 					TotalExecutionTime += StepExecutionTime;
 					StartTime = GenTime;
@@ -457,7 +457,7 @@ void FEnvQueryInstance::ExecuteOneStep(float TimeLimit)
 			*QueryName, OptionIndex, CurrentTest);
 	}
 
-	const float StepExecutionTime = FPlatformTime::Seconds() - StepStartTime;
+	const double StepExecutionTime = FPlatformTime::Seconds() - StepStartTime;
 	TotalExecutionTime += StepExecutionTime;
 
 #if USE_EQS_DEBUGGER
@@ -503,7 +503,7 @@ void FEnvQueryInstance::ExecuteOneStep(float TimeLimit)
 
 FString FEnvQueryInstance::GetExecutionTimeDescription() const
 {
-	FString Description = FString::Printf(TEXT("Total Execution Time: %.2f ms"), TotalExecutionTime * 1000.f);
+	FString Description = FString::Printf(TEXT("Total Execution Time: %.2f ms"), TotalExecutionTime * 1000.);
 
 #if USE_EQS_DEBUGGER
 	for (int32 OptionIdx = 0; OptionIdx <= OptionIndex; OptionIdx++)
@@ -564,8 +564,8 @@ FEnvQueryInstance::FItemIterator::FItemIterator(const UEnvQueryTest* QueryTest,
 {
 	check(QueryTest);
 
-	CachedFilterOp = QueryTest->MultipleContextFilterOp.GetValue();
-	CachedScoreOp = QueryTest->MultipleContextScoreOp.GetValue();
+	CachedFilterOp = QueryTest->MultipleContextFilterOp.GetIntValue();
+	CachedScoreOp = QueryTest->MultipleContextScoreOp.GetIntValue();
 	bIsFiltering = (QueryTest->TestPurpose == EEnvTestPurpose::Filter) || (QueryTest->TestPurpose == EEnvTestPurpose::FilterAndScore);
 
 	Deadline = QueryInstance.CurrentStepTimeLimit > 0.0 ? (FPlatformTime::Seconds() + QueryInstance.CurrentStepTimeLimit) : -1.0;
@@ -884,7 +884,7 @@ void FEnvQueryInstance::FinalizeTest()
 #if STATS
 uint32 FEnvQueryInstance::GetAllocatedSize() const
 {
-	uint32 MemSize = sizeof(*this) + Items.GetAllocatedSize() + RawData.GetAllocatedSize();
+	SIZE_T MemSize = sizeof(*this) + Items.GetAllocatedSize() + RawData.GetAllocatedSize();
 	MemSize += GetContextAllocatedSize();
 	MemSize += NamedParams.GetAllocatedSize();
 	MemSize += ItemDetails.GetAllocatedSize();
@@ -895,18 +895,18 @@ uint32 FEnvQueryInstance::GetAllocatedSize() const
 		MemSize += Options[OptionCount].GetAllocatedSize();
 	}
 
-	return MemSize;
+	return IntCastChecked<uint32>(MemSize);
 }
 
 uint32 FEnvQueryInstance::GetContextAllocatedSize() const
 {
-	uint32 MemSize = ContextCache.GetAllocatedSize();
+	SIZE_T MemSize = ContextCache.GetAllocatedSize();
 	for (TMap<UClass*, FEnvQueryContextData>::TConstIterator It(ContextCache); It; ++It)
 	{
 		MemSize += It.Value().GetAllocatedSize();
 	}
 
-	return MemSize;
+	return IntCastChecked<uint32>(MemSize);
 }
 #endif // STATS
 
diff --git a/Engine/Source/Runtime/AIModule/Private/EnvironmentQuery/EnvQueryManager.cpp b/Engine/Source/Runtime/AIModule/Private/EnvironmentQuery/EnvQueryManager.cpp
--- a/Engine/Source/Runtime/AIModule/Private/EnvironmentQuery/EnvQueryManager.cpp
+++ b/Engine/Source/Runtime/AIModule/Private/EnvironmentQuery/EnvQueryManager.cpp
@@ -119,12 +119,12 @@ FEnvQueryRequest& FEnvQueryRequest::SetDynamicParam(const FAIDynamicParam& Param
 		break;
 		case EAIParamType::Int:
 		{
-			SetIntParam(Param.ParamName, Param.Value);
+			SetIntParam(Param.ParamName, static_cast<int32>(Param.Value));
 		}
 		break;
 		case EAIParamType::Bool:
 		{
-			bool Result = Param.Value > 0;
+			bool Result = Param.Value > 0.;
 			SetBoolParam(Param.ParamName, Result);
 		}
 		break;
@@ -427,7 +427,7 @@ void UEnvQueryManager::Tick(float DeltaTime)
 	CheckQueryCount();
 #endif
 
-	float TimeLeft = MaxAllowedTestingTime;
+	double TimeLeft = MaxAllowedTestingTime;
 	int32 QueriesFinishedDuringUpdate = 0;
 
 	{
@@ -443,7 +443,7 @@ void UEnvQueryManager::Tick(float DeltaTime)
 			&& (QueriesFinishedDuringUpdate + NumRunningQueriesAbortedSinceLastUpdate < NumRunningQueries))
 		{
 			const double StepStartTime = FPlatformTime::Seconds();
-			float ResultHandlingDuration = 0.0f;
+			double ResultHandlingDuration = 0.;
 #if USE_EQS_DEBUGGER
 			bool bWorkHasBeenDone = false;
 #endif // USE_EQS_DEBUGGER
@@ -535,7 +535,7 @@ void UEnvQueryManager::Tick(float DeltaTime)
 			if (bAllowEQSTimeSlicing) // if Time slicing is enabled...
 #endif
 			{	// Don't include the querier handling as part of the total time spent by EQS for time-slicing purposes.
-				const float StepProcessingTime = ((FPlatformTime::Seconds() - StepStartTime) - ResultHandlingDuration);
+				const double StepProcessingTime = ((FPlatformTime::Seconds() - StepStartTime) - ResultHandlingDuration);
 				TimeLeft -= StepProcessingTime;
 
 #if USE_EQS_DEBUGGER
@@ -551,7 +551,7 @@ void UEnvQueryManager::Tick(float DeltaTime)
 
 	{
 		const int32 NumQueriesFinished = QueriesFinishedDuringUpdate + NumRunningQueriesAbortedSinceLastUpdate;
-		float FinishedQueriesTotalTime = 0.0;
+		double FinishedQueriesTotalTime = 0.;
 
 		if (NumQueriesFinished > 0)
 		{
@@ -598,7 +598,7 @@ void UEnvQueryManager::Tick(float DeltaTime)
 		// Reset the running queries aborted since last update counter
 		NumRunningQueriesAbortedSinceLastUpdate = 0;
 
-		const float InstanceAverageResponseTime = (NumQueriesFinished > 0) ? (1000.0f * FinishedQueriesTotalTime / NumQueriesFinished) : 0.0f;
+		const double InstanceAverageResponseTime = (NumQueriesFinished > 0) ? (1000. * FinishedQueriesTotalTime / NumQueriesFinished) : 0.;
 		SET_FLOAT_STAT(STAT_AI_EQS_AvgInstanceResponseTime, InstanceAverageResponseTime);
 	}
 }
@@ -1206,7 +1206,7 @@ void FEQSDebugger::StoreTickTime(const FEnvQueryInstance& QueryInstance, float T
 	UpdateInfo.LastTickEntry = (TickIdx > UpdateInfo.LastTickEntry) ? TickIdx : UpdateInfo.LastTickEntry;
 
 	UpdateInfo.LastTickTime += TickTime;
-	UpdateInfo.TickPct[TickIdx] = FMath::Min(255, FMath::TruncToInt(255 * UpdateInfo.LastTickTime / MaxTickTime));
+	UpdateInfo.TickPct[TickIdx] = static_cast<uint8>(FMath::Min(255, FMath::TruncToInt(255 * UpdateInfo.LastTickTime / MaxTickTime)));
 #endif // USE_EQS_TICKLOADDATA
 }
 
diff --git a/Engine/Source/Runtime/AIModule/Private/EnvironmentQuery/EnvQueryTraceHelpers.h b/Engine/Source/Runtime/AIModule/Private/EnvironmentQuery/EnvQueryTraceHelpers.h
--- a/Engine/Source/Runtime/AIModule/Private/EnvironmentQuery/EnvQueryTraceHelpers.h
+++ b/Engine/Source/Runtime/AIModule/Private/EnvironmentQuery/EnvQueryTraceHelpers.h
@@ -64,7 +64,7 @@ namespace FEQSHelpers
 		FORCEINLINE_DEBUGGABLE bool RunSphereTrace(const FVector& StartPos, const FVector& EndPos, FVector& HitPos)
 		{
 			FHitResult OutHit;
-			const bool bHit = World->SweepSingleByChannel(OutHit, StartPos, EndPos, FQuat::Identity, Channel, FCollisionShape::MakeSphere(Extent.X), QueryParams, ResponseParams);
+			const bool bHit = World->SweepSingleByChannel(OutHit, StartPos, EndPos, FQuat::Identity, Channel, FCollisionShape::MakeSphere(FloatCastChecked<float>(Extent.X, UE::LWC::DefaultFloatPrecision)), QueryParams, ResponseParams);
 			HitPos = OutHit.Location;
 			return bHit;
 		}
@@ -72,7 +72,7 @@ namespace FEQSHelpers
 		FORCEINLINE_DEBUGGABLE bool RunCapsuleTrace(const FVector& StartPos, const FVector& EndPos, FVector& HitPos)
 		{
 			FHitResult OutHit;
-			const bool bHit = World->SweepSingleByChannel(OutHit, StartPos, EndPos, FQuat::Identity, Channel, FCollisionShape::MakeCapsule(Extent.X, Extent.Z), QueryParams, ResponseParams);
+			const bool bHit = World->SweepSingleByChannel(OutHit, StartPos, EndPos, FQuat::Identity, Channel, FCollisionShape::MakeCapsule(FloatCastChecked<float>(Extent.X, UE::LWC::DefaultFloatPrecision), FloatCastChecked<float>(Extent.Z, 1./16.)), QueryParams, ResponseParams);
 			HitPos = OutHit.Location;
 			return bHit;
 		}
diff --git a/Engine/Source/Runtime/AIModule/Private/EnvironmentQuery/Tests/EnvQueryTest_Distance.cpp b/Engine/Source/Runtime/AIModule/Private/EnvironmentQuery/Tests/EnvQueryTest_Distance.cpp
--- a/Engine/Source/Runtime/AIModule/Private/EnvironmentQuery/Tests/EnvQueryTest_Distance.cpp
+++ b/Engine/Source/Runtime/AIModule/Private/EnvironmentQuery/Tests/EnvQueryTest_Distance.cpp
@@ -79,10 +79,10 @@ void UEnvQueryTest_Distance::RunTest(FEnvQueryInstance& QueryInstance) const
 			for (FEnvQueryInstance::ItemIterator It(this, QueryInstance); It; ++It)
 			{
 				const FVector ItemLocation = GetItemLocation(QueryInstance, It.GetIndex());
-				CheckItemLocationForNaN(ItemLocation, QueryOwner, It.GetIndex(), TestMode);
+				CheckItemLocationForNaN(ItemLocation, QueryOwner, It.GetIndex(), TestMode.GetIntValue());
 				for (int32 ContextIndex = 0; ContextIndex < ContextLocations.Num(); ContextIndex++)
 				{
-					CheckContextLocationForNaN(ContextLocations[ContextIndex], QueryOwner, ContextIndex, TestMode);
+					CheckContextLocationForNaN(ContextLocations[ContextIndex], QueryOwner, ContextIndex, TestMode.GetIntValue());
 					const float Distance = CalcDistance3D(ItemLocation, ContextLocations[ContextIndex]);
 					It.SetScore(TestPurpose, FilterType, Distance, MinThresholdValue, MaxThresholdValue);
 				}
@@ -93,10 +93,10 @@ void UEnvQueryTest_Distance::RunTest(FEnvQueryInstance& QueryInstance) const
 			for (FEnvQueryInstance::ItemIterator It(this, QueryInstance); It; ++It)
 			{
 				const FVector ItemLocation = GetItemLocation(QueryInstance, It.GetIndex());
-				CheckItemLocationForNaN(ItemLocation, QueryOwner, It.GetIndex(), TestMode);
+				CheckItemLocationForNaN(ItemLocation, QueryOwner, It.GetIndex(), TestMode.GetIntValue());
 				for (int32 ContextIndex = 0; ContextIndex < ContextLocations.Num(); ContextIndex++)
 				{
-					CheckContextLocationForNaN(ContextLocations[ContextIndex], QueryOwner, ContextIndex, TestMode);
+					CheckContextLocationForNaN(ContextLocations[ContextIndex], QueryOwner, ContextIndex, TestMode.GetIntValue());
 					const float Distance = CalcDistance2D(ItemLocation, ContextLocations[ContextIndex]);
 					It.SetScore(TestPurpose, FilterType, Distance, MinThresholdValue, MaxThresholdValue);
 				}
@@ -107,10 +107,10 @@ void UEnvQueryTest_Distance::RunTest(FEnvQueryInstance& QueryInstance) const
 			for (FEnvQueryInstance::ItemIterator It(this, QueryInstance); It; ++It)
 			{
 				const FVector ItemLocation = GetItemLocation(QueryInstance, It.GetIndex());
-				CheckItemLocationForNaN(ItemLocation, QueryOwner, It.GetIndex(), TestMode);
+				CheckItemLocationForNaN(ItemLocation, QueryOwner, It.GetIndex(), TestMode.GetIntValue());
 				for (int32 ContextIndex = 0; ContextIndex < ContextLocations.Num(); ContextIndex++)
 				{
-					CheckContextLocationForNaN(ContextLocations[ContextIndex], QueryOwner, ContextIndex, TestMode);
+					CheckContextLocationForNaN(ContextLocations[ContextIndex], QueryOwner, ContextIndex, TestMode.GetIntValue());
 					const float Distance = CalcDistanceZ(ItemLocation, ContextLocations[ContextIndex]);
 					It.SetScore(TestPurpose, FilterType, Distance, MinThresholdValue, MaxThresholdValue);
 				}
@@ -121,10 +121,10 @@ void UEnvQueryTest_Distance::RunTest(FEnvQueryInstance& QueryInstance) const
 			for (FEnvQueryInstance::ItemIterator It(this, QueryInstance); It; ++It)
 			{
 				const FVector ItemLocation = GetItemLocation(QueryInstance, It.GetIndex());
-				CheckItemLocationForNaN(ItemLocation, QueryOwner, It.GetIndex(), TestMode);
+				CheckItemLocationForNaN(ItemLocation, QueryOwner, It.GetIndex(), TestMode.GetIntValue());
 				for (int32 ContextIndex = 0; ContextIndex < ContextLocations.Num(); ContextIndex++)
 				{
-					CheckContextLocationForNaN(ContextLocations[ContextIndex], QueryOwner, ContextIndex, TestMode);
+					CheckContextLocationForNaN(ContextLocations[ContextIndex], QueryOwner, ContextIndex, TestMode.GetIntValue());
 					const float Distance = CalcDistanceAbsoluteZ(ItemLocation, ContextLocations[ContextIndex]);
 					It.SetScore(TestPurpose, FilterType, Distance, MinThresholdValue, MaxThresholdValue);
 				}
diff --git a/Engine/Source/Runtime/AIModule/Private/EnvironmentQuery/Tests/EnvQueryTest_Dot.cpp b/Engine/Source/Runtime/AIModule/Private/EnvironmentQuery/Tests/EnvQueryTest_Dot.cpp
--- a/Engine/Source/Runtime/AIModule/Private/EnvironmentQuery/Tests/EnvQueryTest_Dot.cpp
+++ b/Engine/Source/Runtime/AIModule/Private/EnvironmentQuery/Tests/EnvQueryTest_Dot.cpp
@@ -90,11 +90,11 @@ void UEnvQueryTest_Dot::RunTest(FEnvQueryInstance& QueryInstance) const
 				switch (TestMode)
 				{
 					case EEnvTestDot::Dot3D:
-						DotValue = FVector::DotProduct(LineADirs[LineAIndex], LineBDirs[LineBIndex]);
+						DotValue = static_cast<float>(FVector::DotProduct(LineADirs[LineAIndex], LineBDirs[LineBIndex]));
 						break;
 
 					case EEnvTestDot::Dot2D:
-						DotValue = LineADirs[LineAIndex].CosineAngle2D(LineBDirs[LineBIndex]);
+						DotValue = static_cast<float>(LineADirs[LineAIndex].CosineAngle2D(LineBDirs[LineBIndex]));
 						break;
 
 					default:
diff --git a/Engine/Source/Runtime/AIModule/Private/EnvironmentQuery/Tests/EnvQueryTest_Overlap.cpp b/Engine/Source/Runtime/AIModule/Private/EnvironmentQuery/Tests/EnvQueryTest_Overlap.cpp
--- a/Engine/Source/Runtime/AIModule/Private/EnvironmentQuery/Tests/EnvQueryTest_Overlap.cpp
+++ b/Engine/Source/Runtime/AIModule/Private/EnvironmentQuery/Tests/EnvQueryTest_Overlap.cpp
@@ -39,11 +39,11 @@ void UEnvQueryTest_Overlap::RunTest(FEnvQueryInstance& QueryInstance) const
 			break;
 
 		case EEnvOverlapShape::Sphere:
-			CollisionShape = FCollisionShape::MakeSphere(TraceExtent.X);
+			CollisionShape = FCollisionShape::MakeSphere(FloatCastChecked<float>(TraceExtent.X, UE::LWC::DefaultFloatPrecision));
 			break;
 
 		case EEnvOverlapShape::Capsule:
-			CollisionShape = FCollisionShape::MakeCapsule(TraceExtent.X, TraceExtent.Z);
+			CollisionShape = FCollisionShape::MakeCapsule(FloatCastChecked<float>(TraceExtent.X, UE::LWC::DefaultFloatPrecision), FloatCastChecked<float>(TraceExtent.Z, UE::LWC::DefaultFloatPrecision));
 			break;
 
 		default:
diff --git a/Engine/Source/Runtime/AIModule/Private/GameplayDebugger/GameplayDebuggerCategory_AI.cpp b/Engine/Source/Runtime/AIModule/Private/GameplayDebugger/GameplayDebuggerCategory_AI.cpp
--- a/Engine/Source/Runtime/AIModule/Private/GameplayDebugger/GameplayDebuggerCategory_AI.cpp
+++ b/Engine/Source/Runtime/AIModule/Private/GameplayDebugger/GameplayDebuggerCategory_AI.cpp
@@ -282,7 +282,7 @@ void FGameplayDebuggerCategory_AI::CollectPathData(AAIController* DebugAI)
 						NavData->GetPolyVerts(NavMeshPath->PathCorridor[Idx], PolyData.Points);
 						
 						const uint32 AreaId = NavData->GetPolyAreaID(NavMeshPath->PathCorridor[Idx]);
-						PolyData.Color = NavData->GetAreaIDColor(AreaId);
+						PolyData.Color = NavData->GetAreaIDColor(IntCastChecked<uint8>(AreaId));
 
 						PathDataPack.PathCorridor.Add(PolyData);
 					}
@@ -488,7 +488,7 @@ void FGameplayDebuggerCategory_AI::DrawOverheadInfo(AActor& DebugActor, FGamepla
 
 		float SizeX = 0.0f, SizeY = 0.0f;
 		OverheadContext.MeasureString(ActorDesc, SizeX, SizeY);
-		OverheadContext.PrintAt(ScreenLoc.X - (SizeX * 0.5f), ScreenLoc.Y - (SizeY * 1.2f), ActorDesc);
+		OverheadContext.PrintAt(static_cast<float>(ScreenLoc.X - (SizeX * 0.5f)), static_cast<float>(ScreenLoc.Y - (SizeY * 1.2f)), ActorDesc);
 	}
 }
 
@@ -512,7 +512,7 @@ void FGameplayDebuggerCategory_AI::DrawPawnIcons(UWorld* World, AActor* DebugAct
 					const FVector2D ScreenLoc = CanvasContext.ProjectLocation(IconLocation);
 					const float IconSize = (DebugActor == ItPawn) ? 32.0f : 16.0f;
 
-					CanvasContext.DrawIcon(FColor::White, CanvasIcon, ScreenLoc.X, ScreenLoc.Y - IconSize, IconSize / CanvasIcon.Texture->GetSurfaceWidth());
+					CanvasContext.DrawIcon(FColor::White, CanvasIcon, static_cast<float>(ScreenLoc.X), static_cast<float>(ScreenLoc.Y - IconSize), IconSize / CanvasIcon.Texture->GetSurfaceWidth());
 				}
 			}
 		}
diff --git a/Engine/Source/Runtime/AIModule/Private/GameplayDebugger/GameplayDebuggerCategory_BehaviorTree.cpp b/Engine/Source/Runtime/AIModule/Private/GameplayDebugger/GameplayDebuggerCategory_BehaviorTree.cpp
--- a/Engine/Source/Runtime/AIModule/Private/GameplayDebugger/GameplayDebuggerCategory_BehaviorTree.cpp
+++ b/Engine/Source/Runtime/AIModule/Private/GameplayDebugger/GameplayDebuggerCategory_BehaviorTree.cpp
@@ -54,8 +54,8 @@ void FGameplayDebuggerCategory_BehaviorTree::DrawData(APlayerController* OwnerPC
 		TArray<FString> BlackboardLines;
 		DataPack.BlackboardDesc.ParseIntoArrayLines(BlackboardLines, true);
 
-		int32 SavedDefX = CanvasContext.DefaultX;
-		int32 SavedPosY = CanvasContext.CursorY;
+		const float SavedDefX = CanvasContext.DefaultX;
+		const float SavedPosY = CanvasContext.CursorY;
 		CanvasContext.DefaultX = CanvasContext.CursorX = 600.0f;
 		CanvasContext.CursorY = CanvasContext.DefaultY;
 
diff --git a/Engine/Source/Runtime/AIModule/Private/GameplayDebugger/GameplayDebuggerCategory_EQS.cpp b/Engine/Source/Runtime/AIModule/Private/GameplayDebugger/GameplayDebuggerCategory_EQS.cpp
--- a/Engine/Source/Runtime/AIModule/Private/GameplayDebugger/GameplayDebuggerCategory_EQS.cpp
+++ b/Engine/Source/Runtime/AIModule/Private/GameplayDebugger/GameplayDebuggerCategory_EQS.cpp
@@ -184,19 +184,19 @@ int32 FGameplayDebuggerCategory_EQS::DrawLookedAtItem(const EQSDebug::FQueryData
 	const FVector CameraDirection = CanvasContext.Canvas->SceneView->GetViewDirection();
 
 	int32 BestItemIndex = INDEX_NONE;
-	float BestScore = -FLT_MAX;
+	FVector::FReal BestScore = TNumericLimits<FVector::FReal>::Min();
 
 	for (int32 Idx = 0; Idx < QueryData.RenderDebugHelpers.Num(); Idx++)
 	{
 		const EQSDebug::FDebugHelper& ItemInfo = QueryData.RenderDebugHelpers[Idx];
 		const FVector DirToItem = ItemInfo.Location - CameraLocation;
-		float DistToItem = DirToItem.Size();
+		FVector::FReal DistToItem = DirToItem.Size();
 		if (FMath::IsNearlyZero(DistToItem))
 		{
-			DistToItem = 1.0f;
+			DistToItem = 1.;
 		}
 
-		const float ItemScore = FVector::DotProduct(DirToItem, CameraDirection) / DistToItem;
+		const FVector::FReal ItemScore = FVector::DotProduct(DirToItem, CameraDirection) / DistToItem;
 		if (ItemScore > BestScore)
 		{
 			BestItemIndex = Idx;
@@ -429,8 +429,8 @@ void FGameplayDebuggerCategory_EQS::DrawDetailedItemRow(const EQSDebug::FItemDat
 	BackTileItem.Position.X = PosX + ActiveTileItem.Size.X;
 	BackTileItem.Size.X = FMath::Max(BarWidth * (1.0f - Pct), 0.0f);
 
-	CanvasContext.DrawItem(ActiveTileItem, ActiveTileItem.Position.X, ActiveTileItem.Position.Y);
-	CanvasContext.DrawItem(BackTileItem, BackTileItem.Position.X, BackTileItem.Position.Y);
+	CanvasContext.DrawItem(ActiveTileItem, (float)ActiveTileItem.Position.X, (float)ActiveTileItem.Position.Y);
+	CanvasContext.DrawItem(BackTileItem, (float)BackTileItem.Position.X, (float)BackTileItem.Position.Y);
 
 	CanvasContext.PrintfAt(PosX, PosY, FColor::Yellow, TEXT("%.2f"), TotalScoreNotNormalized);
 	PosX += FEQSDebugTable::ItemScoreWidth;
diff --git a/Engine/Source/Runtime/AIModule/Private/GameplayDebugger/GameplayDebuggerCategory_NavLocalGrid.cpp b/Engine/Source/Runtime/AIModule/Private/GameplayDebugger/GameplayDebuggerCategory_NavLocalGrid.cpp
--- a/Engine/Source/Runtime/AIModule/Private/GameplayDebugger/GameplayDebuggerCategory_NavLocalGrid.cpp
+++ b/Engine/Source/Runtime/AIModule/Private/GameplayDebugger/GameplayDebuggerCategory_NavLocalGrid.cpp
@@ -91,7 +91,7 @@ public:
 					FDynamicMeshBuilder	MeshBuilder(View->GetFeatureLevel());
 					MeshBuilder.AddVertices(FreeCellMeshVerts);
 					MeshBuilder.AddTriangles(FreeCellMeshIndices);
-					MeshBuilder.GetMesh(FMatrix::Identity, MeshColorInstance, GetDepthPriorityGroup(View), false, false, ViewIndex, Collector);
+					MeshBuilder.GetMesh(FMatrix::Identity, MeshColorInstance, IntCastChecked<uint8>((int32)GetDepthPriorityGroup(View)), false, false, ViewIndex, Collector);
 				}
 
 				if (MarkedCellMeshVerts.Num() > 0)
@@ -100,7 +100,7 @@ public:
 					FDynamicMeshBuilder	MeshBuilder(View->GetFeatureLevel());
 					MeshBuilder.AddVertices(MarkedCellMeshVerts);
 					MeshBuilder.AddTriangles(MarkedCellMeshIndices);
-					MeshBuilder.GetMesh(FMatrix::Identity, MeshColorInstance, GetDepthPriorityGroup(View), false, false, ViewIndex, Collector);
+					MeshBuilder.GetMesh(FMatrix::Identity, MeshColorInstance, IntCastChecked<uint8>((int32)GetDepthPriorityGroup(View)), false, false, ViewIndex, Collector);
 				}
 
 				if (BoundsMeshVerts.Num() > 0)
@@ -109,7 +109,7 @@ public:
 					FDynamicMeshBuilder	MeshBuilder(View->GetFeatureLevel());
 					MeshBuilder.AddVertices(BoundsMeshVerts);
 					MeshBuilder.AddTriangles(BoundsMeshIndices);
-					MeshBuilder.GetMesh(FMatrix::Identity, MeshColorInstance, GetDepthPriorityGroup(View), false, false, ViewIndex, Collector);
+					MeshBuilder.GetMesh(FMatrix::Identity, MeshColorInstance, IntCastChecked<uint8>((int32)GetDepthPriorityGroup(View)), false, false, ViewIndex, Collector);
 				}
 
 				if (PathMeshVerts.Num() > 0)
@@ -118,7 +118,7 @@ public:
 					FDynamicMeshBuilder	MeshBuilder(View->GetFeatureLevel());
 					MeshBuilder.AddVertices(PathMeshVerts);
 					MeshBuilder.AddTriangles(PathMeshIndices);
-					MeshBuilder.GetMesh(FMatrix::Identity, MeshColorInstance, GetDepthPriorityGroup(View), false, false, ViewIndex, Collector);
+					MeshBuilder.GetMesh(FMatrix::Identity, MeshColorInstance, IntCastChecked<uint8>((int32)GetDepthPriorityGroup(View)), false, false, ViewIndex, Collector);
 				}
 			}
 		}
@@ -300,13 +300,13 @@ protected:
 
 	uint32 GetAllocatedSize(void) const
 	{
-		uint32 AllocatedSize = FDebugRenderSceneProxy::GetAllocatedSize() +
+		const SIZE_T AllocatedSize = FDebugRenderSceneProxy::GetAllocatedSize() +
 			FreeCellMeshVerts.GetAllocatedSize() + FreeCellMeshIndices.GetAllocatedSize() +
 			MarkedCellMeshVerts.GetAllocatedSize() + MarkedCellMeshIndices.GetAllocatedSize() +
 			BoundsMeshVerts.GetAllocatedSize() + BoundsMeshIndices.GetAllocatedSize() +
 			PathMeshVerts.GetAllocatedSize() + PathMeshIndices.GetAllocatedSize();
 
-		return AllocatedSize;
+		return IntCastChecked<uint32>(AllocatedSize);
 	}
 
 private:
diff --git a/Engine/Source/Runtime/AIModule/Private/GameplayDebugger/GameplayDebuggerCategory_Navmesh.cpp b/Engine/Source/Runtime/AIModule/Private/GameplayDebugger/GameplayDebuggerCategory_Navmesh.cpp
--- a/Engine/Source/Runtime/AIModule/Private/GameplayDebugger/GameplayDebuggerCategory_Navmesh.cpp
+++ b/Engine/Source/Runtime/AIModule/Private/GameplayDebugger/GameplayDebuggerCategory_Navmesh.cpp
@@ -119,8 +119,8 @@ void FGameplayDebuggerCategory_Navmesh::CollectData(APlayerController* OwnerPC,
 		if (NavSys) 
 		{
 			DataPack.NumDirtyAreas = NavSys->GetNumDirtyAreas();
-			DataPack.NumRunningTasks = NavSys->GetNumRunningBuildTasks();
-			DataPack.NumRemainingTasks = NavSys->GetNumRemainingBuildTasks();
+			DataPack.NumRunningTasks = IntCastChecked<uint16>(NavSys->GetNumRunningBuildTasks());
+			DataPack.NumRemainingTasks = IntCastChecked<uint16>(NavSys->GetNumRemainingBuildTasks());
 			DataPack.bIsNavOctreeLocked = NavSys->IsNavigationOctreeLocked();
 			DataPack.bIsNavBuildLocked = NavSys->IsNavigationBuildingLocked();
 			DataPack.NavBuildLockStatusDesc = FString("Unknown");
diff --git a/Engine/Source/Runtime/AIModule/Private/Navigation/CrowdFollowingComponent.cpp b/Engine/Source/Runtime/AIModule/Private/Navigation/CrowdFollowingComponent.cpp
--- a/Engine/Source/Runtime/AIModule/Private/Navigation/CrowdFollowingComponent.cpp
+++ b/Engine/Source/Runtime/AIModule/Private/Navigation/CrowdFollowingComponent.cpp
@@ -390,9 +390,9 @@ void UCrowdFollowingComponent::ApplyCrowdAgentVelocity(const FVector& NewVelocit
 			const bool bAccelerationBased = MovementComp->UseAccelerationForPathFollowing();
 			if (bAccelerationBased)
 			{
-				const float MaxSpeed = GetCrowdAgentMaxSpeed();
-				const float NewSpeed = NewVelocity.Size();
-				const float SpeedPct = FMath::Clamp(NewSpeed / MaxSpeed, 0.0f, 1.0f);
+				const FVector::FReal MaxSpeed = GetCrowdAgentMaxSpeed();
+				const FVector::FReal NewSpeed = NewVelocity.Size();
+				const FVector::FReal SpeedPct = FMath::Clamp(NewSpeed / MaxSpeed, 0., 1.);
 				const FVector MoveInput = FMath::IsNearlyZero(NewSpeed) ? FVector::ZeroVector : ((NewVelocity / NewSpeed) * SpeedPct);
 
 				MovementComp->RequestPathMove(MoveInput);
@@ -717,7 +717,7 @@ void LogPathPartHelper(AActor* LogOwner, FNavMeshPath* NavMeshPath, int32 StartI
 		TArray<FVector> Verts;
 		for (int32 Idx = StartIdx; Idx <= EndIdx; Idx++)
 		{
-			const uint8 AreaID = NavMesh->GetPolyAreaID(NavMeshPath->PathCorridor[Idx]);
+			const uint8 AreaID = IntCastChecked<uint8>(NavMesh->GetPolyAreaID(NavMeshPath->PathCorridor[Idx]));
 			const UClass* AreaClass = NavMesh->GetAreaClass(AreaID);
 
 			Verts.Reset();
@@ -1192,7 +1192,7 @@ void UCrowdFollowingComponent::GetDebugStringTokens(TArray<FString>& Tokens, TAr
 		const FVector CurrentLocation = MovementComp->GetActorFeetLocation();
 
 		// make sure we're not too close to end of path part (poly count can always fail when AI goes off path)
-		const float DistSq = (GetCurrentTargetLocation() - CurrentLocation).SizeSquared();
+		const FVector::FReal DistSq = (GetCurrentTargetLocation() - CurrentLocation).SizeSquared();
 		const float PathSwitchThresSq = FMath::Square(AgentRadius * 5.0f);
 
 		Tokens.Add(TEXT("distance"));
diff --git a/Engine/Source/Runtime/AIModule/Private/Navigation/CrowdManager.cpp b/Engine/Source/Runtime/AIModule/Private/Navigation/CrowdManager.cpp
--- a/Engine/Source/Runtime/AIModule/Private/Navigation/CrowdManager.cpp
+++ b/Engine/Source/Runtime/AIModule/Private/Navigation/CrowdManager.cpp
@@ -743,7 +743,7 @@ void UCrowdManager::GetAgentParams(const ICrowdAgentInterface* Agent, dtCrowdAge
 		AgentParams.collisionQueryRange = CrowdComponent->GetCrowdCollisionQueryRange();
 		AgentParams.pathOptimizationRange = CrowdComponent->GetCrowdPathOptimizationRange();
 		AgentParams.separationWeight = CrowdComponent->GetCrowdSeparationWeight();
-		AgentParams.obstacleAvoidanceType = CrowdComponent->GetCrowdAvoidanceQuality();
+		AgentParams.obstacleAvoidanceType = IntCastChecked<unsigned char>((int32)CrowdComponent->GetCrowdAvoidanceQuality());
 		AgentParams.avoidanceQueryMultiplier = CrowdComponent->GetCrowdAvoidanceRangeMultiplier();
 
 		if (CrowdComponent->IsCrowdSimulationEnabled())
@@ -1282,7 +1282,7 @@ void UCrowdManager::DebugTick() const
 					FVector Pt0 = Recast2UnrealPoint(s);
 					FVector Pt1 = Recast2UnrealPoint(s + 3);
 
-					UE_VLOG_SEGMENT_THICK(LogOwner, LogCrowdFollowing, Log, Pt0 + FCrowdDebug::Offset, Pt1 + FCrowdDebug::Offset, Color, 3.0f, TEXT(""));
+					UE_VLOG_SEGMENT_THICK(LogOwner, LogCrowdFollowing, Log, Pt0 + FCrowdDebug::Offset, Pt1 + FCrowdDebug::Offset, Color, 3, TEXT(""));
 				}
 			}
 		}
diff --git a/Engine/Source/Runtime/AIModule/Private/Navigation/GridPathFollowingComponent.cpp b/Engine/Source/Runtime/AIModule/Private/Navigation/GridPathFollowingComponent.cpp
--- a/Engine/Source/Runtime/AIModule/Private/Navigation/GridPathFollowingComponent.cpp
+++ b/Engine/Source/Runtime/AIModule/Private/Navigation/GridPathFollowingComponent.cpp
@@ -145,7 +145,7 @@ void UGridPathFollowingComponent::UpdateActiveGrid(const FVector& CurrentLocatio
 					for (int32 Idx = 1; Idx < GridPathPoints.Num(); Idx++)
 					{
 						UE_VLOG_BOX(this, LogPathFollowing, Log, FBox::BuildAABB(GridPathPoints[Idx] + DebugDrawOffset, DebugPathPointExtent), FColor::Yellow, TEXT(""));
-						UE_VLOG_SEGMENT_THICK(this, LogPathFollowing, Log, GridPathPoints[Idx - 1] + DebugDrawOffset, GridPathPoints[Idx] + DebugDrawOffset, FColor::Yellow, 3.0f, TEXT(""));
+						UE_VLOG_SEGMENT_THICK(this, LogPathFollowing, Log, GridPathPoints[Idx - 1] + DebugDrawOffset, GridPathPoints[Idx] + DebugDrawOffset, FColor::Yellow, 3, TEXT(""));
 					}
 
 					for (int32 Idx = 0; Idx < GridData.GetCellsCount(); Idx++)
diff --git a/Engine/Source/Runtime/AIModule/Private/Navigation/MetaNavMeshPath.cpp b/Engine/Source/Runtime/AIModule/Private/Navigation/MetaNavMeshPath.cpp
--- a/Engine/Source/Runtime/AIModule/Private/Navigation/MetaNavMeshPath.cpp
+++ b/Engine/Source/Runtime/AIModule/Private/Navigation/MetaNavMeshPath.cpp
@@ -238,7 +238,7 @@ void FMetaNavMeshPath::DescribeSelfToVisLog(FVisualLogEntry* Snapshot) const
 		Element.Category = LogNavigation.GetCategoryName();
 		Element.SetColor(FColorList::Yellow);
 		Element.Points.Reserve(Waypoints.Num());
-		Element.Thicknes = 3.f;
+		Element.Thicknes = 3;
 
 		const FVector DrawingOffset(0, 0, 50);
 		for (int32 Idx = 0; Idx < Waypoints.Num(); Idx++)
diff --git a/Engine/Source/Runtime/AIModule/Private/Navigation/NavLocalGridData.cpp b/Engine/Source/Runtime/AIModule/Private/Navigation/NavLocalGridData.cpp
--- a/Engine/Source/Runtime/AIModule/Private/Navigation/NavLocalGridData.cpp
+++ b/Engine/Source/Runtime/AIModule/Private/Navigation/NavLocalGridData.cpp
@@ -14,7 +14,11 @@ FNavLocalGridData::FNavLocalGridData(const FVector& Center, float Extent2D) : Gr
 	const FBox UseBounds = FBox::BuildAABB(CenterOnGrid, FVector(Extent2D, Extent2D, 0.0f));
 	Init(UseCellSize, UseBounds);
 	
-	OriginWorldCoord = FIntVector(FMath::TruncToInt(Origin.X / UseCellSize), FMath::TruncToInt(Origin.Y / UseCellSize), FMath::TruncToInt(Origin.Z / UseCellSize));
+	const int32 OriginX = IntCastChecked<int32>(FMath::TruncToInt(Origin.X / UseCellSize));
+	const int32 OriginY = IntCastChecked<int32>(FMath::TruncToInt(Origin.Y / UseCellSize));
+	const int32 OriginZ = IntCastChecked<int32>(FMath::TruncToInt(Origin.Z / UseCellSize));
+
+	OriginWorldCoord = FIntVector(OriginX, OriginY, OriginZ);
 }
 
 FNavLocalGridData::FNavLocalGridData(const FVector& Center, const FVector2D& Extent2D) : GridId(0)
@@ -24,7 +28,11 @@ FNavLocalGridData::FNavLocalGridData(const FVector& Center, const FVector2D& Ext
 	const FBox UseBounds = FBox::BuildAABB(CenterOnGrid, FVector(Extent2D.X, Extent2D.Y, 0.0f));
 	Init(UseCellSize, UseBounds);
 
-	OriginWorldCoord = FIntVector(FMath::TruncToInt(Origin.X / UseCellSize), FMath::TruncToInt(Origin.Y / UseCellSize), FMath::TruncToInt(Origin.Z / UseCellSize));
+	const int32 OriginX = IntCastChecked<int32>(FMath::TruncToInt(Origin.X / UseCellSize));
+	const int32 OriginY = IntCastChecked<int32>(FMath::TruncToInt(Origin.Y / UseCellSize));
+	const int32 OriginZ = IntCastChecked<int32>(FMath::TruncToInt(Origin.Z / UseCellSize));
+
+	OriginWorldCoord = FIntVector(OriginX, OriginY, OriginZ);
 }
 
 FNavLocalGridData::FNavLocalGridData(const TArray<FNavLocalGridData>& SourceGrids) : GridId(0)
@@ -42,7 +50,10 @@ FNavLocalGridData::FNavLocalGridData(const TArray<FNavLocalGridData>& SourceGrid
 
 	const float UseCellSize = UNavLocalGridManager::GetCellSize();
 	Init(UseCellSize, CombinedBounds);
-	OriginWorldCoord = FIntVector(FMath::TruncToInt(Origin.X / UseCellSize), FMath::TruncToInt(Origin.Y / UseCellSize), FMath::TruncToInt(Origin.Z / UseCellSize));
+
+	const int32 OriginX = IntCastChecked<int32>(FMath::TruncToInt(Origin.X / UseCellSize));
+	const int32 OriginY = IntCastChecked<int32>(FMath::TruncToInt(Origin.Y / UseCellSize));
+	const int32 OriginZ = IntCastChecked<int32>(FMath::TruncToInt(Origin.Z / UseCellSize));
 
 	for (int32 Idx = 0; Idx < Cells.Num(); Idx++)
 	{
@@ -135,8 +146,8 @@ void FNavLocalGridData::MarkCapsuleObstacle(const FVector& Center, float Radius,
 		for (int32 IdxY = 0; IdxY < (int32)GridSize.Height; IdxY++)
 		{
 			const FVector GridCellLocation = GetWorldCellCenter(IdxX, IdxY);
-			const float Dist2DSq = (GridCellLocation - Center).SizeSquared2D();
-			const float DistZ = FMath::Abs(GridCellLocation.Z - Center.Z);
+			const FVector::FReal Dist2DSq = (GridCellLocation - Center).SizeSquared2D();
+			const FVector::FReal DistZ = FMath::Abs(GridCellLocation.Z - Center.Z);
 
 			const bool bIsInside = (Dist2DSq < RadiusSq) && (DistZ < HalfHeight);
 			if (bIsInside)
@@ -150,7 +161,7 @@ void FNavLocalGridData::MarkCapsuleObstacle(const FVector& Center, float Radius,
 
 void FNavLocalGridData::SetHeight(float ExtentZ)
 {
-	const float CenterZ = (WorldBounds.Max.Z + WorldBounds.Min.Z) * 0.5f;
+	const FVector::FReal CenterZ = (WorldBounds.Max.Z + WorldBounds.Min.Z) * 0.5;
 	WorldBounds.Min.Z = CenterZ - ExtentZ;
 	WorldBounds.Max.Z = CenterZ + ExtentZ;
 }
@@ -385,7 +396,7 @@ void FNavLocalGridData::ProjectCells(const ANavigationData& NavData)
 	const bool bAssignBoundsHeight = FMath::IsNearlyZero(WorldBounds.GetExtent().Z);
 	if (bAssignBoundsHeight)
 	{
-		SetHeight(NavData.GetDefaultQueryExtent().Z);
+		SetHeight(FloatCastChecked<float>(NavData.GetDefaultQueryExtent().Z, UE::LWC::DefaultFloatPrecision));
 	}
 
 	TArray<FNavigationProjectionWork> Workload;
@@ -401,10 +412,10 @@ void FNavLocalGridData::ProjectCells(const ANavigationData& NavData)
 	}
 
 	const FVector ProjectionExtent(GridCellSize * 0.25f, GridCellSize * 0.25f, (WorldBounds.Max.Z - WorldBounds.Min.Z) * 0.5f);
-	const float Extent2DSq = FMath::Square(ProjectionExtent.X);
+	const FVector::FReal Extent2DSq = FMath::Square(ProjectionExtent.X);
 	NavData.BatchProjectPoints(Workload, ProjectionExtent);
 	
-	const float BoundsMidZ = (WorldBounds.Max.Z + WorldBounds.Min.Z) * 0.5f;
+	const FVector::FReal BoundsMidZ = (WorldBounds.Max.Z + WorldBounds.Min.Z) * 0.5;
 	CellZ.Init(BoundsMidZ, Cells.Num());
 
 	int32 WorkloadIdx = 0;
diff --git a/Engine/Source/Runtime/AIModule/Private/Navigation/NavLocalGridManager.cpp b/Engine/Source/Runtime/AIModule/Private/Navigation/NavLocalGridManager.cpp
--- a/Engine/Source/Runtime/AIModule/Private/Navigation/NavLocalGridManager.cpp
+++ b/Engine/Source/Runtime/AIModule/Private/Navigation/NavLocalGridManager.cpp
@@ -17,7 +17,7 @@ UNavLocalGridManager::UNavLocalGridManager(const FObjectInitializer& ObjectIniti
 
 int32 UNavLocalGridManager::AddGridData(const FNavLocalGridData& GridData, bool bUpdate)
 {
-	const float GameTime = HasSourceGridLimit() ? GetWorld()->GetTimeSeconds() : 0;
+	const double GameTime = HasSourceGridLimit() ? GetWorld()->GetTimeSeconds() : 0.;
 	const int32 NewGridIdx = SourceGrids.Add(GridData);
 
 	SourceGrids[NewGridIdx].SetGridId(NextGridId);
@@ -217,7 +217,7 @@ void UNavLocalGridManager::UpdateAccessTime(int32 CombinedGridIdx)
 {
 	if (CombinedGrids.IsValidIndex(CombinedGridIdx))
 	{
-		const float GameTime = GetWorld()->GetTimeSeconds();
+		const double GameTime = GetWorld()->GetTimeSeconds();
 
 		for (int32 Idx = 0; Idx < CombinedGrids[CombinedGridIdx].SourceIds.Num(); Idx++)
 		{
@@ -250,7 +250,7 @@ bool UNavLocalGridManager::UpdateSourceGrids()
 
 	while (SourceGrids.Num() > MaxActiveSourceGrids)
 	{
-		float BestScore = FLT_MAX;
+		double BestScore = DBL_MAX;
 		int32 BestIdx = 0;
 
 		for (int32 Idx = 0; Idx < SourceGrids.Num(); Idx++)
@@ -343,7 +343,7 @@ int32 UNavLocalGridManager::AddLocalNavigationGridForPoints(UObject* WorldContex
 	if (GridManager)
 	{
 		const FBox Bounds(Locations);
-		const float BoundsSize2D = FMath::Max(Bounds.Max.X - Bounds.Min.X, Bounds.Max.Y - Bounds.Min.Y);
+		const float BoundsSize2D = FloatCastChecked<float>(FMath::Max(Bounds.Max.X - Bounds.Min.X, Bounds.Max.Y - Bounds.Min.Y), UE::LWC::DefaultFloatPrecision);
 
 		FNavLocalGridData GridData(Bounds.GetCenter(), (UNavLocalGridManager::GridCellSize * Radius2D) + BoundsSize2D);
 		GridData.SetHeight(Height);
diff --git a/Engine/Source/Runtime/AIModule/Private/Navigation/PathFollowingComponent.cpp b/Engine/Source/Runtime/AIModule/Private/Navigation/PathFollowingComponent.cpp
--- a/Engine/Source/Runtime/AIModule/Private/Navigation/PathFollowingComponent.cpp
+++ b/Engine/Source/Runtime/AIModule/Private/Navigation/PathFollowingComponent.cpp
@@ -207,7 +207,7 @@ void LogBlockHelper(AActor* LogOwner, UNavMovementComponent* MoveComp, float Rad
 	{
 		const FVector AgentLocation = MoveComp->GetActorFeetLocation();
 		const FVector ToTarget = (SegmentEnd - AgentLocation);
-		const float SegmentDot = FVector::DotProduct(ToTarget.GetSafeNormal(), (SegmentEnd - SegmentStart).GetSafeNormal());
+		const FVector::FReal SegmentDot = FVector::DotProduct(ToTarget.GetSafeNormal(), (SegmentEnd - SegmentStart).GetSafeNormal());
 		UE_VLOG(LogOwner, LogPathFollowing, Verbose, TEXT("[agent to segment end] dot [segment dir]: %f"), SegmentDot);
 		
 		float AgentRadius = 0.0f;
@@ -215,10 +215,10 @@ void LogBlockHelper(AActor* LogOwner, UNavMovementComponent* MoveComp, float Rad
 		AActor* MovingAgent = MoveComp->GetOwner();
 		MovingAgent->GetSimpleCollisionCylinder(AgentRadius, AgentHalfHeight);
 
-		const float Dist2D = ToTarget.Size2D();
+		const FVector::FReal Dist2D = ToTarget.Size2D();
 		UE_VLOG(LogOwner, LogPathFollowing, Verbose, TEXT("dist 2d: %f (agent radius: %f [%f])"), Dist2D, AgentRadius, AgentRadius * (1 + RadiusPct));
 
-		const float ZDiff = FMath::Abs(ToTarget.Z);
+		const FVector::FReal ZDiff = FMath::Abs(ToTarget.Z);
 		UE_VLOG(LogOwner, LogPathFollowing, Verbose, TEXT("Z diff: %f (agent halfZ: %f [%f])"), ZDiff, AgentHalfHeight, AgentHalfHeight * (1 + HeightPct));
 	}
 #endif // ENABLE_VISUAL_LOG
@@ -530,9 +530,9 @@ bool UPathFollowingComponent::ShouldCheckPathOnResume() const
 		MovementComp->GetOwner()->GetSimpleCollisionCylinder(AgentRadius, AgentHalfHeight);
 
 		const FVector CurrentLocation = MovementComp->GetActorFeetLocation();
-		const float DeltaMove2DSq = (CurrentLocation - LocationWhenPaused).SizeSquared2D();
-		const float DeltaZ = FMath::Abs(CurrentLocation.Z - LocationWhenPaused.Z);
-		if (DeltaMove2DSq < FMath::Square(AgentRadius) && DeltaZ < (AgentHalfHeight * 0.5f))
+		const FVector::FReal DeltaMove2DSq = (CurrentLocation - LocationWhenPaused).SizeSquared2D();
+		const FVector::FReal DeltaZ = FMath::Abs(CurrentLocation.Z - LocationWhenPaused.Z);
+		if (DeltaMove2DSq < FMath::Square(AgentRadius) && DeltaZ < (AgentHalfHeight * 0.5))
 		{
 			bCheckPath = false;
 		}
@@ -794,8 +794,8 @@ int32 UPathFollowingComponent::DetermineStartingPathPoint(const FNavigationPath*
 				const FVector PathPt1 = *ConsideredPath->GetPathPointLocation(1);
 				// making this test in 2d to avoid situation where agent's Z location not being in "navmesh plane"
 				// would influence the result
-				const float SqDistToFirstPoint = (CurrentLocation - PathPt0).SizeSquared2D();
-				const float SqDistToSecondPoint = (CurrentLocation - PathPt1).SizeSquared2D();
+				const FVector::FReal SqDistToFirstPoint = (CurrentLocation - PathPt0).SizeSquared2D();
+				const FVector::FReal SqDistToSecondPoint = (CurrentLocation - PathPt1).SizeSquared2D();
 
 				PickedPathPoint = FMath::IsNearlyEqual(SqDistToFirstPoint, SqDistToSecondPoint) ?
 					((FMath::Abs(CurrentLocation.Z - PathPt0.Z) < FMath::Abs(CurrentLocation.Z - PathPt1.Z)) ? 0 : 1) :
@@ -1031,13 +1031,13 @@ void UPathFollowingComponent::FollowPathSegment(float DeltaTime)
 		if (bStopMovementOnFinish && (MoveSegmentStartIndex >= DecelerationSegmentIndex))
 		{
 			const FVector PathEnd = Path->GetEndLocation();
-			const float DistToEndSq = FVector::DistSquared(CurrentLocation, PathEnd);
+			const FVector::FReal DistToEndSq = FVector::DistSquared(CurrentLocation, PathEnd);
 			const bool bShouldDecelerate = DistToEndSq < FMath::Square(CachedBrakingDistance);
 			if (bShouldDecelerate)
 			{
 				bIsDecelerating = true;
 
-				const float SpeedPct = FMath::Clamp(FMath::Sqrt(DistToEndSq) / CachedBrakingDistance, 0.0f, 1.0f);
+				const FVector::FReal  SpeedPct = FMath::Clamp(FMath::Sqrt(DistToEndSq) / CachedBrakingDistance, 0., 1.);
 				CurrentMoveInput *= SpeedPct;
 			}
 		}
@@ -1164,8 +1164,8 @@ bool UPathFollowingComponent::HasReachedCurrentTarget(const FVector& CurrentLoca
 
 	// check if moved too far
 	const FVector ToTarget = (CurrentTarget - MovementComp->GetActorFeetLocation());
-	const float SegmentDot = FVector::DotProduct(ToTarget, CurrentDirection);
-	if (SegmentDot < 0.0)
+	const FVector::FReal SegmentDot = FVector::DotProduct(ToTarget, CurrentDirection);
+	if (SegmentDot < 0.)
 	{
 		return true;
 	}
@@ -1194,15 +1194,15 @@ bool UPathFollowingComponent::HasReachedInternal(const FVector& GoalLocation, fl
 	// check if they overlap (with added AcceptanceRadius)
 	const FVector ToGoal = GoalLocation - AgentLocation;
 
-	const float Dist2DSq = ToGoal.SizeSquared2D();
-	const float UseRadius = RadiusThreshold + GoalRadius + (AgentRadius * AgentRadiusMultiplier);
+	const FVector::FReal Dist2DSq = ToGoal.SizeSquared2D();
+	const FVector::FReal UseRadius = RadiusThreshold + GoalRadius + (AgentRadius * AgentRadiusMultiplier);
 	if (Dist2DSq > FMath::Square(UseRadius))
 	{
 		return false;
 	}
 
-	const float ZDiff = FMath::Abs(ToGoal.Z);
-	const float UseHeight = GoalHalfHeight + (AgentHalfHeight * MinAgentHalfHeightPct);
+	const FVector::FReal ZDiff = FMath::Abs(ToGoal.Z);
+	const FVector::FReal UseHeight = GoalHalfHeight + (AgentHalfHeight * MinAgentHalfHeightPct);
 	if (ZDiff > UseHeight)
 	{
 		return false;
@@ -1221,7 +1221,7 @@ int32 UPathFollowingComponent::FindPreciseAcceptanceRadiusTestsStartNodeIndex(co
 
 	if (PathPoints.Num() > 1)
 	{
-		float DistanceSum = 0.f;
+		FVector::FReal DistanceSum = 0.f;
 		int32 NodeIndex = PathPoints.Num() - 1;
 		FVector PrevLocation = PathPoints[NodeIndex].Location;
 		--NodeIndex;
@@ -1249,12 +1249,12 @@ float UPathFollowingComponent::GetFinalAcceptanceRadius(const FNavigationPath& P
 	{
 		ensure(FAISystem::IsValidLocation(OriginalGoalLocation));
 
-		const float PathEndToGoalDistance = FVector::Dist(PathEndOverride ? *PathEndOverride : PathInstance.GetEndLocation(), OriginalGoalLocation);
-		const float Remaining = AcceptanceRadius - PathEndToGoalDistance;
+		const FVector::FReal PathEndToGoalDistance = FVector::Dist(PathEndOverride ? *PathEndOverride : PathInstance.GetEndLocation(), OriginalGoalLocation);
+		const FVector::FReal Remaining = AcceptanceRadius - PathEndToGoalDistance;
 		
 		// if goal is more than AcceptanceRadius from the end of the path
 		// we just need to do regular path following, with default radius value
-		return FMath::Max(Remaining, MyDefaultAcceptanceRadius);
+		return FloatCastChecked<float>(FMath::Max(Remaining, (FVector::FReal)MyDefaultAcceptanceRadius), UE::LWC::DefaultFloatPrecision);
 	}
 
 	return AcceptanceRadius;
@@ -1424,7 +1424,7 @@ bool UPathFollowingComponent::IsBlocked() const
 
 	if (LocationSamples.Num() == BlockDetectionSampleCount && BlockDetectionSampleCount > 0)
 	{
-		const float BlockDetectionDistanceSq = FMath::Square(BlockDetectionDistance);
+		const FVector::FReal BlockDetectionDistanceSq = FMath::Square(BlockDetectionDistance);
 
 		FVector Center = FVector::ZeroVector;
 		for (int32 SampleIndex = 0; SampleIndex < LocationSamples.Num(); SampleIndex++)
@@ -1437,7 +1437,7 @@ bool UPathFollowingComponent::IsBlocked() const
 
 		for (int32 SampleIndex = 0; SampleIndex < LocationSamples.Num(); SampleIndex++)
 		{
-			const float TestDistanceSq = FVector::DistSquared(*LocationSamples[SampleIndex], Center);
+			const FVector::FReal TestDistanceSq = FVector::DistSquared(*LocationSamples[SampleIndex], Center);
 			if (TestDistanceSq > BlockDetectionDistanceSq)
 			{
 				bBlocked = false;
@@ -1527,11 +1527,11 @@ void UPathFollowingComponent::UpdateDecelerationData()
 		DecelerationSegmentIndex = 0;
 
 		const TArray<FNavPathPoint>& PathPoints = Path->GetPathPoints();
-		float PathLengthFromEnd = 0.0f;
+		FVector::FReal PathLengthFromEnd = 0.;
 
 		for (int32 Idx = PathPoints.Num() - 1; Idx > 0; Idx--)
 		{
-			const float PathSegmentLength = FVector::Dist(PathPoints[Idx].Location, PathPoints[Idx - 1].Location);
+			const FVector::FReal PathSegmentLength = FVector::Dist(PathPoints[Idx].Location, PathPoints[Idx - 1].Location);
 			PathLengthFromEnd += PathSegmentLength;
 
 			if (PathLengthFromEnd > CachedBrakingDistance)
@@ -1568,9 +1568,9 @@ void UPathFollowingComponent::SetPreciseReachThreshold(float AgentRadiusMultipli
 	MinAgentHalfHeightPct = AgentHalfHeightMultiplier;
 }
 
-float UPathFollowingComponent::GetRemainingPathCost() const
+FVector::FReal UPathFollowingComponent::GetRemainingPathCost() const
 {
-	float Cost = 0.f;
+	FVector::FReal Cost = 0.;
 
 	if (Path.IsValid() && Path->IsValid() && Status == EPathFollowingStatus::Moving)
 	{
diff --git a/Engine/Source/Runtime/AIModule/Private/Perception/AIPerceptionSystem.cpp b/Engine/Source/Runtime/AIModule/Private/Perception/AIPerceptionSystem.cpp
--- a/Engine/Source/Runtime/AIModule/Private/Perception/AIPerceptionSystem.cpp
+++ b/Engine/Source/Runtime/AIModule/Private/Perception/AIPerceptionSystem.cpp
@@ -36,8 +36,8 @@ UAIPerceptionSystem::UAIPerceptionSystem(const FObjectInitializer& ObjectInitial
 	: Super(ObjectInitializer)
 	, PerceptionAgingRate(0.3f)
 	, bHandlePawnNotification(false)
-	, NextStimuliAgingTick(0.f)
-	, CurrentTime(0.f)
+	, NextStimuliAgingTick(0.)
+	, CurrentTime(0.)
 {
 PRAGMA_DISABLE_DEPRECATION_WARNINGS
 #if WITH_EDITORONLY_DATA
@@ -182,7 +182,9 @@ void UAIPerceptionSystem::Tick(float DeltaSeconds)
 		bool bSomeListenersNeedUpdateDueToStimuliAging = false;
 		if (NextStimuliAgingTick <= CurrentTime)
 		{
-			bSomeListenersNeedUpdateDueToStimuliAging = AgeStimuli(PerceptionAgingRate + (CurrentTime - NextStimuliAgingTick));
+			constexpr double Precision = 1./64.;
+			const float AgingDt = FloatCastChecked<float>(CurrentTime - NextStimuliAgingTick, Precision);
+			bSomeListenersNeedUpdateDueToStimuliAging = AgeStimuli(PerceptionAgingRate + AgingDt);
 			NextStimuliAgingTick = CurrentTime + PerceptionAgingRate;
 		}
 
@@ -601,7 +603,7 @@ void UAIPerceptionSystem::StartPlay()
 	}
 
 	UWorld* World = GetWorld();
-	NextStimuliAgingTick = World ? World->GetTimeSeconds() : 0.f;
+	NextStimuliAgingTick = World ? World->GetTimeSeconds() : 0.;
 }
 
 void UAIPerceptionSystem::RegisterAllPawnsAsSourcesForSense(FAISenseID SenseID)
diff --git a/Engine/Source/Runtime/AIModule/Private/Perception/AISense_Hearing.cpp b/Engine/Source/Runtime/AIModule/Private/Perception/AISense_Hearing.cpp
--- a/Engine/Source/Runtime/AIModule/Private/Perception/AISense_Hearing.cpp
+++ b/Engine/Source/Runtime/AIModule/Private/Perception/AISense_Hearing.cpp
@@ -143,7 +143,7 @@ float UAISense_Hearing::Update()
 		for (const FAINoiseEvent& Event : NoiseEvents)
 		{
 			const float ClampedLoudness = FMath::Max(0.f, Event.Loudness);
-			const float DistToSoundSquared = FVector::DistSquared(Event.NoiseLocation, Listener.CachedLocation);
+			const FVector::FReal DistToSoundSquared = FVector::DistSquared(Event.NoiseLocation, Listener.CachedLocation);
 			
 			// Limit by loudness modified squared range (this is the old behavior)
 			if (DistToSoundSquared > PropDigest.HearingRangeSq * FMath::Square(ClampedLoudness))
@@ -160,8 +160,8 @@ float UAISense_Hearing::Update()
 			{
 				continue;
 			}
-			// calculate delay and fake it with Age
-			const float Delay = FMath::Sqrt(DistToSoundSquared * SpeedOfSoundSqScalar);
+			// calculate delay and fake it with Age, Delay should be pretty small so a static_cast is safe enough here.
+			const float Delay = FloatCastChecked<float>(FMath::Sqrt(DistToSoundSquared * SpeedOfSoundSqScalar), UE::LWC::DefaultFloatPrecision);
 			// pass over to listener to process 			
 			PerseptionSys->RegisterDelayedStimulus(Listener.GetListenerID(), Delay, Event.Instigator
 				, FAIStimulus(*this, ClampedLoudness, Event.NoiseLocation, Listener.CachedLocation, FAIStimulus::SensingSucceeded, Event.Tag) );
diff --git a/Engine/Source/Runtime/AIModule/Private/Perception/AISense_Sight.cpp b/Engine/Source/Runtime/AIModule/Private/Perception/AISense_Sight.cpp
--- a/Engine/Source/Runtime/AIModule/Private/Perception/AISense_Sight.cpp
+++ b/Engine/Source/Runtime/AIModule/Private/Perception/AISense_Sight.cpp
@@ -153,9 +153,9 @@ UAISense_Sight::UAISense_Sight(const FObjectInitializer& ObjectInitializer)
 
 FORCEINLINE_DEBUGGABLE float UAISense_Sight::CalcQueryImportance(const FPerceptionListener& Listener, const FVector& TargetLocation, const float SightRadiusSq) const
 {
-	const float DistanceSq = FVector::DistSquared(Listener.CachedLocation, TargetLocation);
+	const FVector::FReal DistanceSq = FVector::DistSquared(Listener.CachedLocation, TargetLocation);
 	return DistanceSq <= HighImportanceDistanceSquare ? MaxQueryImportance
-		: FMath::Clamp((SightLimitQueryImportance - MaxQueryImportance) / SightRadiusSq * DistanceSq + MaxQueryImportance, 0.f, MaxQueryImportance);
+		: static_cast<float>(FMath::Clamp((SightLimitQueryImportance - MaxQueryImportance) / SightRadiusSq * DistanceSq + MaxQueryImportance, 0.f, MaxQueryImportance));
 }
 
 void UAISense_Sight::PostInitProperties()
@@ -170,7 +170,7 @@ bool UAISense_Sight::ShouldAutomaticallySeeTarget(const FDigestedSightProperties
 
 	if ((PropDigest.AutoSuccessRangeSqFromLastSeenLocation != FAISystem::InvalidRange) && (SightQuery->LastSeenLocation != FAISystem::InvalidLocation))
 	{
-		const float DistanceToLastSeenLocationSq = FVector::DistSquared(TargetActor->GetActorLocation(), SightQuery->LastSeenLocation);
+		const FVector::FReal DistanceToLastSeenLocationSq = FVector::DistSquared(TargetActor->GetActorLocation(), SightQuery->LastSeenLocation);
 		return (DistanceToLastSeenLocationSq <= PropDigest.AutoSuccessRangeSqFromLastSeenLocation);
 	}
 
diff --git a/Engine/Source/Runtime/AIModule/Private/Perception/PawnSensingComponent.cpp b/Engine/Source/Runtime/AIModule/Private/Perception/PawnSensingComponent.cpp
--- a/Engine/Source/Runtime/AIModule/Private/Perception/PawnSensingComponent.cpp
+++ b/Engine/Source/Runtime/AIModule/Private/Perception/PawnSensingComponent.cpp
@@ -331,7 +331,7 @@ bool UPawnSensingComponent::CouldSeePawn(const APawn *Other, bool bMaySkipChecks
 	FVector const SelfToOther = OtherLoc - SensorLoc;
 
 	// check max sight distance
-	float const SelfToOtherDistSquared = SelfToOther.SizeSquared();
+	FVector::FReal const SelfToOtherDistSquared = SelfToOther.SizeSquared();
 	if (SelfToOtherDistSquared > FMath::Square(SightRadius))
 	{
 		return false;
@@ -422,7 +422,7 @@ bool UPawnSensingComponent::CanHear(const FVector& NoiseLoc, float Loudness, boo
 
 	FVector const HearingLocation = GetSensorLocation();
 
-	float const LoudnessAdjustedDistSq = (HearingLocation - NoiseLoc).SizeSquared()/(Loudness*Loudness);
+	FVector::FReal const LoudnessAdjustedDistSq = (HearingLocation - NoiseLoc).SizeSquared()/(Loudness*Loudness);
 	if (LoudnessAdjustedDistSq <= FMath::Square(HearingThreshold))
 	{
 		// Hear even occluded sounds within HearingThreshold
diff --git a/Engine/Source/Runtime/AIModule/Private/VisualLoggerExtension.cpp b/Engine/Source/Runtime/AIModule/Private/VisualLoggerExtension.cpp
--- a/Engine/Source/Runtime/AIModule/Private/VisualLoggerExtension.cpp
+++ b/Engine/Source/Runtime/AIModule/Private/VisualLoggerExtension.cpp
@@ -44,7 +44,7 @@ void FVisualLoggerExtension::OnItemsSelectionChanged(IVisualLoggerEditorInterfac
 
 void FVisualLoggerExtension::OnLogLineSelectionChanged(IVisualLoggerEditorInterface* EdInterface, TSharedPtr<struct FLogEntryItem> SelectedItem, int64 UserData)
 {
-	SelectedEQSId = SelectedItem.IsValid() ? UserData : INDEX_NONE;
+	SelectedEQSId = SelectedItem.IsValid() ? IntCastChecked<int32>(UserData) : INDEX_NONE;
 	EdInterface->GetHelperActor()->MarkComponentsRenderStateDirty();
 	DrawData(EdInterface, NULL); //we have to refresh rendering components
 }
@@ -133,16 +133,15 @@ void FVisualLoggerExtension::DrawData(UWorld* InWorld, class UEQSRenderingCompon
 			FVector FireDir = Canvas->SceneView->GetViewDirection();
 			FVector CamLocation = Canvas->SceneView->ViewMatrices.GetViewOrigin();
 
-			float bestAim = 0;
+			FVector::FReal bestAim = 0.;
 			for (int32 Index = 0; Index < DebugData.RenderDebugHelpers.Num(); ++Index)
 			{
 				auto& CurrentItem = DebugData.RenderDebugHelpers[Index];
 
 				const FVector AimDir = CurrentItem.Location - CamLocation;
-				float FireDist = AimDir.SizeSquared();
+				const FVector::FReal FireDist = AimDir.Size();
 
-				FireDist = FMath::Sqrt(FireDist);
-				float newAim = FireDir | AimDir;
+				FVector::FReal newAim = FireDir | AimDir;
 				newAim = newAim / FireDist;
 				if (newAim > bestAim)
 				{
diff --git a/Engine/Source/Runtime/AIModule/Public/AIHelpers.h b/Engine/Source/Runtime/AIModule/Public/AIHelpers.h
--- a/Engine/Source/Runtime/AIModule/Public/AIHelpers.h
+++ b/Engine/Source/Runtime/AIModule/Public/AIHelpers.h
@@ -58,18 +58,21 @@ namespace UE::AI
 	 * This method will extract the yaw radian from the specified vector (The vector does not need to be normalized)
 	 * if it is not possible to compute yaw, the function will return an invalid value	 *
 	 */
+	UE_DEPRECATED(5.2, "Deprecated as returning float instead of FVector::FReal. Use FVector::ToOrientationRotator()::Yaw etc.")
 	extern AIMODULE_API TOptional<float> GetYawFromVector(const FVector& Vector);
 
 	/**
 	 * This method will extract the yaw radian from the specified rotator
 	 * if it is not possible to compute yaw, the function will not return an invalid value
 	 */
+	UE_DEPRECATED(5.2, "Deprecated as returning float instead of FVector::FReal. Use FRotator::Yaw etc.")
 	extern AIMODULE_API TOptional<float> GetYawFromRotator(const FRotator& Rotator);
 
 	/**
 	 * This method will extract the yaw radian from the specified quaternion
 	 * if it is not possible to compute yaw, the function will not return an invalid value
 	 */
+	UE_DEPRECATED(5.2, "Deprecated as returning float instead of FVector::FReal. Use FQuat::Rotator()::Yaw etc.")
 	extern AIMODULE_API TOptional<float> GetYawFromQuaternion(const FQuat& Quaternion);
 
 	/**
diff --git a/Engine/Source/Runtime/AIModule/Public/SimpleCellGrid.h b/Engine/Source/Runtime/AIModule/Public/SimpleCellGrid.h
--- a/Engine/Source/Runtime/AIModule/Public/SimpleCellGrid.h
+++ b/Engine/Source/Runtime/AIModule/Public/SimpleCellGrid.h
@@ -50,7 +50,7 @@ public:
 		GridCellSize = InCellSize;
 
 		const FVector RealBoundsSize = Bounds.GetSize();
-		GridSize = FGridSize2D(FMath::CeilToInt(RealBoundsSize.X / InCellSize), FMath::CeilToInt(RealBoundsSize.Y / InCellSize));
+		GridSize = FGridSize2D(IntCastChecked<uint32>(FMath::CeilToInt(RealBoundsSize.X / InCellSize)), IntCastChecked<uint32>(FMath::CeilToInt(RealBoundsSize.Y / InCellSize)));
 		BoundsSize = FVector(GridSize.Width * InCellSize, GridSize.Height * InCellSize, RealBoundsSize.Z);
 		Origin = FVector(Bounds.Min.X, Bounds.Min.Y, (Bounds.Min.Z + Bounds.Max.Z) * 0.5f);
 		UpdateWorldBounds();
@@ -87,15 +87,15 @@ public:
 
 	FORCEINLINE uint32 GetAllocatedSize() const
 	{
-		return Cells.GetAllocatedSize();
+		return IntCastChecked<uint32>(Cells.GetAllocatedSize());
 	}
 
 	/** Convert world location to (X,Y) coords on grid, result can be outside grid */
 	FORCEINLINE FIntVector GetCellCoordsUnsafe(const FVector& WorldLocation) const
 	{
 		return FIntVector(
-			FMath::TruncToInt((WorldLocation.X - Origin.X) / GridCellSize),
-			FMath::TruncToInt((WorldLocation.Y - Origin.Y) / GridCellSize),
+			IntCastChecked<int32>(FMath::TruncToInt((WorldLocation.X - Origin.X) / GridCellSize)),
+			IntCastChecked<int32>(FMath::TruncToInt((WorldLocation.Y - Origin.Y) / GridCellSize)),
 			0);
 	}
 
@@ -241,7 +241,7 @@ public:
 		}
 		else
 		{
-			int32 DataBytesCount = GetAllocatedSize();
+			uint32 DataBytesCount = GetAllocatedSize();
 			Ar << DataBytesCount;
 
 			if (DataBytesCount > 0)
diff --git a/Engine/Source/Runtime/Core/Public/Containers/EnumAsByte.h b/Engine/Source/Runtime/Core/Public/Containers/EnumAsByte.h
--- a/Engine/Source/Runtime/Core/Public/Containers/EnumAsByte.h
+++ b/Engine/Source/Runtime/Core/Public/Containers/EnumAsByte.h
@@ -94,6 +94,16 @@ public:
 		return (TEnum)Value;
 	}
 
+	/**
+	 * Gets the integer enumeration value.
+	 *
+	 * @return The enumeration value.
+	 */
+	uint8 GetIntValue() const
+	{
+		return Value;
+	}
+
 private:
 
 	/** Holds the value as a byte. **/
diff --git a/Engine/Source/Runtime/Core/Public/Math/Vector.h b/Engine/Source/Runtime/Core/Public/Math/Vector.h
--- a/Engine/Source/Runtime/Core/Public/Math/Vector.h
+++ b/Engine/Source/Runtime/Core/Public/Math/Vector.h
@@ -2657,16 +2657,18 @@ inline TVector<T> TVector2<T>::SphericalToUnitCartesian() const
 	
 namespace LWC
 {
+constexpr FVector::FReal DefaultFloatPrecision = 1./16.;
+
 // Validated narrowing cast for world positions. FVector -> FVector3f
 FORCEINLINE FVector3f NarrowWorldPositionChecked(const FVector& WorldPosition)
 {
 	FVector3f Narrowed;
-	constexpr FVector::FReal Precision = 1/16.0;
-	Narrowed.X = FloatCastChecked<float>(WorldPosition.X, Precision);
-	Narrowed.Y = FloatCastChecked<float>(WorldPosition.Y, Precision);
-	Narrowed.Z = FloatCastChecked<float>(WorldPosition.Z, Precision);
+	Narrowed.X = FloatCastChecked<float>(WorldPosition.X, DefaultFloatPrecision);
+	Narrowed.Y = FloatCastChecked<float>(WorldPosition.Y, DefaultFloatPrecision);
+	Narrowed.Z = FloatCastChecked<float>(WorldPosition.Z, DefaultFloatPrecision);
 	return Narrowed;
 }
+
 } // namespace UE::LWC
 
 } // namespace UE
diff --git a/Engine/Source/Runtime/GameplayTasks/Classes/GameplayTaskResource.h b/Engine/Source/Runtime/GameplayTasks/Classes/GameplayTaskResource.h
--- a/Engine/Source/Runtime/GameplayTasks/Classes/GameplayTaskResource.h
+++ b/Engine/Source/Runtime/GameplayTasks/Classes/GameplayTaskResource.h
@@ -34,7 +34,7 @@ public:
 
 	uint8 GetResourceID() const
 	{
-		return bManuallySetID || (ManualResourceID != INDEX_NONE) ? ManualResourceID : AutoResourceID;
+		return bManuallySetID || (ManualResourceID != INDEX_NONE) ? IntCastChecked<uint8>(ManualResourceID) : AutoResourceID;
 	}
 
 	template <class T>
diff --git a/Engine/Source/Runtime/Navmesh/Public/DetourCrowd/DetourCrowd.h b/Engine/Source/Runtime/Navmesh/Public/DetourCrowd/DetourCrowd.h
--- a/Engine/Source/Runtime/Navmesh/Public/DetourCrowd/DetourCrowd.h
+++ b/Engine/Source/Runtime/Navmesh/Public/DetourCrowd/DetourCrowd.h
@@ -528,7 +528,7 @@ public:
 
 	inline int getNumActiveAgents() const { return m_numActiveAgents; }
 
-	inline int getAgentIndex(const dtCrowdAgent* agent) const  { return agent - m_agents; }
+	inline int getAgentIndex(const dtCrowdAgent* agent) const  { return (int)(agent - m_agents); }
 
 	/// Gets all agent animations
 	const dtCrowdAgentAnimation* getAgentAnims() const { return m_agentAnims;  }
