Fix for some missing code in CL 22777361 which allows FBlackboard FKey to be uint16.

diff --git a/Engine/Source/Runtime/AIModule/Classes/BehaviorTree/BlackboardComponent.h b/Engine/Source/Runtime/AIModule/Classes/BehaviorTree/BlackboardComponent.h
--- a/Engine/Source/Runtime/AIModule/Classes/BehaviorTree/BlackboardComponent.h
+++ b/Engine/Source/Runtime/AIModule/Classes/BehaviorTree/BlackboardComponent.h
@@ -273,13 +273,13 @@ protected:
 	mutable int32 ObserversToRemoveCount = 0;
 
 	/** observers registered for blackboard keys */
-	mutable TMultiMap<uint8, FOnBlackboardChangeNotificationInfo> Observers;
+	mutable TMultiMap<FBlackboard::FKey, FOnBlackboardChangeNotificationInfo> Observers;
 	
 	/** observers registered from owner objects */
 	mutable TMultiMap<const UObject*, FDelegateHandle> ObserverHandles;
 
 	/** queued key change notification, will be processed on ResumeUpdates call */
-	mutable TArray<uint8> QueuedUpdates;
+	mutable TArray<FBlackboard::FKey> QueuedUpdates;
 
 	/** set when observation notifies are paused and shouldn't be passed to observers */
 	uint32 bPausedNotifies : 1;
diff --git a/Engine/Source/Runtime/AIModule/Private/BehaviorTree/BlackboardComponent.cpp b/Engine/Source/Runtime/AIModule/Private/BehaviorTree/BlackboardComponent.cpp
--- a/Engine/Source/Runtime/AIModule/Private/BehaviorTree/BlackboardComponent.cpp
+++ b/Engine/Source/Runtime/AIModule/Private/BehaviorTree/BlackboardComponent.cpp
@@ -409,7 +409,7 @@ void UBlackboardComponent::ResumeObserverNotifications(bool bSendQueuedObserverN
 
 void UBlackboardComponent::NotifyObservers(FBlackboard::FKey KeyID) const
 {
-	TMultiMap<uint8, FOnBlackboardChangeNotificationInfo>::TKeyIterator KeyIt(Observers, KeyID);
+	TMultiMap<FBlackboard::FKey, FOnBlackboardChangeNotificationInfo>::TKeyIterator KeyIt(Observers, KeyID);
 
 	// checking it here mostly to avoid storing this update in QueuedUpdates while
 	// at this point no one observes it, and there can be someone added before QueuedUpdates
@@ -544,14 +544,13 @@ FString UBlackboardComponent::GetDebugInfoString(EBlackboardDescription::Type Mo
 	{
 		DebugString += TEXT("Observed Keys:\n");
 
-		TArray<uint8> ObserversKeys;
+		TArray<FBlackboard::FKey> ObserversKeys;
 		if (Observers.Num() > 0)
 		{
 			Observers.GetKeys(ObserversKeys);
 
-			for (int32 KeyIndex = 0; KeyIndex < ObserversKeys.Num(); ++KeyIndex)
+			for (FBlackboard::FKey KeyID : ObserversKeys)
 			{
-				const FBlackboard::FKey KeyID = ObserversKeys[KeyIndex];
 				//@todo shouldn't be using a localized value?; GetKeyName() [10/11/2013 justin.sargent]
 				DebugString += FString::Printf(TEXT("  %s:\n"), *BlackboardAsset->GetKeyName(KeyID).ToString());
 			}
