Compile time error when a uint8 is used as a FBlackboard::Fkey
Fix BlackboardDecoratorDetails and BlackboardSelectorDetails
#tests Test FBlackboardDecoratorDetails in editor, rebuild projects

diff --git a/Engine/Source/Editor/BehaviorTreeEditor/Private/DetailCustomizations/BlackboardDecoratorDetails.cpp b/Engine/Source/Editor/BehaviorTreeEditor/Private/DetailCustomizations/BlackboardDecoratorDetails.cpp
--- a/Engine/Source/Editor/BehaviorTreeEditor/Private/DetailCustomizations/BlackboardDecoratorDetails.cpp
+++ b/Engine/Source/Editor/BehaviorTreeEditor/Private/DetailCustomizations/BlackboardDecoratorDetails.cpp
@@ -138,7 +138,7 @@ void FBlackboardDecoratorDetails::OnKeyIDChanged()
 		return;
 	}
 
-	uint8 KeyID;
+	int32 KeyID;
 	FPropertyAccess::Result Result = KeyIDProperty->GetValue(KeyID);
 	if (Result == FPropertyAccess::Success)
 	{
diff --git a/Engine/Source/Editor/BehaviorTreeEditor/Private/DetailCustomizations/BlackboardSelectorDetails.cpp b/Engine/Source/Editor/BehaviorTreeEditor/Private/DetailCustomizations/BlackboardSelectorDetails.cpp
--- a/Engine/Source/Editor/BehaviorTreeEditor/Private/DetailCustomizations/BlackboardSelectorDetails.cpp
+++ b/Engine/Source/Editor/BehaviorTreeEditor/Private/DetailCustomizations/BlackboardSelectorDetails.cpp
@@ -239,7 +239,7 @@ void FBlackboardSelectorDetails::InitKeyFromProperty()
 			else
 			{
 				MyKeyClassProperty->SetValue((UObject*)NULL);
-				MyKeyIDProperty->SetValue(FBlackboard::InvalidKey);
+				MyKeyIDProperty->SetValue((int32)FBlackboard::InvalidKey);
 				MyKeyNameProperty->SetValue(TEXT("None"));
 			}
 		}
@@ -275,11 +275,11 @@ void FBlackboardSelectorDetails::OnKeyComboChange(int32 Index)
 		UBlackboardData* BlackboardAsset = CachedBlackboardAsset.Get();
 		if (BlackboardAsset)
 		{
-			const uint8 KeyID = BlackboardAsset->GetKeyID(KeyValues[Index]);
+			const FBlackboard::FKey KeyID = BlackboardAsset->GetKeyID(KeyValues[Index]);
 			const UObject* KeyClass = BlackboardAsset->GetKeyType(KeyID);
 
 			MyKeyClassProperty->SetValue(KeyClass);
-			MyKeyIDProperty->SetValue(KeyID);
+			MyKeyIDProperty->SetValue((int32)KeyID);
 
 			MyKeyNameProperty->SetValue(KeyValues[Index]);
 		}
diff --git a/Engine/Source/Runtime/AIModule/Classes/BehaviorTree/BehaviorTreeTypes.h b/Engine/Source/Runtime/AIModule/Classes/BehaviorTree/BehaviorTreeTypes.h
--- a/Engine/Source/Runtime/AIModule/Classes/BehaviorTree/BehaviorTreeTypes.h
+++ b/Engine/Source/Runtime/AIModule/Classes/BehaviorTree/BehaviorTreeTypes.h
@@ -50,12 +50,35 @@ namespace FBlackboard
 #ifdef AI_BLACKBOARD_KEY_SIZE_8
 	// this is the legacy BB key size. Add AI_BLACKBOARD_KEY_SIZE_8 to your *.target.cs to enable it.
 	using FKey = uint8;
+	constexpr FKey InvalidKey = FKey(-1);
 #else
-	// this is the default BB key size
-	using FKey = uint16;
-#endif
+	//the default BB key size is now 16
+	struct FKey
+	{
+		constexpr FKey() = default;
+		constexpr FKey(int32 InKey) : Key(InKey) {}
+		constexpr FKey(uint16 InKey) : Key(InKey) {}
+		constexpr FKey(uint8 InKey) = delete;
+		constexpr operator int32() const { return Key; }
+		constexpr operator uint16() const { return Key; }
+		constexpr operator uint8() const = delete;
+		constexpr bool operator==(const FKey& Other) const {return Key == Other.Key;}
+		constexpr bool operator!=(const FKey& Other) const {return Key != Other.Key;}
+	private:
+		friend uint32 GetTypeHash(const FKey& Key);
+		uint16 Key = static_cast<uint16>(-1);
+	};
 
-	constexpr FKey InvalidKey = FKey(-1);
+	constexpr FKey InvalidKey = FKey();
+
+	inline uint32 GetTypeHash(const FKey& Key) { return ::GetTypeHash(Key.Key);}
+
+	template <>
+	struct TIsValidVariadicFunctionArg<FKey>
+	{
+		enum { Value = true };
+	};
+#endif
 }
 
 enum class EBlackboardNotificationResult : uint8
@@ -642,9 +665,9 @@ public:
 	void AddStringFilter(UObject* Owner, FName PropertyName);
 	void AddNameFilter(UObject* Owner, FName PropertyName);
 
-	FORCEINLINE bool IsNone() const { return bNoneIsAllowedValue && SelectedKeyID == FBlackboard::InvalidKey; }
-	FORCEINLINE bool IsSet() const { return SelectedKeyID != FBlackboard::InvalidKey; }
-	FORCEINLINE bool NeedsResolving() const { return SelectedKeyID == FBlackboard::InvalidKey && SelectedKeyName.IsNone() == false; }
+	FORCEINLINE bool IsNone() const { return bNoneIsAllowedValue && GetSelectedKeyID() == FBlackboard::InvalidKey; }
+	FORCEINLINE bool IsSet() const { return GetSelectedKeyID() != FBlackboard::InvalidKey; }
+	FORCEINLINE bool NeedsResolving() const { return GetSelectedKeyID() == FBlackboard::InvalidKey && SelectedKeyName.IsNone() == false; }
 	FORCEINLINE void InvalidateResolvedKey() { SelectedKeyID = FBlackboard::InvalidKey; }
 
 	friend FBlackboardDecoratorDetails;
