Fix build errors on Clang

diff --git a/Engine/Source/Runtime/AIModule/Classes/BehaviorTree/BehaviorTreeTypes.h b/Engine/Source/Runtime/AIModule/Classes/BehaviorTree/BehaviorTreeTypes.h
--- a/Engine/Source/Runtime/AIModule/Classes/BehaviorTree/BehaviorTreeTypes.h
+++ b/Engine/Source/Runtime/AIModule/Classes/BehaviorTree/BehaviorTreeTypes.h
@@ -72,15 +72,16 @@ namespace FBlackboard
 	constexpr FKey InvalidKey = FKey();
 
 	inline uint32 GetTypeHash(const FKey& Key) { return ::GetTypeHash(Key.Key);}
-
-	template <>
-	struct TIsValidVariadicFunctionArg<FKey>
-	{
-		enum { Value = true };
-	};
 #endif
 }
 
+#ifndef AI_BLACKBOARD_KEY_SIZE_8
+template <> struct TIsValidVariadicFunctionArg<FBlackboard::FKey>
+{
+	enum { Value = true };
+};
+#endif
+
 enum class EBlackboardNotificationResult : uint8
 {
 	RemoveObserver,
diff --git a/Engine/Source/Runtime/AIModule/Classes/BehaviorTree/BlackboardData.h b/Engine/Source/Runtime/AIModule/Classes/BehaviorTree/BlackboardData.h
--- a/Engine/Source/Runtime/AIModule/Classes/BehaviorTree/BlackboardData.h
+++ b/Engine/Source/Runtime/AIModule/Classes/BehaviorTree/BlackboardData.h
@@ -86,7 +86,7 @@ public:
 
 	FORCEINLINE FBlackboard::FKey GetFirstKeyID() const { return FirstKeyID; }
 
-	bool IsValidKey(FBlackboard::FKey KeyID) const { return KeyID != FBlackboard::InvalidKey && KeyID < GetNumKeys(); }
+	bool IsValidKey(FBlackboard::FKey KeyID) const { return KeyID != FBlackboard::InvalidKey && (int32)KeyID < GetNumKeys(); }
 
 	/** @return key data */
 	const FBlackboardEntry* GetKey(FBlackboard::FKey KeyID) const;
