Fix build errors on Clang

diff --git a/Engine/Source/Runtime/AIModule/Classes/BehaviorTree/BlackboardData.h b/Engine/Source/Runtime/AIModule/Classes/BehaviorTree/BlackboardData.h
--- a/Engine/Source/Runtime/AIModule/Classes/BehaviorTree/BlackboardData.h
+++ b/Engine/Source/Runtime/AIModule/Classes/BehaviorTree/BlackboardData.h
@@ -126,7 +126,7 @@ public:
 		}
 		else if (KeyID != FBlackboard::InvalidKey && Parent != NULL)
 		{
-			const FBlackboard::FKey KeyIndex = KeyID - FirstKeyID;
+			const FBlackboard::FKey KeyIndex = (int32)KeyID - (int32)FirstKeyID;
 			Keys.RemoveAt(KeyIndex);
 			MarkPackageDirty();
 			PropagateKeyChangesToDerivedBlackboardAssets();
diff --git a/Engine/Source/Runtime/AIModule/Private/BehaviorTree/BlackboardComponent.cpp b/Engine/Source/Runtime/AIModule/Private/BehaviorTree/BlackboardComponent.cpp
--- a/Engine/Source/Runtime/AIModule/Private/BehaviorTree/BlackboardComponent.cpp
+++ b/Engine/Source/Runtime/AIModule/Private/BehaviorTree/BlackboardComponent.cpp
@@ -155,7 +155,7 @@ bool UBlackboardComponent::InitializeBlackboard(UBlackboardData& NewAsset)
 					KeyType->PreInitialize(*this);
 
 					const uint16 KeyMemory = KeyType->GetValueSize() + (KeyType->HasInstance() ? sizeof(FBlackboardInstancedKeyMemory) : 0);
-					const int32 OffsetKey = KeyIndex + It->GetFirstKeyID();
+					const int32 OffsetKey = KeyIndex + (int32)It->GetFirstKeyID();
 					check(OffsetKey != FBlackboard::InvalidKey);
 
 					InitList.Add(FBlackboardInitializationData(IntCastChecked<FBlackboard::FKey>(OffsetKey), KeyMemory));
@@ -211,7 +211,7 @@ void UBlackboardComponent::DestroyValues()
 			UBlackboardKeyType* KeyType = It->Keys[KeyIndex].KeyType;
 			if (KeyType)
 			{
-				const int32 OffsetKey = KeyIndex + It->GetFirstKeyID();
+				const int32 OffsetKey = KeyIndex + (int32)It->GetFirstKeyID();
 				check(OffsetKey != FBlackboard::InvalidKey);
 
 				uint8* KeyMemory = GetKeyRawData(IntCastChecked<FBlackboard::FKey>(OffsetKey));
@@ -616,7 +616,7 @@ void UBlackboardComponent::DescribeSelfToVisLog(FVisualLogEntry* Snapshot) const
 			{
 				const FBlackboardEntry& Key = It->Keys[KeyIndex];
 
-				const int32 OffsetKey = It->GetFirstKeyID() + KeyIndex;
+				const int32 OffsetKey = (int32)It->GetFirstKeyID() + KeyIndex;
 				check(OffsetKey != FBlackboard::InvalidKey);
 
 				const uint8* ValueData = GetKeyRawData(IntCastChecked<FBlackboard::FKey>(OffsetKey));
diff --git a/Engine/Source/Runtime/AIModule/Private/BehaviorTree/BlackboardData.cpp b/Engine/Source/Runtime/AIModule/Private/BehaviorTree/BlackboardData.cpp
--- a/Engine/Source/Runtime/AIModule/Private/BehaviorTree/BlackboardData.cpp
+++ b/Engine/Source/Runtime/AIModule/Private/BehaviorTree/BlackboardData.cpp
@@ -83,9 +83,9 @@ const FBlackboardEntry* UBlackboardData::GetKey(FBlackboard::FKey KeyID) const
 {
 	if (KeyID != FBlackboard::InvalidKey)
 	{
-		if (KeyID >= FirstKeyID)
+		if ((int32)KeyID >= (int32)FirstKeyID)
 		{
-			return &Keys[KeyID - FirstKeyID];
+			return &Keys[(int32)KeyID - (int32)FirstKeyID];
 		}
 		else if (Parent)
 		{
@@ -98,7 +98,7 @@ const FBlackboardEntry* UBlackboardData::GetKey(FBlackboard::FKey KeyID) const
 
 int32 UBlackboardData::GetNumKeys() const
 {
-	return FirstKeyID + Keys.Num();
+	return (int32)FirstKeyID + Keys.Num();
 }
 
 FBlackboard::FKey UBlackboardData::InternalGetKeyID(const FName& KeyName, EKeyLookupMode LookupMode) const
@@ -107,7 +107,7 @@ FBlackboard::FKey UBlackboardData::InternalGetKeyID(const FName& KeyName, EKeyLo
 	{
 		if (Keys[KeyIndex].EntryName == KeyName)
 		{
-			const int32 OffsetKey = KeyIndex + FirstKeyID;
+			const int32 OffsetKey = KeyIndex + (int32)FirstKeyID;
 			check(OffsetKey != FBlackboard::InvalidKey);
 
 			return IntCastChecked<FBlackboard::FKey>(OffsetKey);
