Fix build errors on Clang
#tests rebuild files

diff --git a/Engine/Source/Runtime/AIModule/Classes/BehaviorTree/BehaviorTreeTypes.h b/Engine/Source/Runtime/AIModule/Classes/BehaviorTree/BehaviorTreeTypes.h
--- a/Engine/Source/Runtime/AIModule/Classes/BehaviorTree/BehaviorTreeTypes.h
+++ b/Engine/Source/Runtime/AIModule/Classes/BehaviorTree/BehaviorTreeTypes.h
@@ -56,7 +56,7 @@ namespace FBlackboard
 	struct FKey
 	{
 		constexpr FKey() = default;
-		constexpr FKey(int32 InKey) : Key(InKey) {}
+		FKey(int32 InKey) {Key = IntCastChecked<uint16>(InKey); }
 		constexpr FKey(uint16 InKey) : Key(InKey) {}
 		constexpr FKey(uint8 InKey) = delete;
 		constexpr operator int32() const { return Key; }
diff --git a/Engine/Source/Runtime/AIModule/Private/BehaviorTree/BlackboardComponent.cpp b/Engine/Source/Runtime/AIModule/Private/BehaviorTree/BlackboardComponent.cpp
--- a/Engine/Source/Runtime/AIModule/Private/BehaviorTree/BlackboardComponent.cpp
+++ b/Engine/Source/Runtime/AIModule/Private/BehaviorTree/BlackboardComponent.cpp
@@ -156,9 +156,9 @@ bool UBlackboardComponent::InitializeBlackboard(UBlackboardData& NewAsset)
 
 					const uint16 KeyMemory = KeyType->GetValueSize() + (KeyType->HasInstance() ? sizeof(FBlackboardInstancedKeyMemory) : 0);
 					const int32 OffsetKey = KeyIndex + (int32)It->GetFirstKeyID();
-					check(OffsetKey != FBlackboard::InvalidKey);
+					check(FBlackboard::FKey(OffsetKey) != FBlackboard::InvalidKey);
 
-					InitList.Add(FBlackboardInitializationData(IntCastChecked<FBlackboard::FKey>(OffsetKey), KeyMemory));
+					InitList.Add(FBlackboardInitializationData(FBlackboard::FKey(OffsetKey), KeyMemory));
 				}
 			}
 		}
@@ -212,9 +212,9 @@ void UBlackboardComponent::DestroyValues()
 			if (KeyType)
 			{
 				const int32 OffsetKey = KeyIndex + (int32)It->GetFirstKeyID();
-				check(OffsetKey != FBlackboard::InvalidKey);
+				check(FBlackboard::FKey(OffsetKey) != FBlackboard::InvalidKey);
 
-				uint8* KeyMemory = GetKeyRawData(IntCastChecked<FBlackboard::FKey>(OffsetKey));
+				uint8* KeyMemory = GetKeyRawData(FBlackboard::FKey(OffsetKey));
 				KeyType->WrappedFree(*this, KeyMemory);
 			}
 		}
@@ -525,9 +525,9 @@ FString UBlackboardComponent::GetDebugInfoString(EBlackboardDescription::Type Mo
 		for (int32 KeyIndex = 0; KeyIndex < It->Keys.Num(); KeyIndex++)
 		{
 			const int32 OffsetKey = KeyIndex + Offset;
-			check(OffsetKey != FBlackboard::InvalidKey);
+			check(FBlackboard::FKey(OffsetKey) != FBlackboard::InvalidKey);
 
-			KeyDesc.Add(DescribeKeyValue(IntCastChecked<FBlackboard::FKey>(OffsetKey), Mode));
+			KeyDesc.Add(DescribeKeyValue(FBlackboard::FKey(OffsetKey), Mode));
 		}
 		Offset += It->Keys.Num();
 	}
@@ -617,9 +617,9 @@ void UBlackboardComponent::DescribeSelfToVisLog(FVisualLogEntry* Snapshot) const
 				const FBlackboardEntry& Key = It->Keys[KeyIndex];
 
 				const int32 OffsetKey = (int32)It->GetFirstKeyID() + KeyIndex;
-				check(OffsetKey != FBlackboard::InvalidKey);
+				check(FBlackboard::FKey(OffsetKey) != FBlackboard::InvalidKey);
 
-				const uint8* ValueData = GetKeyRawData(IntCastChecked<FBlackboard::FKey>(OffsetKey));
+				const uint8* ValueData = GetKeyRawData(FBlackboard::FKey(OffsetKey));
 				FString ValueDesc = Key.KeyType ? *(Key.KeyType->WrappedDescribeValue(*this, ValueData)) : TEXT("empty");
 
 				Category.Add(Key.EntryName.ToString(), ValueDesc);
diff --git a/Engine/Source/Runtime/AIModule/Private/BehaviorTree/BlackboardData.cpp b/Engine/Source/Runtime/AIModule/Private/BehaviorTree/BlackboardData.cpp
--- a/Engine/Source/Runtime/AIModule/Private/BehaviorTree/BlackboardData.cpp
+++ b/Engine/Source/Runtime/AIModule/Private/BehaviorTree/BlackboardData.cpp
@@ -108,9 +108,9 @@ FBlackboard::FKey UBlackboardData::InternalGetKeyID(const FName& KeyName, EKeyLo
 		if (Keys[KeyIndex].EntryName == KeyName)
 		{
 			const int32 OffsetKey = KeyIndex + (int32)FirstKeyID;
-			check(OffsetKey != FBlackboard::InvalidKey);
+			check(FBlackboard::FKey(OffsetKey) != FBlackboard::InvalidKey);
 
-			return IntCastChecked<FBlackboard::FKey>(OffsetKey);
+			return FBlackboard::FKey(OffsetKey);
 		}
 	}
 
@@ -284,9 +284,9 @@ void UBlackboardData::UpdateParentKeys()
 void UBlackboardData::UpdateKeyIDs()
 {
 	const int32 FirstKeyIDInt = Parent ? Parent->GetNumKeys() : 0;
-	check(FirstKeyIDInt != FBlackboard::InvalidKey);
+	check(FBlackboard::FKey(FirstKeyIDInt) != FBlackboard::InvalidKey);
 
-	FirstKeyID = Parent ? IntCastChecked<FBlackboard::FKey>(FirstKeyIDInt) : 0;
+	FirstKeyID = FBlackboard::FKey(FirstKeyIDInt);
 }
 
 void UBlackboardData::UpdateDeprecatedKeys()
