Fix build errors on Clang
#tests rebuild files

diff --git a/Engine/Source/Runtime/AIModule/Private/BehaviorTree/BehaviorTreeComponent.cpp b/Engine/Source/Runtime/AIModule/Private/BehaviorTree/BehaviorTreeComponent.cpp
--- a/Engine/Source/Runtime/AIModule/Private/BehaviorTree/BehaviorTreeComponent.cpp
+++ b/Engine/Source/Runtime/AIModule/Private/BehaviorTree/BehaviorTreeComponent.cpp
@@ -3363,7 +3363,7 @@ void UBehaviorTreeComponent::StoreDebuggerBlackboard(TMap<FName, FString>& Black
 
 		for (int32 KeyIndex = 0; KeyIndex < NumKeys; KeyIndex++)
 		{
-			const FBlackboard::FKey Key = IntCastChecked<FBlackboard::FKey>(KeyIndex);
+			const FBlackboard::FKey Key = FBlackboard::FKey(KeyIndex);
 
 			FString Value = BlackboardComp->DescribeKeyValue(Key, EBlackboardDescription::OnlyValue);
 			if (Value.Len() == 0)
