Add option in AISense_Sight to perform asynchronous trace queries for the engine version of the visbility test

diff --git a/Engine/Source/Runtime/AIModule/Classes/Perception/AISense_Sight.h b/Engine/Source/Runtime/AIModule/Classes/Perception/AISense_Sight.h
--- a/Engine/Source/Runtime/AIModule/Classes/Perception/AISense_Sight.h
+++ b/Engine/Source/Runtime/AIModule/Classes/Perception/AISense_Sight.h
@@ -5,6 +5,7 @@
 #include "CoreMinimal.h"
 #include "UObject/ObjectMacros.h"
 #include "GenericTeamAgentInterface.h"
+#include "WorldCollision.h"
 #include "Misc/MTAccessDetector.h"
 #include "Perception/AISense.h"
 #include "AISense_Sight.generated.h"
@@ -81,19 +82,54 @@ struct FAISightQuery
 	/** User data that can be used inside the IAISightTargetInterface::CanBeSeenFrom method to store a persistence state */ 
 	mutable int32 UserData; 
 
-	uint64 bLastResult:1;
-	uint64 LastProcessedFrameNumber :63;
+	union
+	{
+		/**
+		 * We can share the memory for these values because they aren't used at the same time :
+		 * - The FrameInfo is used when the query is queued for an update at a later frame. It stores the last time the
+		 *   query was processed so that we can prioritize it accordingly against the other queries
+		 * - The TraceInfo is used when the query has requested a asynchronous trace and is waiting for the result.
+		 *   The engine guarantees that we'll get the info at the next frame, but since we can have multiple queries that
+		 *   are pending at the same time, we need to store some information to identify them when receiving the result callback
+		 */
+
+		struct  
+		{
+			uint64 bLastResult:1;
+			uint64 LastProcessedFrameNumber:63;
+		} FrameInfo;
+
+		/**
+		 * The 'FrameNumber' value can increase indefinitely while the 'Index' represents the number of queries that were
+		 * already requested during this frame. So it shouldn't reach high values in the allocated 32 bits.
+		 * Thanks to that we can reliable only use 31 bits for this value and thus have space to keep the bLastResult value
+		 */
+		struct
+		{
+			uint32 bLastResult:1;
+			uint32 Index:31;
+			uint32 FrameNumber;
+		} TraceInfo;
+	};
 
 	FAISightQuery(FPerceptionListenerID ListenerId = FPerceptionListenerID::InvalidID(), FAISightTarget::FTargetId Target = FAISightTarget::InvalidTargetId)
-		: ObserverId(ListenerId), TargetId(Target), Score(0), Importance(0), LastSeenLocation(FAISystem::InvalidLocation), UserData(0), bLastResult(false), LastProcessedFrameNumber(GFrameCounter)
+		: ObserverId(ListenerId), TargetId(Target), Score(0), Importance(0), LastSeenLocation(FAISystem::InvalidLocation), UserData(0)
 	{
+		FrameInfo.bLastResult = false;
+		FrameInfo.LastProcessedFrameNumber = GFrameCounter;
 	}
 
+	/**
+	 * Note: This should only be called on queries that are queued up for later processing (in SightQueriesOutOfRange or SightQueriesOutOfRange)
+	 */
 	float GetAge() const
 	{
-		return (float)(GFrameCounter - LastProcessedFrameNumber);
+		return (float)(GFrameCounter - FrameInfo.LastProcessedFrameNumber);
 	}
 
+	/**
+	 * Note: This should only be called on queries that are queued up for later processing (in SightQueriesOutOfRange or SightQueriesOutOfRange)
+	 */
 	void RecalcScore()
 	{
 		Score = GetAge() + Importance;
@@ -101,13 +137,33 @@ struct FAISightQuery
 
 	void OnProcessed()
 	{
-		LastProcessedFrameNumber = GFrameCounter;
+		FrameInfo.LastProcessedFrameNumber = GFrameCounter;
 	}
 
 	void ForgetPreviousResult()
 	{
 		LastSeenLocation = FAISystem::InvalidLocation;
-		bLastResult = false;
+		SetLastResult(false);
+	}
+
+	bool GetLastResult() const
+	{
+		return FrameInfo.bLastResult;
+	}
+
+	void SetLastResult(const bool bValue)
+	{
+		FrameInfo.bLastResult = bValue;
+	}
+
+	/**
+	* Note: This only be called for pending queries because it will erase the LastProcessedFrameNumber value
+	*/
+	void SetTraceInfo(const FTraceHandle& TraceHandle)
+	{
+		check((TraceHandle._Data.Index & (static_cast<uint32>(1) << 31)) == 0);
+		TraceInfo.Index = TraceHandle._Data.Index;
+		TraceInfo.FrameNumber = TraceHandle._Data.FrameNumber;
 	}
 
 	class FSortPredicate
@@ -213,9 +269,14 @@ protected:
 	UPROPERTY(EditDefaultsOnly, Category = "AI Perception", config)
 	float PendingQueriesBudgetReductionRatio;
 
+	/** Defines if we are allowed to use asynchronous trace queries when there is no IAISightTargetInterface for a Target */
+	UPROPERTY(EditDefaultsOnly, Category = "AI Perception", config)
+	bool bUseAsynchronousTraceForDefaultSightQueries;
+
 	ECollisionChannel DefaultSightCollisionChannel;
 
-	FOnPendingVisibilityQueryProcessedDelegate OnPendingVisibilityQueryProcessedDelegate;
+	FOnPendingVisibilityQueryProcessedDelegate OnPendingCanBeSeenQueryProcessedDelegate;
+	FTraceDelegate OnPendingTraceQueryProcessedDelegate;
 
 	UE_MT_DECLARE_RW_ACCESS_DETECTOR(QueriesListAccessDetector);
 
@@ -234,10 +295,12 @@ public:
 protected:
 	virtual float Update() override;
 
-	EVisibilityResult ComputeVisibility(const UWorld* World, FAISightQuery& SightQuery, FPerceptionListener& Listener, const AActor* ListenerActor, FAISightTarget& Target, AActor* TargetActor, const FDigestedSightProperties& PropDigest, float& OutStimulusStrength, FVector& OutSeenLocation, int32& OutNumberOfLoSChecksPerformed, int32& OutNumberOfAsyncLosCheckRequested) const;
+	EVisibilityResult ComputeVisibility(UWorld* World, FAISightQuery& SightQuery, FPerceptionListener& Listener, const AActor* ListenerActor, FAISightTarget& Target, AActor* TargetActor, const FDigestedSightProperties& PropDigest, float& OutStimulusStrength, FVector& OutSeenLocation, int32& OutNumberOfLoSChecksPerformed, int32& OutNumberOfAsyncLosCheckRequested) const;
 	virtual bool ShouldAutomaticallySeeTarget(const FDigestedSightProperties& PropDigest, FAISightQuery* SightQuery, FPerceptionListener& Listener, AActor* TargetActor, float& OutStimulusStrength) const;
 	void UpdateQueryVisibilityStatus(FAISightQuery& SightQuery, FPerceptionListener& Listener, const bool bIsVisible, const FVector& SeenLocation, const float StimulusStrength, AActor* TargetActor, const FVector& TargetLocation) const;
-	void OnPendingVisibilityQueryProcessed(const FAISightQueryID& QueryID, const bool bIsVisible, const float StimulusStrength, const FVector& SeenLocation, const TOptional<int32>& UserData);
+	void OnPendingCanBeSeenQueryProcessed(const FAISightQueryID& QueryID, const bool bIsVisible, const float StimulusStrength, const FVector& SeenLocation, const TOptional<int32>& UserData);
+	void OnPendingTraceQueryProcessed(const FTraceHandle& TraceHandle, FTraceDatum& TraceDatum);
+	void OnPendingQueryProcessed(const int32 SightQueryIndex, const bool bIsVisible, const float StimulusStrength, const FVector& SeenLocation, const TOptional<int32>& UserData, const TOptional<AActor*> InTargetActor = NullOpt);
 
 	void OnNewListenerImpl(const FPerceptionListener& NewListener);
 	void OnListenerUpdateImpl(const FPerceptionListener& UpdatedListener);
diff --git a/Engine/Source/Runtime/AIModule/Private/Perception/AISense_Sight.cpp b/Engine/Source/Runtime/AIModule/Private/Perception/AISense_Sight.cpp
--- a/Engine/Source/Runtime/AIModule/Private/Perception/AISense_Sight.cpp
+++ b/Engine/Source/Runtime/AIModule/Private/Perception/AISense_Sight.cpp
@@ -40,6 +40,8 @@ static const int32 DefaultMaxTracesPerTick = 6;
 static const int32 DefaultMaxAsyncTracesPerTick = 10;
 static const int32 DefaultMinQueriesPerTimeSliceCheck = 40;
 static const float DefaultPendingQueriesBudgetReductionRatio = 0.5f;
+static const bool bDefaultUseAsynchronousTraceForDefaultSightQueries = false;
+static const float DefaultStimulusStrength = 1.f;
 
 enum class EForEachResult : uint8
 {
@@ -130,6 +132,7 @@ UAISense_Sight::UAISense_Sight(const FObjectInitializer& ObjectInitializer)
 	, MaxQueryImportance(60.f)
 	, SightLimitQueryImportance(10.f)
 	, PendingQueriesBudgetReductionRatio(DefaultPendingQueriesBudgetReductionRatio)
+	, bUseAsynchronousTraceForDefaultSightQueries(bDefaultUseAsynchronousTraceForDefaultSightQueries)
 {
 	if (HasAnyFlags(RF_ClassDefaultObject) == false)
 	{
@@ -140,7 +143,8 @@ UAISense_Sight::UAISense_Sight(const FObjectInitializer& ObjectInitializer)
 		OnListenerUpdateDelegate.BindUObject(this, &UAISense_Sight::OnListenerUpdateImpl);
 		OnListenerRemovedDelegate.BindUObject(this, &UAISense_Sight::OnListenerRemovedImpl);
 
-		OnPendingVisibilityQueryProcessedDelegate.BindUObject(this, &UAISense_Sight::OnPendingVisibilityQueryProcessed);
+		OnPendingCanBeSeenQueryProcessedDelegate.BindUObject(this, &UAISense_Sight::OnPendingCanBeSeenQueryProcessed);
+		OnPendingTraceQueryProcessedDelegate.BindUObject(this, &UAISense_Sight::OnPendingTraceQueryProcessed);
 	}
 
 	NotifyType = EAISenseNotifyType::OnPerceptionChange;
@@ -177,9 +181,9 @@ bool UAISense_Sight::ShouldAutomaticallySeeTarget(const FDigestedSightProperties
 	return false;
 }
 
-#if AISENSE_SIGHT_TIMESLICING_DEBUG
-namespace 
+namespace UE::AISense_Sight
 {
+#if AISENSE_SIGHT_TIMESLICING_DEBUG
 	struct FTimingSlicingInfo
 	{
 		FTimingSlicingInfo() 
@@ -227,14 +231,25 @@ namespace
 			return Info;
 		}
 	};
-}
 #endif // AISENSE_SIGHT_TIMESLICING_DEBUG
 
+	bool IsTraceConsideredVisible(const FHitResult* HitResult, const AActor* TargetActor)
+	{
+		if (HitResult == nullptr)
+		{
+			return true;
+		}
+
+		const AActor* HitResultActor = HitResult->HitObjectHandle.FetchActor();
+		return (HitResultActor ? HitResultActor->IsOwnedBy(TargetActor) : false);
+	}
+}
+
 float UAISense_Sight::Update()
 {
 	SCOPE_CYCLE_COUNTER(STAT_AI_Sense_Sight);
 
-	const UWorld* World = GEngine->GetWorldFromContextObject(GetPerceptionSystem()->GetOuter(), EGetWorldErrorMode::LogAndReturnNull);
+	UWorld* World = GEngine->GetWorldFromContextObject(GetPerceptionSystem()->GetOuter(), EGetWorldErrorMode::LogAndReturnNull);
 
 	if (World == nullptr)
 	{
@@ -272,7 +287,7 @@ float UAISense_Sight::Update()
 	const double TimeSliceEnd = FPlatformTime::Seconds() + MaxTimeSlicePerTick;
 	bool bHitTimeSliceLimit = false;
 #if AISENSE_SIGHT_TIMESLICING_DEBUG
-	FTimingSlicingInfo SlicingInfo;
+	UE::AISense_Sight::FTimingSlicingInfo SlicingInfo;
 #endif // AISENSE_SIGHT_TIMESLICING_DEBUG
 	static const int32 InitialInvalidItemsSize = 16;
 	enum class EOperationType : uint8
@@ -346,7 +361,7 @@ float UAISense_Sight::Update()
 		{
 			const FDigestedSightProperties& PropDigest = DigestedProperties[SightQuery->ObserverId];
 			const AActor* ListenerBodyActor = ListenerPtr->GetBodyActor();
-			float StimulusStrength = 1.f;
+			float StimulusStrength = DefaultStimulusStrength;
 			FVector SeenLocation(0.f);
 			int32 NumberOfLoSChecksPerformed = 0;
 			int32 NumberOfAsyncLosCheckRequested = 0;
@@ -359,16 +374,13 @@ float UAISense_Sight::Update()
 			if (VisibilityResult == EVisibilityResult::Pending)
 			{
 				QueryOperations.Add(FQueryOperation(bIsInRangeQuery, EOperationType::MoveToPending, bIsInRangeQuery ? InRangeIndex : OutOfRangeIndex));
-
-				// Call this to be able to know how many frames were used to process the visibility
-				SightQuery->OnProcessed();
 			}
 			else
 			{
 				UE_CLOG(VisibilityResult != EVisibilityResult::Visible && VisibilityResult != EVisibilityResult::NotVisible, LogAIPerception, Error, TEXT("UAISense_Sight::Update received invalid Visibility result [%d] for query between Listener %s and Target %s. We'll consider it as NotVisible"), VisibilityResult, *GetNameSafe(ListenerBodyActor), *GetNameSafe(TargetActor));
 
 				const bool bIsVisible = VisibilityResult == EVisibilityResult::Visible;
-				const bool bWasVisible = SightQuery->bLastResult;
+				const bool bWasVisible = SightQuery->GetLastResult();
 				const FVector TargetLocation = TargetActor->GetActorLocation();
 				UpdateQueryVisibilityStatus(*SightQuery, Listener, bIsVisible, SeenLocation, StimulusStrength, TargetActor, TargetLocation);
 
@@ -491,7 +503,7 @@ float UAISense_Sight::Update()
 	return 0.f;
 }
 
-UAISense_Sight::EVisibilityResult UAISense_Sight::ComputeVisibility(const UWorld* World, FAISightQuery& SightQuery, FPerceptionListener& Listener, const AActor* ListenerActor, FAISightTarget& Target, AActor* TargetActor, const FDigestedSightProperties& PropDigest, float& OutStimulusStrength, FVector& OutSeenLocation, int32& OutNumberOfLoSChecksPerformed, int32& OutNumberOfAsyncLosCheckRequested) const
+UAISense_Sight::EVisibilityResult UAISense_Sight::ComputeVisibility(UWorld* World, FAISightQuery& SightQuery, FPerceptionListener& Listener, const AActor* ListenerActor, FAISightTarget& Target, AActor* TargetActor, const FDigestedSightProperties& PropDigest, float& OutStimulusStrength, FVector& OutSeenLocation, int32& OutNumberOfLoSChecksPerformed, int32& OutNumberOfAsyncLosCheckRequested) const
 {
 	SCOPE_CYCLE_COUNTER(STAT_AI_Sense_Sight_ComputeVisibility);
 
@@ -503,7 +515,7 @@ UAISense_Sight::EVisibilityResult UAISense_Sight::ComputeVisibility(const UWorld
 	}
 
 	const FVector TargetLocation = TargetActor->GetActorLocation();
-	const float SightRadiusSq = SightQuery.bLastResult ? PropDigest.LoseSightRadiusSq : PropDigest.SightRadiusSq;
+	const float SightRadiusSq = SightQuery.GetLastResult() ? PropDigest.LoseSightRadiusSq : PropDigest.SightRadiusSq;
 	if (!FAISystem::CheckIsTargetInSightCone(Listener.CachedLocation, Listener.CachedDirection, PropDigest.PeripheralVisionAngleCos, PropDigest.PointOfViewBackwardOffset, PropDigest.NearClippingRadiusSq, SightRadiusSq, TargetLocation))
 	{
 		return EVisibilityResult::NotVisible;
@@ -511,37 +523,55 @@ UAISense_Sight::EVisibilityResult UAISense_Sight::ComputeVisibility(const UWorld
 
 	if (Target.SightTargetInterface != nullptr)
 	{
-		const bool bWasVisible = SightQuery.bLastResult;
+		const bool bWasVisible = SightQuery.GetLastResult();
 		FCanBeSeenFromContext Context;
 		Context.SightQueryID = FAISightQueryID(SightQuery);
 		Context.ObserverLocation = Listener.CachedLocation;
 		Context.IgnoreActor = ListenerActor;
 		Context.bWasVisible = &bWasVisible;
 
-		return Target.SightTargetInterface->CanBeSeenFrom(Context, OutSeenLocation, OutNumberOfLoSChecksPerformed, OutNumberOfAsyncLosCheckRequested, OutStimulusStrength, &SightQuery.UserData, &OnPendingVisibilityQueryProcessedDelegate);
+		const EVisibilityResult Result = Target.SightTargetInterface->CanBeSeenFrom(Context, OutSeenLocation, OutNumberOfLoSChecksPerformed, OutNumberOfAsyncLosCheckRequested, OutStimulusStrength, &SightQuery.UserData, &OnPendingCanBeSeenQueryProcessedDelegate);
+		if (Result == EVisibilityResult::Pending)
+		{
+			// we need to clear the trace info value in order to avoid interfering with the engine processed asynchronous queries
+			SightQuery.SetTraceInfo(FTraceHandle());
+		}
+		return Result;
 	}
 	else
 	{
 		// we need to do tests ourselves
-		FHitResult HitResult;
-		const bool bHit = World->LineTraceSingleByChannel(HitResult, Listener.CachedLocation, TargetLocation, DefaultSightCollisionChannel, FCollisionQueryParams(SCENE_QUERY_STAT(AILineOfSight), true, ListenerActor));
-
-		++OutNumberOfLoSChecksPerformed;
-
-		auto HitResultActorIsOwnedByTargetActor = [&HitResult, TargetActor]()
+		const FCollisionQueryParams QueryParams = FCollisionQueryParams(SCENE_QUERY_STAT(AILineOfSight), true, ListenerActor);
+		if (bUseAsynchronousTraceForDefaultSightQueries)
 		{
-			AActor* HitResultActor = HitResult.HitObjectHandle.FetchActor();
-			return (HitResultActor ? HitResultActor->IsOwnedBy(TargetActor) : false);
-		};
+			const FTraceHandle TraceHandle = World->AsyncLineTraceByChannel(EAsyncTraceType::Single, Listener.CachedLocation, TargetLocation, DefaultSightCollisionChannel, QueryParams, FCollisionResponseParams::DefaultResponseParam, &OnPendingTraceQueryProcessedDelegate);
+			if (!TraceHandle.IsValid())
+			{
+				return EVisibilityResult::NotVisible;
+			}
 
-		if (bHit == false || HitResultActorIsOwnedByTargetActor())
-		{
-			OutSeenLocation = TargetLocation;
-			return EVisibilityResult::Visible;
+			++OutNumberOfAsyncLosCheckRequested;
+
+			// store the trace handle information here so that we can identify the associated query when we'll receive the delegate callback
+			SightQuery.SetTraceInfo(TraceHandle);
+			return EVisibilityResult::Pending;
 		}
 		else
 		{
-			return EVisibilityResult::NotVisible;
+			FHitResult HitResult;
+			const bool bHit = World->LineTraceSingleByChannel(HitResult, Listener.CachedLocation, TargetLocation, DefaultSightCollisionChannel, QueryParams, FCollisionResponseParams::DefaultResponseParam);
+
+			++OutNumberOfLoSChecksPerformed;
+
+			if (UE::AISense_Sight::IsTraceConsideredVisible(bHit ? &HitResult : nullptr, TargetActor))
+			{
+				OutSeenLocation = TargetLocation;
+				return EVisibilityResult::Visible;
+			}
+			else
+			{
+				return EVisibilityResult::NotVisible;
+			}
 		}
 	}
 }
@@ -552,24 +582,24 @@ void UAISense_Sight::UpdateQueryVisibilityStatus(FAISightQuery& SightQuery, FPer
 	{
 		const bool bHasValidSeenLocation = SeenLocation != FAISystem::InvalidLocation;
 		Listener.RegisterStimulus(TargetActor, FAIStimulus(*this, StimulusStrength, bHasValidSeenLocation ? SeenLocation : SightQuery.LastSeenLocation, Listener.CachedLocation));
-		SightQuery.bLastResult = true;
+		SightQuery.SetLastResult(true);
 		if (bHasValidSeenLocation)
 		{
 			SightQuery.LastSeenLocation = SeenLocation;
 		}
 	}
 	// communicate failure only if we've seen given actor before
-	else if (SightQuery.bLastResult == true)
+	else if (SightQuery.GetLastResult())
 	{
 		Listener.RegisterStimulus(TargetActor, FAIStimulus(*this, 0.f, TargetLocation, Listener.CachedLocation, FAIStimulus::SensingFailed));
-		SightQuery.bLastResult = false;
+		SightQuery.SetLastResult(false);
 		SightQuery.LastSeenLocation = FAISystem::InvalidLocation;
 	}
 
 	SIGHT_LOG_SEGMENT(ListenerPtr->GetOwner(), Listener.CachedLocation, TargetLocation, bIsVisible ? FColor::Green : FColor::Red, TEXT("TargetID %d"), Target.TargetId);
 }
 
-void UAISense_Sight::OnPendingVisibilityQueryProcessed(const FAISightQueryID& QueryID, const bool bIsVisible, const float StimulusStrength, const FVector& SeenLocation, const TOptional<int32>& UserData)
+void UAISense_Sight::OnPendingCanBeSeenQueryProcessed(const FAISightQueryID& QueryID, const bool bIsVisible, const float StimulusStrength, const FVector& SeenLocation, const TOptional<int32>& UserData)
 {
 	SCOPE_CYCLE_COUNTER(STAT_AI_Sense_Sight_ProcessPendingQuery);
 
@@ -587,19 +617,65 @@ void UAISense_Sight::OnPendingVisibilityQueryProcessed(const FAISightQueryID& Qu
 		return;
 	}
 
-	FAISightQuery SightQuery = SightQueriesPending[QueryIdx];
-	SightQueriesPending.RemoveAtSwap(QueryIdx, 1, false);
+	OnPendingQueryProcessed(QueryIdx, bIsVisible, StimulusStrength, SeenLocation, UserData);
+}
+
+void UAISense_Sight::OnPendingTraceQueryProcessed(const FTraceHandle& TraceHandle, FTraceDatum& TraceDatum)
+{
+	SCOPE_CYCLE_COUNTER(STAT_AI_Sense_Sight_ProcessPendingQuery);
+	UE_MT_SCOPED_WRITE_ACCESS(QueriesListAccessDetector);
+
+	const int32 QueryIdx = SightQueriesPending.IndexOfByPredicate([&TraceHandle](const FAISightQuery& Element)
+	{
+		return Element.TraceInfo.FrameNumber == TraceHandle._Data.FrameNumber
+			&& Element.TraceInfo.Index == TraceHandle._Data.Index;
+	});
+
+	if (QueryIdx == INDEX_NONE)
+	{
+		// the query is not pending. It must have been removed because the source or the target have been removed
+		return;
+	}
+
+	AActor* TargetActor = nullptr;
+	if (const FAISightTarget* Target = ObservedTargets.Find(SightQueriesPending[QueryIdx].TargetId))
+	{
+		TargetActor = Target->Target.Get();
+	}
+	const bool bIsVisible = UE::AISense_Sight::IsTraceConsideredVisible(TraceDatum.OutHits.Num() > 0 ? &TraceDatum.OutHits[0] : nullptr, TargetActor);
+
+	OnPendingQueryProcessed(QueryIdx, bIsVisible, DefaultStimulusStrength, TraceDatum.End, NullOpt, TargetActor);
+}
+
+void UAISense_Sight::OnPendingQueryProcessed(const int32 SightQueryIndex, const bool bIsVisible, const float StimulusStrength, const FVector& SeenLocation, const TOptional<int32>& UserData, const TOptional<AActor*> InTargetActor)
+{
+	FAISightQuery SightQuery = SightQueriesPending[SightQueryIndex];
+	SightQueriesPending.RemoveAtSwap(SightQueryIndex, 1, false);
 
 	AIPerception::FListenerMap& ListenersMap = *GetListeners();
 	FPerceptionListener* Listener = ListenersMap.Find(SightQuery.ObserverId);
-	FAISightTarget* Target = ObservedTargets.Find(SightQuery.TargetId);
-	AActor* TargetActor = Target ? Target->Target.Get() : nullptr;
-	if (TargetActor == nullptr || Listener == nullptr)
+	if (Listener == nullptr)
+	{
+		return;
+	}
+
+	AActor* TargetActor = nullptr;
+	if (InTargetActor.IsSet())
+	{
+		TargetActor = InTargetActor.GetValue();
+	}
+	else
+	{
+		const FAISightTarget* Target = ObservedTargets.Find(SightQuery.TargetId);
+		TargetActor = Target ? Target->Target.Get() : nullptr;
+	}
+
+	if (TargetActor == nullptr)
 	{
 		return;
 	}
 
-	const bool bWasVisible = SightQuery.bLastResult;
+	const bool bWasVisible = SightQuery.GetLastResult();
 	const FVector TargetLocation = TargetActor->GetActorLocation();
 	UpdateQueryVisibilityStatus(SightQuery, *Listener, bIsVisible, SeenLocation, StimulusStrength, TargetActor, TargetLocation);
 
@@ -658,7 +734,7 @@ void UAISense_Sight::UnregisterSource(AActor& SourceActor)
 				FAISightQuery* SightQuery = &SightQueries[QueryIndex];
 				if (SightQuery->TargetId == AsTargetId)
 				{
-					if (SightQuery->bLastResult == true)
+					if (SightQuery->GetLastResult())
 					{
 						FPerceptionListener& Listener = ListenersMap[SightQuery->ObserverId];
 						ensure(Listener.Listener.IsValid());
@@ -826,7 +902,7 @@ void UAISense_Sight::OnListenerUpdateImpl(const FPerceptionListener& UpdatedList
 			TSet<FPerceptionListenerID> LastVisibleObservers;
 			RemoveAllQueriesToTarget(AsTargetId, [&LastVisibleObservers](const FAISightQuery& Query)
 			{
-				if (Query.bLastResult)
+				if (Query.GetLastResult())
 				{
 					LastVisibleObservers.Add(Query.ObserverId);
 				}
@@ -834,7 +910,7 @@ void UAISense_Sight::OnListenerUpdateImpl(const FPerceptionListener& UpdatedList
 
 			RegisterTarget(*(AsTarget->Target.Get()), [&LastVisibleObservers](FAISightQuery& Query)
 			{
-				Query.bLastResult = LastVisibleObservers.Contains(Query.ObserverId);
+				Query.SetLastResult(LastVisibleObservers.Contains(Query.ObserverId));
 			});
 		}
 		else
@@ -851,10 +927,10 @@ void UAISense_Sight::OnListenerUpdateImpl(const FPerceptionListener& UpdatedList
 		TSet<FAISightTarget::FTargetId> LastVisibleTargets;
 		RemoveAllQueriesByListener(UpdatedListener, [&LastVisibleTargets](const FAISightQuery& Query)
 		{
-			if (Query.bLastResult)
+			if (Query.GetLastResult())
 			{
 				LastVisibleTargets.Add(Query.TargetId);
-			}			
+			}
 		});
 
 		const UAISenseConfig_Sight* SenseConfig = Cast<const UAISenseConfig_Sight>(UpdatedListener.Listener->GetSenseConfig(GetSenseID()));
@@ -864,7 +940,7 @@ void UAISense_Sight::OnListenerUpdateImpl(const FPerceptionListener& UpdatedList
 
 		GenerateQueriesForListener(UpdatedListener, PropertiesDigest, [&LastVisibleTargets](FAISightQuery& Query)
 		{
-			Query.bLastResult = LastVisibleTargets.Contains(Query.TargetId);
+			Query.SetLastResult(LastVisibleTargets.Contains(Query.TargetId));
 		});
 	}
 	else
