Virtualize condition to stop movement when a path is finished

diff --git a/Engine/Source/Runtime/AIModule/Classes/Navigation/PathFollowingComponent.h b/Engine/Source/Runtime/AIModule/Classes/Navigation/PathFollowingComponent.h
--- a/Engine/Source/Runtime/AIModule/Classes/Navigation/PathFollowingComponent.h
+++ b/Engine/Source/Runtime/AIModule/Classes/Navigation/PathFollowingComponent.h
@@ -649,6 +649,9 @@ protected:
 	/** set move focus in AI owner */
 	virtual void UpdateMoveFocus();
 
+	/** defines if the agent should reset his velocity when the path is finished*/
+	virtual bool ShouldStopMovementOnPathFinished() const; 
+
 	/** For given path finds a path node at which
 	 *	PathfollowingComponent should start doing 
 	 *	precise is-goal-in-acceptance-radius  tests */
diff --git a/Engine/Source/Runtime/AIModule/Private/Navigation/PathFollowingComponent.cpp b/Engine/Source/Runtime/AIModule/Private/Navigation/PathFollowingComponent.cpp
--- a/Engine/Source/Runtime/AIModule/Private/Navigation/PathFollowingComponent.cpp
+++ b/Engine/Source/Runtime/AIModule/Private/Navigation/PathFollowingComponent.cpp
@@ -573,7 +573,7 @@ void UPathFollowingComponent::OnPathFinished(const FPathFollowingResult& Result)
 	Reset();
 	UpdateMoveFocus();
 
-	if (bStopMovementOnFinish && MovementComp && HasMovementAuthority() && !MovementComp->UseAccelerationForPathFollowing())
+	if (ShouldStopMovementOnPathFinished())
 	{
 		MovementComp->StopMovementKeepPathing();
 	}
@@ -586,6 +586,11 @@ void UPathFollowingComponent::OnPathFinished(const FPathFollowingResult& Result)
 	FAIMessage::Send(Cast<AController>(GetOwner()), Msg);
 }
 
+bool UPathFollowingComponent::ShouldStopMovementOnPathFinished() const
+{
+	return bStopMovementOnFinish && MovementComp && HasMovementAuthority() && !MovementComp->UseAccelerationForPathFollowing();
+}
+
 void UPathFollowingComponent::OnSegmentFinished()
 {
 	UE_VLOG(GetOwner(), LogPathFollowing, Verbose, TEXT("OnSegmentFinished"));
