More conversion warnings fixes

diff --git a/Engine/Plugins/Experimental/StructUtils/Source/StructUtils/Public/InstancedStructArray.h b/Engine/Plugins/Experimental/StructUtils/Source/StructUtils/Public/InstancedStructArray.h
--- a/Engine/Plugins/Experimental/StructUtils/Source/StructUtils/Public/InstancedStructArray.h
+++ b/Engine/Plugins/Experimental/StructUtils/Source/StructUtils/Public/InstancedStructArray.h
@@ -1,4 +1,4 @@
-﻿// Copyright Epic Games, Inc. All Rights Reserved.
+// Copyright Epic Games, Inc. All Rights Reserved.
 
 #pragma once
 
@@ -122,7 +122,7 @@ private:
 	static constexpr int32 DefaultMinAlignment = alignof(void*);
 
 	/** Each FItem struct is allocated at the end of the buffer (which is aligned to FItem alignment), and they are stride apart. */
-	static constexpr int32 ItemStride = Align(sizeof(FItem), alignof(FItem));
+	static constexpr int32 ItemStride = static_cast<int32>(Align(sizeof(FItem), alignof(FItem)));
 
 	/** @returns offset of the item in the allocated memory.  */
 	int32 GetItemOffset(const int32 Index) const
diff --git a/Engine/Plugins/Experimental/StructUtils/Source/StructUtils/Public/SharedStruct.h b/Engine/Plugins/Experimental/StructUtils/Source/StructUtils/Public/SharedStruct.h
--- a/Engine/Plugins/Experimental/StructUtils/Source/StructUtils/Public/SharedStruct.h
+++ b/Engine/Plugins/Experimental/StructUtils/Source/StructUtils/Public/SharedStruct.h
@@ -36,7 +36,7 @@ struct STRUCTUTILS_API FStructSharedMemory
 	{
 		// Align RequiredSize to InScriptStruct's alignment to effectively add padding in between ScriptStruct and
 		// StructMemory. GetMemory will then round &StructMemory up past this 'padding' to the nearest aligned address.
-		const int32 RequiredSize = Align(sizeof(FStructSharedMemory), InScriptStruct.GetMinAlignment()) + InScriptStruct.GetStructureSize();
+		const int32 RequiredSize = static_cast<int32>(Align(sizeof(FStructSharedMemory), InScriptStruct.GetMinAlignment())) + InScriptStruct.GetStructureSize();
 		// Code analysis is unable to understand correctly what we are doing here, so disabling the warning C6386: Buffer overrun while writing to...
 		CA_SUPPRESS( 6386 )
 		FStructSharedMemory* StructMemory = new(FMemory::Malloc(RequiredSize, InScriptStruct.GetMinAlignment())) FStructSharedMemory(InScriptStruct, InStructMemory);
diff --git a/Engine/Plugins/Runtime/NavCorridor/Source/NavCorridor/Private/NavCorridor.cpp b/Engine/Plugins/Runtime/NavCorridor/Source/NavCorridor/Private/NavCorridor.cpp
--- a/Engine/Plugins/Runtime/NavCorridor/Source/NavCorridor/Private/NavCorridor.cpp
+++ b/Engine/Plugins/Runtime/NavCorridor/Source/NavCorridor/Private/NavCorridor.cpp
@@ -35,7 +35,7 @@ namespace UE::NavCorridor::Private
 	};
 
 	/** @return approximate distance between two segments. */
-	static float ApproxDistanceSegmentSegment(const FVector StartA, const FVector EndA, const FVector StartB, const FVector EndB)
+	static FVector::FReal ApproxDistanceSegmentSegment(const FVector StartA, const FVector EndA, const FVector StartB, const FVector EndB)
 	{
 		const FVector2D Seg(EndA - StartA);
 		const FVector Mid = (StartB + EndB) * 0.5;
@@ -447,7 +447,7 @@ namespace UE::NavCorridor::Private
 				continue;
 			}
 
-			float U = EdgeU;
+			FReal U = EdgeU;
 			if (SampleAtV < (Segment.StartUV.Y + UE_KINDA_SMALL_NUMBER))
 			{
 				// Hit Start point
@@ -1522,7 +1522,7 @@ FNavCorridorLocation FNavCorridor::AdvancePathLocation(const FNavCorridorLocatio
 		const FReal SectionLength = FVector::Distance(CurrentLocation, CurrPortal.Location);
 		if ((DistanceSoFar + SectionLength) > AdvanceDistance)
 		{
-			Result.T = (AdvanceDistance - DistanceSoFar) / SectionLength;
+			Result.T = FloatCastChecked<float>((AdvanceDistance - DistanceSoFar) / SectionLength, /* Precision */ 1./256.);
 			Result.Location = FMath::Lerp(CurrentLocation, CurrPortal.Location, Result.T);
 			Result.PortalIndex = PortalIndex - 1;
 			break;
@@ -1539,7 +1539,7 @@ FNavCorridorLocation FNavCorridor::AdvancePathLocation(const FNavCorridorLocatio
 		const FReal SectionLength = FVector::Distance(NextPortal.Location, CurrPortal.Location);
 		const FReal LeftoverDistance = AdvanceDistance - DistanceSoFar;
 		Result.PortalIndex = Portals.Num() - 2;
-		Result.T = 1.0 + LeftoverDistance / SectionLength; // T will be > 1
+		Result.T = 1.0 + FloatCastChecked<float>(LeftoverDistance / SectionLength, /* Precision */ 1. / 256.); // T will be > 1
 		Result.Location = FMath::Lerp(CurrPortal.Location, NextPortal.Location, Result.T);
 	}
 	
diff --git a/Engine/Plugins/Runtime/NavCorridor/Source/NavCorridor/Private/NavCorridorTestingComponent.cpp b/Engine/Plugins/Runtime/NavCorridor/Source/NavCorridor/Private/NavCorridorTestingComponent.cpp
--- a/Engine/Plugins/Runtime/NavCorridor/Source/NavCorridor/Private/NavCorridorTestingComponent.cpp
+++ b/Engine/Plugins/Runtime/NavCorridor/Source/NavCorridor/Private/NavCorridorTestingComponent.cpp
@@ -143,7 +143,7 @@ void UNavCorridorTestingComponent::UpdateTests()
 
 		EndTime = FPlatformTime::Seconds();
 		
-		PathfindingTimeUs = (EndTime - StartTime) * 1000000.0; // us
+		PathfindingTimeUs = FloatCastChecked<float>((EndTime - StartTime) * 1000000.0, UE::LWC::DefaultFloatPrecision); // us
 
 		StartTime = FPlatformTime::Seconds();
 
@@ -154,7 +154,7 @@ void UNavCorridorTestingComponent::UpdateTests()
 		}
 		
 		EndTime = FPlatformTime::Seconds();
-		CorridorTimeUs = (EndTime - StartTime) * 1000000.0; // us
+		CorridorTimeUs = FloatCastChecked<float>((EndTime - StartTime) * 1000000.0, UE::LWC::DefaultFloatPrecision); // us
 	}
 
 	if (bFollowPathOnGoalCorridor)
diff --git a/Engine/Plugins/Runtime/StateTree/Source/StateTreeModule/Public/StateTreeTypes.h b/Engine/Plugins/Runtime/StateTree/Source/StateTreeModule/Public/StateTreeTypes.h
--- a/Engine/Plugins/Runtime/StateTree/Source/StateTreeModule/Public/StateTreeTypes.h
+++ b/Engine/Plugins/Runtime/StateTree/Source/StateTreeModule/Public/StateTreeTypes.h
@@ -325,7 +325,7 @@ struct STATETREEMODULE_API FStateTreeActiveStates
 				States[Index] = FStateTreeStateHandle::Invalid;
 			}
 		}
-		NumStates = NewNum;
+		NumStates = static_cast<uint8>(NewNum);
 	}
 
 	/** Returns true of the array contains specified state. */
diff --git a/Engine/Source/Runtime/AIModule/Classes/BehaviorTree/BehaviorTreeTypes.h b/Engine/Source/Runtime/AIModule/Classes/BehaviorTree/BehaviorTreeTypes.h
--- a/Engine/Source/Runtime/AIModule/Classes/BehaviorTree/BehaviorTreeTypes.h
+++ b/Engine/Source/Runtime/AIModule/Classes/BehaviorTree/BehaviorTreeTypes.h
@@ -651,7 +651,7 @@ public:
 		
 	void AllowNoneAsValue(bool bAllow) { bNoneIsAllowedValue = bAllow; }
 
-	FORCEINLINE FBlackboard::FKey GetSelectedKeyID() const { return SelectedKeyID; }
+	FORCEINLINE FBlackboard::FKey GetSelectedKeyID() const { return FBlackboard::FKey(IntCastChecked<uint16>(SelectedKeyID)); }
 
 	/** helper functions for setting basic filters */
 	void AddObjectFilter(UObject* Owner, FName PropertyName, TSubclassOf<UObject> AllowedClass);
diff --git a/Engine/Source/Runtime/AIModule/Classes/EnvironmentQuery/EnvQueryDebugHelpers.h b/Engine/Source/Runtime/AIModule/Classes/EnvironmentQuery/EnvQueryDebugHelpers.h
--- a/Engine/Source/Runtime/AIModule/Classes/EnvironmentQuery/EnvQueryDebugHelpers.h
+++ b/Engine/Source/Runtime/AIModule/Classes/EnvironmentQuery/EnvQueryDebugHelpers.h
@@ -63,7 +63,7 @@ namespace EQSDebug
 		int32 NumValidItems;
 		int32 Id;
 		FString Name;
-		float Timestamp;
+		double Timestamp;
 
 		void Reset()
 		{
@@ -171,7 +171,7 @@ public:
 	static void LogQuery(FEnvQueryInstance& Query, const FName& CategoryName, ELogVerbosity::Type Verbosity);
 
 private:
-	static void LogQueryInternal(FEnvQueryInstance& Query, const FName& CategoryName, ELogVerbosity::Type Verbosity, float TimeSeconds, FVisualLogEntry *CurrentEntry);
+	static void LogQueryInternal(FEnvQueryInstance& Query, const FName& CategoryName, ELogVerbosity::Type Verbosity, double TimeSeconds, FVisualLogEntry *CurrentEntry);
 #endif
 };
 
diff --git a/Engine/Source/Runtime/AIModule/Classes/EnvironmentQuery/EnvQueryManager.h b/Engine/Source/Runtime/AIModule/Classes/EnvironmentQuery/EnvQueryManager.h
--- a/Engine/Source/Runtime/AIModule/Classes/EnvironmentQuery/EnvQueryManager.h
+++ b/Engine/Source/Runtime/AIModule/Classes/EnvironmentQuery/EnvQueryManager.h
@@ -155,16 +155,16 @@ struct AIMODULE_API FEQSDebugger
 	{
 		// most expensive run
 		FEnvQueryDebugProfileData MostExpensive;
-		float MostExpensiveDuration;
+		double MostExpensiveDuration;
 
 		// average run (sum of all runs, divide by TotalAvgCount to get values)
 		FEnvQueryDebugProfileData TotalAvgData;
-		float TotalAvgDuration;
+		double TotalAvgDuration;
 		int32 TotalAvgCount;
 
 		// EQS tick load
 		TArray<uint8> TickPct;
-		float LastTickTime;
+		double LastTickTime;
 		uint64 LastTickFrame;
 		uint16 FirstTickEntry;
 		uint16 LastTickEntry;
@@ -178,7 +178,7 @@ struct AIMODULE_API FEQSDebugger
 	};
 
 	void StoreStats(const FEnvQueryInstance& QueryInstance);
-	void StoreTickTime(const FEnvQueryInstance& QueryInstance, float TickTime, float MaxTickTime);
+	void StoreTickTime(const FEnvQueryInstance& QueryInstance, double TickTime, double MaxTickTime);
 	void StoreQuery(const TSharedPtr<FEnvQueryInstance>& QueryInstance);
 	
 	static void SaveStats(const FString& FileName);
diff --git a/Engine/Source/Runtime/AIModule/Classes/EnvironmentQuery/EnvQueryTypes.h b/Engine/Source/Runtime/AIModule/Classes/EnvironmentQuery/EnvQueryTypes.h
--- a/Engine/Source/Runtime/AIModule/Classes/EnvironmentQuery/EnvQueryTypes.h
+++ b/Engine/Source/Runtime/AIModule/Classes/EnvironmentQuery/EnvQueryTypes.h
@@ -823,7 +823,7 @@ struct AIMODULE_API FEnvQueryInstance : public FEnvQueryResult
 	double TotalExecutionTime;
 
 	/** if > 0 then it's how much time query has for performing current step */
-	float CurrentStepTimeLimit;
+	double CurrentStepTimeLimit;
 
 #if !UE_BUILD_SHIPPING
 	/** Maximum EQS Generator duration (in seconds) before a warning is reported in non-shipping build. */
@@ -844,7 +844,7 @@ struct AIMODULE_API FEnvQueryInstance : public FEnvQueryResult
 	~FEnvQueryInstance();
 
 	/** execute single step of query */
-	void ExecuteOneStep(float TimeLimit);
+	void ExecuteOneStep(double TimeLimit);
 
 	/** update context cache */
 	bool PrepareContext(UClass* Context, FEnvQueryContextData& ContextData);
diff --git a/Engine/Source/Runtime/AIModule/Classes/EnvironmentQuery/Generators/EnvQueryGenerator_Donut.h b/Engine/Source/Runtime/AIModule/Classes/EnvironmentQuery/Generators/EnvQueryGenerator_Donut.h
--- a/Engine/Source/Runtime/AIModule/Classes/EnvironmentQuery/Generators/EnvQueryGenerator_Donut.h
+++ b/Engine/Source/Runtime/AIModule/Classes/EnvironmentQuery/Generators/EnvQueryGenerator_Donut.h
@@ -61,6 +61,6 @@ class AIMODULE_API UEnvQueryGenerator_Donut : public UEnvQueryGenerator_Projecte
 #endif // WITH_EDITOR
 
 protected:
-	float GetArcBisectorAngle(FEnvQueryInstance& QueryInstance) const;
-	bool IsAngleAllowed(float TestAngleRad, float BisectAngleDeg, float AngleRangeDeg, bool bConstrainAngle) const;
+	FVector::FReal GetArcBisectorAngle(FEnvQueryInstance& QueryInstance) const;
+	bool IsAngleAllowed(FVector::FReal TestAngleRad, FVector::FReal BisectAngleDeg, FVector::FReal AngleRangeDeg, bool bConstrainAngle) const;
 };
diff --git a/Engine/Source/Runtime/AIModule/Classes/Navigation/PathFollowingComponent.h b/Engine/Source/Runtime/AIModule/Classes/Navigation/PathFollowingComponent.h
--- a/Engine/Source/Runtime/AIModule/Classes/Navigation/PathFollowingComponent.h
+++ b/Engine/Source/Runtime/AIModule/Classes/Navigation/PathFollowingComponent.h
@@ -487,7 +487,7 @@ protected:
 	FVector OriginalMoveRequestGoalLocation;
 
 	/** timestamp of path update when movement was paused */
-	float PathTimeWhenPaused;
+	double PathTimeWhenPaused;
 
 	/** Indicates a path node index at which precise "is at goal"
 	 *	tests are going to be performed every frame, in regards
@@ -537,7 +537,7 @@ protected:
 	int32 BlockDetectionSampleCount;
 
 	/** timestamp of last location sample */
-	float LastSampleTime;
+	double LastSampleTime;
 
 	/** index of next location sample in array */
 	int32 NextSampleIdx;
diff --git a/Engine/Source/Runtime/AIModule/Private/BehaviorTree/BehaviorTreeTypes.cpp b/Engine/Source/Runtime/AIModule/Private/BehaviorTree/BehaviorTreeTypes.cpp
--- a/Engine/Source/Runtime/AIModule/Private/BehaviorTree/BehaviorTreeTypes.cpp
+++ b/Engine/Source/Runtime/AIModule/Private/BehaviorTree/BehaviorTreeTypes.cpp
@@ -461,7 +461,7 @@ void FBlackboardKeySelector::ResolveSelectedKey(const UBlackboardData& Blackboar
 		}
 
 		SelectedKeyID = BlackboardAsset.GetKeyID(SelectedKeyName);
-		SelectedKeyType = BlackboardAsset.GetKeyType(SelectedKeyID);
+		SelectedKeyType = BlackboardAsset.GetKeyType(FBlackboard::FKey(IntCastChecked<uint16>(SelectedKeyID)));
 		UE_CLOG(IsSet() == false, LogBehaviorTree, Warning
 			, TEXT("%s> Failed to find key \'%s\' in BB asset %s. BB Key Selector will be set to \'Invalid\'")
 			, *UBehaviorTreeTypes::GetBTLoggingContext()
diff --git a/Engine/Source/Runtime/AIModule/Private/EnvironmentQuery/EnvQueryDebugHelpers.cpp b/Engine/Source/Runtime/AIModule/Private/EnvironmentQuery/EnvQueryDebugHelpers.cpp
--- a/Engine/Source/Runtime/AIModule/Private/EnvironmentQuery/EnvQueryDebugHelpers.cpp
+++ b/Engine/Source/Runtime/AIModule/Private/EnvironmentQuery/EnvQueryDebugHelpers.cpp
@@ -177,7 +177,7 @@ namespace
 	}
 }
 
-void UEnvQueryDebugHelpers::LogQueryInternal(FEnvQueryInstance& Query, const FName& CategoryName, ELogVerbosity::Type Verbosity, float TimeSeconds, FVisualLogEntry *CurrentEntry)
+void UEnvQueryDebugHelpers::LogQueryInternal(FEnvQueryInstance& Query, const FName& CategoryName, ELogVerbosity::Type Verbosity, double TimeSeconds, FVisualLogEntry* CurrentEntry)
 {
 	const int32 UniqueId = FVisualLogger::Get().GetUniqueId(TimeSeconds);
 
diff --git a/Engine/Source/Runtime/AIModule/Private/EnvironmentQuery/EnvQueryInstance.cpp b/Engine/Source/Runtime/AIModule/Private/EnvironmentQuery/EnvQueryInstance.cpp
--- a/Engine/Source/Runtime/AIModule/Private/EnvironmentQuery/EnvQueryInstance.cpp
+++ b/Engine/Source/Runtime/AIModule/Private/EnvironmentQuery/EnvQueryInstance.cpp
@@ -286,7 +286,7 @@ bool FEnvQueryInstance::PrepareContext(UClass* Context, TArray<AActor*>& Data)
 	return Data.Num() > 0;
 }
 
-void FEnvQueryInstance::ExecuteOneStep(float TimeLimit)
+void FEnvQueryInstance::ExecuteOneStep(double TimeLimit)
 {
 	if (!Owner.IsValid())
 	{
@@ -369,7 +369,7 @@ void FEnvQueryInstance::ExecuteOneStep(float TimeLimit)
 					NumProcessedItems = Items.Num() - LastValidItems;
 					LastValidItems = Items.Num();
 					DebugData.CurrentOptionGeneratorIdx++;
-					DebugData.Store(*this, StepExecutionTime, false);
+					DebugData.Store(*this, FloatCastChecked<float>(StepExecutionTime, /* Precision */ 1./512.), false);
 					NumProcessedItems = 0;
 				}
 
@@ -463,7 +463,7 @@ void FEnvQueryInstance::ExecuteOneStep(float TimeLimit)
 #if USE_EQS_DEBUGGER
 	if (bStoreDebugInfo)
 	{
-		DebugData.Store(*this, StepExecutionTime, bStepDone);
+		DebugData.Store(*this, FloatCastChecked<float>(StepExecutionTime, /* Precision */ 1./512.), bStepDone);
 	}
 #endif // USE_EQS_DEBUGGER
 	
diff --git a/Engine/Source/Runtime/AIModule/Private/EnvironmentQuery/EnvQueryManager.cpp b/Engine/Source/Runtime/AIModule/Private/EnvironmentQuery/EnvQueryManager.cpp
--- a/Engine/Source/Runtime/AIModule/Private/EnvironmentQuery/EnvQueryManager.cpp
+++ b/Engine/Source/Runtime/AIModule/Private/EnvironmentQuery/EnvQueryManager.cpp
@@ -1125,7 +1125,7 @@ void FEQSDebugger::StoreStats(const FEnvQueryInstance& QueryInstance)
 	FStatsInfo& UpdateInfo = UEnvQueryManager::DebuggerStats.FindOrAdd(FName(*QueryInstance.QueryName));
 
 	const FEnvQueryDebugProfileData& QueryStats = QueryInstance.DebugData;
-	const float ExecutionTime = QueryInstance.TotalExecutionTime;
+	const double ExecutionTime = QueryInstance.TotalExecutionTime;
 	
 	if (ExecutionTime > UpdateInfo.MostExpensiveDuration)
 	{
@@ -1181,7 +1181,7 @@ void FEQSDebugger::StoreQuery(const TSharedPtr<FEnvQueryInstance>& QueryInstance
 	UpdateInfo.Timestamp = QueryInstance->World->GetTimeSeconds();
 }
 
-void FEQSDebugger::StoreTickTime(const FEnvQueryInstance& QueryInstance, float TickTime, float MaxTickTime)
+void FEQSDebugger::StoreTickTime(const FEnvQueryInstance& QueryInstance, double TickTime, double MaxTickTime)
 {
 #if USE_EQS_TICKLOADDATA
 	SCOPE_CYCLE_COUNTER(STAT_AI_EQS_Debug_StoreTickTime);
@@ -1198,7 +1198,7 @@ void FEQSDebugger::StoreTickTime(const FEnvQueryInstance& QueryInstance, float T
 	if (UpdateInfo.LastTickFrame != GFrameCounter)
 	{
 		UpdateInfo.LastTickFrame = GFrameCounter;
-		UpdateInfo.LastTickTime = 0.0f;
+		UpdateInfo.LastTickTime = 0.;
 	}
 
 	const uint16 TickIdx = GFrameCounter & (NumRecordedTicks - 1);
diff --git a/Engine/Source/Runtime/AIModule/Private/EnvironmentQuery/EnvQueryTraceHelpers.cpp b/Engine/Source/Runtime/AIModule/Private/EnvironmentQuery/EnvQueryTraceHelpers.cpp
--- a/Engine/Source/Runtime/AIModule/Private/EnvironmentQuery/EnvQueryTraceHelpers.cpp
+++ b/Engine/Source/Runtime/AIModule/Private/EnvironmentQuery/EnvQueryTraceHelpers.cpp
@@ -388,7 +388,7 @@ void FEQSHelpers::RunNavProjection(const ANavigationData& NavData, const UObject
 		const FNavLocation* PointsBegin = Points.GetData();
 		int32 NewNum = Algo::StableRemoveIf(Points, [&Workload, PointsBegin](FNavLocation& Point)
 		{
-			return !Workload[&Point - PointsBegin].bResult;
+			return !Workload[IntCastChecked<int32>(&Point - PointsBegin)].bResult;
 		});
 		const bool bAllowShrinking = false;
 		Points.SetNum(NewNum, bAllowShrinking);
diff --git a/Engine/Source/Runtime/AIModule/Private/EnvironmentQuery/Generators/EnvQueryGenerator_Donut.cpp b/Engine/Source/Runtime/AIModule/Private/EnvironmentQuery/Generators/EnvQueryGenerator_Donut.cpp
--- a/Engine/Source/Runtime/AIModule/Private/EnvironmentQuery/Generators/EnvQueryGenerator_Donut.cpp
+++ b/Engine/Source/Runtime/AIModule/Private/EnvironmentQuery/Generators/EnvQueryGenerator_Donut.cpp
@@ -23,9 +23,9 @@ UEnvQueryGenerator_Donut::UEnvQueryGenerator_Donut(const FObjectInitializer& Obj
 	ProjectionData.TraceMode = EEnvQueryTrace::None;
 }
 
-float UEnvQueryGenerator_Donut::GetArcBisectorAngle(FEnvQueryInstance& QueryInstance) const
+FVector::FReal UEnvQueryGenerator_Donut::GetArcBisectorAngle(FEnvQueryInstance& QueryInstance) const
 {
-	float BisectAngle = 0.0f;
+	FRotator::FReal BisectAngle = 0.;
 
 	FVector Direction;
 	if (bDefineArc)
@@ -59,13 +59,13 @@ float UEnvQueryGenerator_Donut::GetArcBisectorAngle(FEnvQueryInstance& QueryInst
 	return BisectAngle;
 }
 
-bool UEnvQueryGenerator_Donut::IsAngleAllowed(float TestAngleRad, float BisectAngleDeg, float AngleRangeDeg, bool bConstrainAngle) const
+bool UEnvQueryGenerator_Donut::IsAngleAllowed(FVector::FReal TestAngleRad, FVector::FReal BisectAngleDeg, FVector::FReal AngleRangeDeg, bool bConstrainAngle) const
 {
 	if (bConstrainAngle)
 	{
-		const float TestAngleDeg = FMath::RadiansToDegrees(TestAngleRad);
-		const float AngleDelta = FRotator::NormalizeAxis(TestAngleDeg - BisectAngleDeg);
-		return (FMath::Abs(AngleDelta) - 0.01f) < AngleRangeDeg;
+		const FVector::FReal TestAngleDeg = FMath::RadiansToDegrees(TestAngleRad);
+		const FVector::FReal AngleDelta = FRotator::NormalizeAxis(TestAngleDeg - BisectAngleDeg);
+		return (FMath::Abs(AngleDelta) - 0.01) < AngleRangeDeg;
 	}
 
 	return true;	
@@ -88,25 +88,25 @@ void UEnvQueryGenerator_Donut::GenerateItems(FEnvQueryInstance& QueryInstance) c
 	PointsPerRing.BindData(BindOwner, QueryInstance.QueryID);
 	ArcAngle.BindData(BindOwner, QueryInstance.QueryID);
 
-	float ArcAngleValue = ArcAngle.GetValue();
-	float InnerRadiusValue = InnerRadius.GetValue();
-	float OuterRadiusValue = OuterRadius.GetValue();
+	FVector::FReal ArcAngleValue = ArcAngle.GetValue();
+	FVector::FReal InnerRadiusValue = InnerRadius.GetValue();
+	FVector::FReal OuterRadiusValue = OuterRadius.GetValue();
 	int32 NumRings = NumberOfRings.GetValue();
 	int32 NumPoints = PointsPerRing.GetValue();
 
-	if ((InnerRadiusValue < 0.f) || (OuterRadiusValue <= 0.f) ||
+	if ((InnerRadiusValue < 0.) || (OuterRadiusValue <= 0.) ||
 		(InnerRadiusValue > OuterRadiusValue) ||
 		(NumRings < 1) || (NumPoints < 1))
 	{
 		return;
 	}
 
-	const float ArcBisectDeg = GetArcBisectorAngle(QueryInstance);
-	const float ArcAngleDeg = FMath::Clamp(ArcAngleValue, 0.0f, 360.0f);
+	const FVector::FReal ArcBisectDeg = GetArcBisectorAngle(QueryInstance);
+	const FVector::FReal ArcAngleDeg = FMath::Clamp(ArcAngleValue, 0., 360.);
 
-	const float RadiusDelta = (OuterRadiusValue - InnerRadiusValue) / (NumRings - 1);
-	const float AngleDelta = 2.0 * PI / NumPoints;
-	float SectionAngle = FMath::DegreesToRadians(ArcBisectDeg);
+	const FVector::FReal RadiusDelta = (OuterRadiusValue - InnerRadiusValue) / (NumRings - 1);
+	const FVector::FReal AngleDelta = 2. * UE_DOUBLE_PI / NumPoints;
+	FVector::FReal SectionAngle = FMath::DegreesToRadians(ArcBisectDeg);
 
 	TArray<FNavLocation> Points;
 	Points.Reserve(NumPoints * NumRings);
@@ -117,13 +117,13 @@ void UEnvQueryGenerator_Donut::GenerateItems(FEnvQueryInstance& QueryInstance) c
 		{
 			if (IsAngleAllowed(SectionAngle, ArcBisectDeg, ArcAngleDeg, bDefineArc))
 			{
-				const float SinValue = FMath::Sin(SectionAngle);
-				const float CosValue = FMath::Cos(SectionAngle);
+				const FVector::FReal SinValue = FMath::Sin(SectionAngle);
+				const FVector::FReal CosValue = FMath::Cos(SectionAngle);
 
-				float RingRadius = InnerRadiusValue;
+				FVector::FReal RingRadius = InnerRadiusValue;
 				for (int32 RingIdx = 0; RingIdx < NumRings; RingIdx++, RingRadius += RadiusDelta)
 				{
-					const FVector RingPos(RingRadius * CosValue, RingRadius * SinValue, 0.0f);
+					const FVector RingPos(RingRadius * CosValue, RingRadius * SinValue, 0.);
 					for (int32 ContextIdx = 0; ContextIdx < CenterPoints.Num(); ContextIdx++)
 					{
 						const FNavLocation PointPos = FNavLocation(CenterPoints[ContextIdx] + RingPos);
@@ -137,22 +137,22 @@ void UEnvQueryGenerator_Donut::GenerateItems(FEnvQueryInstance& QueryInstance) c
 	{	// In order to spiral, we need to change the angle for each ring as well as each spoke.  By changing it as a
 		// fraction of the SectionAngle, we guarantee that the offset won't cross to the starting point of the next
 		// spoke.
-		const float RingAngleDelta = AngleDelta / NumRings;
-		float RingRadius = InnerRadiusValue;
-		float CurrentRingAngleDelta = 0.f;
+		const FVector::FReal RingAngleDelta = AngleDelta / NumRings;
+		FVector::FReal RingRadius = InnerRadiusValue;
+		FVector::FReal CurrentRingAngleDelta = 0.f;
 
 		// So, start with ring angle and then add section angle.
 		for (int32 RingIdx = 0; RingIdx < NumRings; RingIdx++, RingRadius += RadiusDelta, CurrentRingAngleDelta += RingAngleDelta)
 		{
 			for (int32 SectionIdx = 0; SectionIdx < NumPoints; SectionIdx++, SectionAngle += AngleDelta)
 			{
-				float RingSectionAngle = CurrentRingAngleDelta + SectionAngle;
+				FVector::FReal RingSectionAngle = CurrentRingAngleDelta + SectionAngle;
 				if (IsAngleAllowed(RingSectionAngle, ArcBisectDeg, ArcAngleDeg, bDefineArc))
 				{
-					const float SinValue = FMath::Sin(RingSectionAngle);
-					const float CosValue = FMath::Cos(RingSectionAngle);
+					const FVector::FReal SinValue = FMath::Sin(RingSectionAngle);
+					const FVector::FReal CosValue = FMath::Cos(RingSectionAngle);
 
-					const FVector RingPos(RingRadius * CosValue, RingRadius * SinValue, 0.0f);
+					const FVector RingPos(RingRadius * CosValue, RingRadius * SinValue, 0.);
 					for (int32 ContextIdx = 0; ContextIdx < CenterPoints.Num(); ContextIdx++)
 					{
 						const FNavLocation PointPos = FNavLocation(CenterPoints[ContextIdx] + RingPos);
diff --git a/Engine/Source/Runtime/AIModule/Private/EnvironmentQuery/Tests/EnvQueryTest_Trace.cpp b/Engine/Source/Runtime/AIModule/Private/EnvironmentQuery/Tests/EnvQueryTest_Trace.cpp
--- a/Engine/Source/Runtime/AIModule/Private/EnvironmentQuery/Tests/EnvQueryTest_Trace.cpp
+++ b/Engine/Source/Runtime/AIModule/Private/EnvironmentQuery/Tests/EnvQueryTest_Trace.cpp
@@ -158,7 +158,7 @@ bool UEnvQueryTest_Trace::RunSphereTraceTo(const FVector& ItemPos, const FVector
 	FCollisionQueryParams TraceParams(Params);
 	TraceParams.AddIgnoredActor(ItemActor);
 
-	const bool bHit = World->SweepTestByChannel(ContextPos, ItemPos, FQuat::Identity, Channel, FCollisionShape::MakeSphere(Extent.X), TraceParams);
+	const bool bHit = World->SweepTestByChannel(ContextPos, ItemPos, FQuat::Identity, Channel, FCollisionShape::MakeSphere(FloatCastChecked<float>(Extent.X, UE::LWC::DefaultFloatPrecision)), TraceParams);
 	return bHit;
 }
 
@@ -167,7 +167,7 @@ bool UEnvQueryTest_Trace::RunSphereTraceFrom(const FVector& ItemPos, const FVect
 	FCollisionQueryParams TraceParams(Params);
 	TraceParams.AddIgnoredActor(ItemActor);
 
-	const bool bHit = World->SweepTestByChannel(ItemPos, ContextPos, FQuat::Identity, Channel, FCollisionShape::MakeSphere(Extent.X), TraceParams);
+	const bool bHit = World->SweepTestByChannel(ItemPos, ContextPos, FQuat::Identity, Channel, FCollisionShape::MakeSphere(FloatCastChecked<float>(Extent.X, UE::LWC::DefaultFloatPrecision)), TraceParams);
 	return bHit;
 }
 
@@ -176,7 +176,7 @@ bool UEnvQueryTest_Trace::RunCapsuleTraceTo(const FVector& ItemPos, const FVecto
 	FCollisionQueryParams TraceParams(Params);
 	TraceParams.AddIgnoredActor(ItemActor);
 
-	const bool bHit = World->SweepTestByChannel(ContextPos, ItemPos, FQuat::Identity, Channel, FCollisionShape::MakeCapsule(Extent.X, Extent.Z), TraceParams);
+	const bool bHit = World->SweepTestByChannel(ContextPos, ItemPos, FQuat::Identity, Channel, FCollisionShape::MakeCapsule(FloatCastChecked<float>(Extent.X, UE::LWC::DefaultFloatPrecision), FloatCastChecked<float>(Extent.Z, UE::LWC::DefaultFloatPrecision)), TraceParams);
 	return bHit;
 }
 
@@ -185,7 +185,7 @@ bool UEnvQueryTest_Trace::RunCapsuleTraceFrom(const FVector& ItemPos, const FVec
 	FCollisionQueryParams TraceParams(Params);
 	TraceParams.AddIgnoredActor(ItemActor);
 
-	const bool bHit = World->SweepTestByChannel(ItemPos, ContextPos, FQuat::Identity, Channel, FCollisionShape::MakeCapsule(Extent.X, Extent.Z), TraceParams);
+	const bool bHit = World->SweepTestByChannel(ItemPos, ContextPos, FQuat::Identity, Channel, FCollisionShape::MakeCapsule(FloatCastChecked<float>(Extent.X, UE::LWC::DefaultFloatPrecision), FloatCastChecked<float>(Extent.Z, UE::LWC::DefaultFloatPrecision)), TraceParams);
 	return bHit;
 }
 
diff --git a/Engine/Source/Runtime/AIModule/Private/GameplayDebugger/GameplayDebuggerCategory_AI.cpp b/Engine/Source/Runtime/AIModule/Private/GameplayDebugger/GameplayDebuggerCategory_AI.cpp
--- a/Engine/Source/Runtime/AIModule/Private/GameplayDebugger/GameplayDebuggerCategory_AI.cpp
+++ b/Engine/Source/Runtime/AIModule/Private/GameplayDebugger/GameplayDebuggerCategory_AI.cpp
@@ -45,7 +45,7 @@ FGameplayDebuggerCategory_AI::FGameplayDebuggerCategory_AI()
 {
 	bShowOnlyWithDebugActor = false;
 	RawLastPath = nullptr;
-	LastPathUpdateTime = 0.0f;
+	LastPathUpdateTime = 0.;
 
 	SetDataPackReplication<FRepData>(&DataPack);
 	PathDataPackId = SetDataPackReplication<FRepDataPath>(&PathDataPack, EGameplayDebuggerDataPack::ResetOnActorChange);
diff --git a/Engine/Source/Runtime/AIModule/Private/GameplayDebugger/GameplayDebuggerCategory_AI.h b/Engine/Source/Runtime/AIModule/Private/GameplayDebugger/GameplayDebuggerCategory_AI.h
--- a/Engine/Source/Runtime/AIModule/Private/GameplayDebugger/GameplayDebuggerCategory_AI.h
+++ b/Engine/Source/Runtime/AIModule/Private/GameplayDebugger/GameplayDebuggerCategory_AI.h
@@ -88,7 +88,7 @@ protected:
 private:
 	// do NOT read from this pointer, it's only for validating if path has changed and can be invalid!
 	FNavigationPath* RawLastPath;
-	float LastPathUpdateTime;
+	double LastPathUpdateTime;
 };
 
 #endif // WITH_GAMEPLAY_DEBUGGER
diff --git a/Engine/Source/Runtime/AIModule/Private/GameplayDebugger/GameplayDebuggerCategory_NavLocalGrid.cpp b/Engine/Source/Runtime/AIModule/Private/GameplayDebugger/GameplayDebuggerCategory_NavLocalGrid.cpp
--- a/Engine/Source/Runtime/AIModule/Private/GameplayDebugger/GameplayDebuggerCategory_NavLocalGrid.cpp
+++ b/Engine/Source/Runtime/AIModule/Private/GameplayDebugger/GameplayDebuggerCategory_NavLocalGrid.cpp
@@ -165,14 +165,14 @@ protected:
 		const FVector CellMin = FVector(GridBounds.Min.X, GridBounds.Min.Y, GridBounds.Max.Z) + FVector(CellSize * CellIdx.X, CellSize * CellIdx.Y, 0);
 		const FBox CellBox(CellMin, CellMin + FVector(CellSize, CellSize, 0.0f));
 		const FVector Center = CellBox.GetCenter();
-		const FVector Extent = CellBox.GetExtent() - FVector(CellGapSize, CellGapSize, 0.f);
+		const FVector Extent = CellBox.GetExtent() - FVector(CellGapSize, CellGapSize, 0.);
 
 		const int32 FirstVertIdx = Verts.Num();
 
-		Verts.Add(FDynamicMeshVertex(FVector3f(Center.X - Extent.X, Center.Y - Extent.Y, Center.Z - Extent.Z)));
-		Verts.Add(FDynamicMeshVertex(FVector3f(Center.X + Extent.X, Center.Y - Extent.Y, Center.Z - Extent.Z)));
-		Verts.Add(FDynamicMeshVertex(FVector3f(Center.X - Extent.X, Center.Y + Extent.Y, Center.Z - Extent.Z)));
-		Verts.Add(FDynamicMeshVertex(FVector3f(Center.X + Extent.X, Center.Y + Extent.Y, Center.Z - Extent.Z)));
+		Verts.Add(FDynamicMeshVertex(UE::LWC::NarrowWorldPositionChecked(FVector(Center.X - Extent.X, Center.Y - Extent.Y, Center.Z - Extent.Z))));
+		Verts.Add(FDynamicMeshVertex(UE::LWC::NarrowWorldPositionChecked(FVector(Center.X + Extent.X, Center.Y - Extent.Y, Center.Z - Extent.Z))));
+		Verts.Add(FDynamicMeshVertex(UE::LWC::NarrowWorldPositionChecked(FVector(Center.X - Extent.X, Center.Y + Extent.Y, Center.Z - Extent.Z))));
+		Verts.Add(FDynamicMeshVertex(UE::LWC::NarrowWorldPositionChecked(FVector(Center.X + Extent.X, Center.Y + Extent.Y, Center.Z - Extent.Z))));
 
 		StoreQuadIndices(0, 1, 3, 2, FirstVertIdx, Indices);
 	}
@@ -188,14 +188,14 @@ protected:
 
 		const int32 FirstVertIdx = Verts.Num();
 
-		Verts.Add(FDynamicMeshVertex(FVector3f(Center.X - Extent.X, Center.Y - Extent.Y, Center.Z - Extent.Z)));
-		Verts.Add(FDynamicMeshVertex(FVector3f(Center.X + Extent.X, Center.Y - Extent.Y, Center.Z - Extent.Z)));
-		Verts.Add(FDynamicMeshVertex(FVector3f(Center.X - Extent.X, Center.Y + Extent.Y, Center.Z - Extent.Z)));
-		Verts.Add(FDynamicMeshVertex(FVector3f(Center.X + Extent.X, Center.Y + Extent.Y, Center.Z - Extent.Z)));
-		Verts.Add(FDynamicMeshVertex(FVector3f(Center.X - Extent.X, Center.Y - Extent.Y, Center.Z + Extent.Z)));
-		Verts.Add(FDynamicMeshVertex(FVector3f(Center.X + Extent.X, Center.Y - Extent.Y, Center.Z + Extent.Z)));
-		Verts.Add(FDynamicMeshVertex(FVector3f(Center.X - Extent.X, Center.Y + Extent.Y, Center.Z + Extent.Z)));
-		Verts.Add(FDynamicMeshVertex(FVector3f(Center.X + Extent.X, Center.Y + Extent.Y, Center.Z + Extent.Z)));
+		Verts.Add(FDynamicMeshVertex(UE::LWC::NarrowWorldPositionChecked(FVector(Center.X - Extent.X, Center.Y - Extent.Y, Center.Z - Extent.Z))));
+		Verts.Add(FDynamicMeshVertex(UE::LWC::NarrowWorldPositionChecked(FVector(Center.X + Extent.X, Center.Y - Extent.Y, Center.Z - Extent.Z))));
+		Verts.Add(FDynamicMeshVertex(UE::LWC::NarrowWorldPositionChecked(FVector(Center.X - Extent.X, Center.Y + Extent.Y, Center.Z - Extent.Z))));
+		Verts.Add(FDynamicMeshVertex(UE::LWC::NarrowWorldPositionChecked(FVector(Center.X + Extent.X, Center.Y + Extent.Y, Center.Z - Extent.Z))));
+		Verts.Add(FDynamicMeshVertex(UE::LWC::NarrowWorldPositionChecked(FVector(Center.X - Extent.X, Center.Y - Extent.Y, Center.Z + Extent.Z))));
+		Verts.Add(FDynamicMeshVertex(UE::LWC::NarrowWorldPositionChecked(FVector(Center.X + Extent.X, Center.Y - Extent.Y, Center.Z + Extent.Z))));
+		Verts.Add(FDynamicMeshVertex(UE::LWC::NarrowWorldPositionChecked(FVector(Center.X - Extent.X, Center.Y + Extent.Y, Center.Z + Extent.Z))));
+		Verts.Add(FDynamicMeshVertex(UE::LWC::NarrowWorldPositionChecked(FVector(Center.X + Extent.X, Center.Y + Extent.Y, Center.Z + Extent.Z))));
 
 		StoreQuadIndices(4, 5, 7, 6, FirstVertIdx, Indices);
 		StoreQuadIndices(1, 3, 7, 5, FirstVertIdx, Indices);
@@ -214,44 +214,44 @@ protected:
 		const FVector Center = GridBounds.GetCenter();
 		const FVector Extent = GridBounds.GetExtent();
 
-		Verts.Add(FDynamicMeshVertex(FVector3f(Center.X - Extent.X, Center.Y - Extent.Y, Center.Z - Extent.Z)));
-		Verts.Add(FDynamicMeshVertex(FVector3f(Center.X + Extent.X, Center.Y - Extent.Y, Center.Z - Extent.Z)));
-		Verts.Add(FDynamicMeshVertex(FVector3f(Center.X - Extent.X, Center.Y + Extent.Y, Center.Z - Extent.Z)));
-		Verts.Add(FDynamicMeshVertex(FVector3f(Center.X + Extent.X, Center.Y + Extent.Y, Center.Z - Extent.Z)));
-		Verts.Add(FDynamicMeshVertex(FVector3f(Center.X - Extent.X, Center.Y - Extent.Y, Center.Z + Extent.Z)));
-		Verts.Add(FDynamicMeshVertex(FVector3f(Center.X + Extent.X, Center.Y - Extent.Y, Center.Z + Extent.Z)));
-		Verts.Add(FDynamicMeshVertex(FVector3f(Center.X - Extent.X, Center.Y + Extent.Y, Center.Z + Extent.Z)));
-		Verts.Add(FDynamicMeshVertex(FVector3f(Center.X + Extent.X, Center.Y + Extent.Y, Center.Z + Extent.Z)));
-
-		Verts.Add(FDynamicMeshVertex(FVector3f(Center.X - Extent.X + FaceWidth, Center.Y - Extent.Y, Center.Z - Extent.Z + FaceWidth)));
-		Verts.Add(FDynamicMeshVertex(FVector3f(Center.X + Extent.X - FaceWidth, Center.Y - Extent.Y, Center.Z - Extent.Z + FaceWidth)));
-		Verts.Add(FDynamicMeshVertex(FVector3f(Center.X - Extent.X + FaceWidth, Center.Y - Extent.Y, Center.Z + Extent.Z - FaceWidth)));
-		Verts.Add(FDynamicMeshVertex(FVector3f(Center.X + Extent.X - FaceWidth, Center.Y - Extent.Y, Center.Z + Extent.Z - FaceWidth)));
-
-		Verts.Add(FDynamicMeshVertex(FVector3f(Center.X + Extent.X, Center.Y - Extent.Y + FaceWidth, Center.Z - Extent.Z + FaceWidth)));
-		Verts.Add(FDynamicMeshVertex(FVector3f(Center.X + Extent.X, Center.Y + Extent.Y - FaceWidth, Center.Z - Extent.Z + FaceWidth)));
-		Verts.Add(FDynamicMeshVertex(FVector3f(Center.X + Extent.X, Center.Y - Extent.Y + FaceWidth, Center.Z + Extent.Z - FaceWidth)));
-		Verts.Add(FDynamicMeshVertex(FVector3f(Center.X + Extent.X, Center.Y + Extent.Y - FaceWidth, Center.Z + Extent.Z - FaceWidth)));
-
-		Verts.Add(FDynamicMeshVertex(FVector3f(Center.X - Extent.X + FaceWidth, Center.Y - Extent.Y + FaceWidth, Center.Z + Extent.Z)));
-		Verts.Add(FDynamicMeshVertex(FVector3f(Center.X + Extent.X - FaceWidth, Center.Y - Extent.Y + FaceWidth, Center.Z + Extent.Z)));
-		Verts.Add(FDynamicMeshVertex(FVector3f(Center.X - Extent.X + FaceWidth, Center.Y + Extent.Y - FaceWidth, Center.Z + Extent.Z)));
-		Verts.Add(FDynamicMeshVertex(FVector3f(Center.X + Extent.X - FaceWidth, Center.Y + Extent.Y - FaceWidth, Center.Z + Extent.Z)));
-
-		Verts.Add(FDynamicMeshVertex(FVector3f(Center.X - Extent.X + FaceWidth, Center.Y + Extent.Y, Center.Z - Extent.Z + FaceWidth)));
-		Verts.Add(FDynamicMeshVertex(FVector3f(Center.X + Extent.X - FaceWidth, Center.Y + Extent.Y, Center.Z - Extent.Z + FaceWidth)));
-		Verts.Add(FDynamicMeshVertex(FVector3f(Center.X - Extent.X + FaceWidth, Center.Y + Extent.Y, Center.Z + Extent.Z - FaceWidth)));
-		Verts.Add(FDynamicMeshVertex(FVector3f(Center.X + Extent.X - FaceWidth, Center.Y + Extent.Y, Center.Z + Extent.Z - FaceWidth)));
-
-		Verts.Add(FDynamicMeshVertex(FVector3f(Center.X - Extent.X, Center.Y - Extent.Y + FaceWidth, Center.Z - Extent.Z + FaceWidth)));
-		Verts.Add(FDynamicMeshVertex(FVector3f(Center.X - Extent.X, Center.Y + Extent.Y - FaceWidth, Center.Z - Extent.Z + FaceWidth)));
-		Verts.Add(FDynamicMeshVertex(FVector3f(Center.X - Extent.X, Center.Y - Extent.Y + FaceWidth, Center.Z + Extent.Z - FaceWidth)));
-		Verts.Add(FDynamicMeshVertex(FVector3f(Center.X - Extent.X, Center.Y + Extent.Y - FaceWidth, Center.Z + Extent.Z - FaceWidth)));
-
-		Verts.Add(FDynamicMeshVertex(FVector3f(Center.X - Extent.X + FaceWidth, Center.Y - Extent.Y + FaceWidth, Center.Z - Extent.Z)));
-		Verts.Add(FDynamicMeshVertex(FVector3f(Center.X + Extent.X - FaceWidth, Center.Y - Extent.Y + FaceWidth, Center.Z - Extent.Z)));
-		Verts.Add(FDynamicMeshVertex(FVector3f(Center.X - Extent.X + FaceWidth, Center.Y + Extent.Y - FaceWidth, Center.Z - Extent.Z)));
-		Verts.Add(FDynamicMeshVertex(FVector3f(Center.X + Extent.X - FaceWidth, Center.Y + Extent.Y - FaceWidth, Center.Z - Extent.Z)));
+		Verts.Add(FDynamicMeshVertex(UE::LWC::NarrowWorldPositionChecked(FVector(Center.X - Extent.X, Center.Y - Extent.Y, Center.Z - Extent.Z))));
+		Verts.Add(FDynamicMeshVertex(UE::LWC::NarrowWorldPositionChecked(FVector(Center.X + Extent.X, Center.Y - Extent.Y, Center.Z - Extent.Z))));
+		Verts.Add(FDynamicMeshVertex(UE::LWC::NarrowWorldPositionChecked(FVector(Center.X - Extent.X, Center.Y + Extent.Y, Center.Z - Extent.Z))));
+		Verts.Add(FDynamicMeshVertex(UE::LWC::NarrowWorldPositionChecked(FVector(Center.X + Extent.X, Center.Y + Extent.Y, Center.Z - Extent.Z))));
+		Verts.Add(FDynamicMeshVertex(UE::LWC::NarrowWorldPositionChecked(FVector(Center.X - Extent.X, Center.Y - Extent.Y, Center.Z + Extent.Z))));
+		Verts.Add(FDynamicMeshVertex(UE::LWC::NarrowWorldPositionChecked(FVector(Center.X + Extent.X, Center.Y - Extent.Y, Center.Z + Extent.Z))));
+		Verts.Add(FDynamicMeshVertex(UE::LWC::NarrowWorldPositionChecked(FVector(Center.X - Extent.X, Center.Y + Extent.Y, Center.Z + Extent.Z))));
+		Verts.Add(FDynamicMeshVertex(UE::LWC::NarrowWorldPositionChecked(FVector(Center.X + Extent.X, Center.Y + Extent.Y, Center.Z + Extent.Z))));
+
+		Verts.Add(FDynamicMeshVertex(UE::LWC::NarrowWorldPositionChecked(FVector(Center.X - Extent.X + FaceWidth, Center.Y - Extent.Y, Center.Z - Extent.Z + FaceWidth))));
+		Verts.Add(FDynamicMeshVertex(UE::LWC::NarrowWorldPositionChecked(FVector(Center.X + Extent.X - FaceWidth, Center.Y - Extent.Y, Center.Z - Extent.Z + FaceWidth))));
+		Verts.Add(FDynamicMeshVertex(UE::LWC::NarrowWorldPositionChecked(FVector(Center.X - Extent.X + FaceWidth, Center.Y - Extent.Y, Center.Z + Extent.Z - FaceWidth))));
+		Verts.Add(FDynamicMeshVertex(UE::LWC::NarrowWorldPositionChecked(FVector(Center.X + Extent.X - FaceWidth, Center.Y - Extent.Y, Center.Z + Extent.Z - FaceWidth))));
+
+		Verts.Add(FDynamicMeshVertex(UE::LWC::NarrowWorldPositionChecked(FVector(Center.X + Extent.X, Center.Y - Extent.Y + FaceWidth, Center.Z - Extent.Z + FaceWidth))));
+		Verts.Add(FDynamicMeshVertex(UE::LWC::NarrowWorldPositionChecked(FVector(Center.X + Extent.X, Center.Y + Extent.Y - FaceWidth, Center.Z - Extent.Z + FaceWidth))));
+		Verts.Add(FDynamicMeshVertex(UE::LWC::NarrowWorldPositionChecked(FVector(Center.X + Extent.X, Center.Y - Extent.Y + FaceWidth, Center.Z + Extent.Z - FaceWidth))));
+		Verts.Add(FDynamicMeshVertex(UE::LWC::NarrowWorldPositionChecked(FVector(Center.X + Extent.X, Center.Y + Extent.Y - FaceWidth, Center.Z + Extent.Z - FaceWidth))));
+
+		Verts.Add(FDynamicMeshVertex(UE::LWC::NarrowWorldPositionChecked(FVector(Center.X - Extent.X + FaceWidth, Center.Y - Extent.Y + FaceWidth, Center.Z + Extent.Z))));
+		Verts.Add(FDynamicMeshVertex(UE::LWC::NarrowWorldPositionChecked(FVector(Center.X + Extent.X - FaceWidth, Center.Y - Extent.Y + FaceWidth, Center.Z + Extent.Z))));
+		Verts.Add(FDynamicMeshVertex(UE::LWC::NarrowWorldPositionChecked(FVector(Center.X - Extent.X + FaceWidth, Center.Y + Extent.Y - FaceWidth, Center.Z + Extent.Z))));
+		Verts.Add(FDynamicMeshVertex(UE::LWC::NarrowWorldPositionChecked(FVector(Center.X + Extent.X - FaceWidth, Center.Y + Extent.Y - FaceWidth, Center.Z + Extent.Z))));
+
+		Verts.Add(FDynamicMeshVertex(UE::LWC::NarrowWorldPositionChecked(FVector(Center.X - Extent.X + FaceWidth, Center.Y + Extent.Y, Center.Z - Extent.Z + FaceWidth))));
+		Verts.Add(FDynamicMeshVertex(UE::LWC::NarrowWorldPositionChecked(FVector(Center.X + Extent.X - FaceWidth, Center.Y + Extent.Y, Center.Z - Extent.Z + FaceWidth))));
+		Verts.Add(FDynamicMeshVertex(UE::LWC::NarrowWorldPositionChecked(FVector(Center.X - Extent.X + FaceWidth, Center.Y + Extent.Y, Center.Z + Extent.Z - FaceWidth))));
+		Verts.Add(FDynamicMeshVertex(UE::LWC::NarrowWorldPositionChecked(FVector(Center.X + Extent.X - FaceWidth, Center.Y + Extent.Y, Center.Z + Extent.Z - FaceWidth))));
+
+		Verts.Add(FDynamicMeshVertex(UE::LWC::NarrowWorldPositionChecked(FVector(Center.X - Extent.X, Center.Y - Extent.Y + FaceWidth, Center.Z - Extent.Z + FaceWidth))));
+		Verts.Add(FDynamicMeshVertex(UE::LWC::NarrowWorldPositionChecked(FVector(Center.X - Extent.X, Center.Y + Extent.Y - FaceWidth, Center.Z - Extent.Z + FaceWidth))));
+		Verts.Add(FDynamicMeshVertex(UE::LWC::NarrowWorldPositionChecked(FVector(Center.X - Extent.X, Center.Y - Extent.Y + FaceWidth, Center.Z + Extent.Z - FaceWidth))));
+		Verts.Add(FDynamicMeshVertex(UE::LWC::NarrowWorldPositionChecked(FVector(Center.X - Extent.X, Center.Y + Extent.Y - FaceWidth, Center.Z + Extent.Z - FaceWidth))));
+
+		Verts.Add(FDynamicMeshVertex(UE::LWC::NarrowWorldPositionChecked(FVector(Center.X - Extent.X + FaceWidth, Center.Y - Extent.Y + FaceWidth, Center.Z - Extent.Z))));
+		Verts.Add(FDynamicMeshVertex(UE::LWC::NarrowWorldPositionChecked(FVector(Center.X + Extent.X - FaceWidth, Center.Y - Extent.Y + FaceWidth, Center.Z - Extent.Z))));
+		Verts.Add(FDynamicMeshVertex(UE::LWC::NarrowWorldPositionChecked(FVector(Center.X - Extent.X + FaceWidth, Center.Y + Extent.Y - FaceWidth, Center.Z - Extent.Z))));
+		Verts.Add(FDynamicMeshVertex(UE::LWC::NarrowWorldPositionChecked(FVector(Center.X + Extent.X - FaceWidth, Center.Y + Extent.Y - FaceWidth, Center.Z - Extent.Z))));
 
 		StoreQuadIndices(0, 1, 9, 8, FirstVertIdx, Indices);
 		StoreQuadIndices(1, 5, 11, 9, FirstVertIdx, Indices);
diff --git a/Engine/Source/Runtime/AIModule/Private/Navigation/NavLocalGridManager.cpp b/Engine/Source/Runtime/AIModule/Private/Navigation/NavLocalGridManager.cpp
--- a/Engine/Source/Runtime/AIModule/Private/Navigation/NavLocalGridManager.cpp
+++ b/Engine/Source/Runtime/AIModule/Private/Navigation/NavLocalGridManager.cpp
@@ -367,7 +367,7 @@ int32 UNavLocalGridManager::AddLocalNavigationGridForBox(UObject* WorldContextOb
 	if (GridManager)
 	{
 		FNavLocalGridData GridData(Location, FVector2D(Extent.X + UNavLocalGridManager::GridCellSize * Radius2D, Extent.Y + UNavLocalGridManager::GridCellSize * Radius2D));
-		GridData.SetHeight(Height + Extent.Z);
+		GridData.SetHeight(FloatCastChecked<float>(Height + Extent.Z, UE::LWC::DefaultFloatPrecision));
 		GridData.MarkBoxObstacle(Location, Extent, Rotation.Quaternion());
 
 		GridId = GridManager->AddGridData(GridData, bRebuildGrids);
diff --git a/Engine/Source/Runtime/AIModule/Private/Navigation/PathFollowingComponent.cpp b/Engine/Source/Runtime/AIModule/Private/Navigation/PathFollowingComponent.cpp
--- a/Engine/Source/Runtime/AIModule/Private/Navigation/PathFollowingComponent.cpp
+++ b/Engine/Source/Runtime/AIModule/Private/Navigation/PathFollowingComponent.cpp
@@ -390,7 +390,7 @@ FAIRequestID UPathFollowingComponent::RequestMove(const FAIMoveRequest& RequestD
 			MetaNavPath->Initialize(CurrentLocation);
 		}
 
-		PathTimeWhenPaused = 0.0f;
+		PathTimeWhenPaused = 0.;
 		OnPathUpdated();
 
 		// make sure that OnPathUpdated didn't change current move request
@@ -474,7 +474,7 @@ void UPathFollowingComponent::PauseMove(FAIRequestID RequestID, EPathFollowingVe
 		}
 
 		LocationWhenPaused = MovementComp ? MovementComp->GetActorFeetLocation() : FVector::ZeroVector;
-		PathTimeWhenPaused = Path.IsValid() ? Path->GetTimeStamp() : 0.0f;
+		PathTimeWhenPaused = Path.IsValid() ? Path->GetTimeStamp() : 0.;
 		Status = EPathFollowingStatus::Paused;
 
 		UpdateMoveFocus();
@@ -1309,7 +1309,7 @@ void UPathFollowingComponent::DebugReachTest(float& CurrentDot, float& CurrentDi
 
 	const FVector ToGoal = (GoalLocation - AgentLocation);
 	const FVector CurrentDirection = GetCurrentDirection();
-	CurrentDot = FVector::DotProduct(ToGoal.GetSafeNormal(), CurrentDirection);
+	CurrentDot = FloatCastChecked<float>(FVector::DotProduct(ToGoal.GetSafeNormal(), CurrentDirection), /* Precision */ 1./128.);
 	bDotFailed = (CurrentDot < 0.0f) ? 1 : 0;
 
 	// get cylinder of moving agent
@@ -1318,11 +1318,11 @@ void UPathFollowingComponent::DebugReachTest(float& CurrentDot, float& CurrentDi
 	AActor* MovingAgent = MovementComp->GetOwner();
 	MovingAgent->GetSimpleCollisionCylinder(AgentRadius, AgentHalfHeight);
 
-	CurrentDistance = ToGoal.Size2D();
+	CurrentDistance = FloatCastChecked<float>(ToGoal.Size2D(), UE::LWC::DefaultFloatPrecision);
 	const float UseRadius = FMath::Max(RadiusThreshold, GoalRadius + (AgentRadius * AgentRadiusPct));
 	bDistanceFailed = (CurrentDistance > UseRadius) ? 1 : 0;
 
-	CurrentHeight = FMath::Abs(ToGoal.Z);
+	CurrentHeight = FloatCastChecked<float>(FMath::Abs(ToGoal.Z), UE::LWC::DefaultFloatPrecision);
 	const float UseHeight = GoalHalfHeight + (AgentHalfHeight * MinAgentHalfHeightPct);
 	bHeightFailed = (CurrentHeight > UseHeight) ? 1 : 0;
 }
@@ -1456,7 +1456,7 @@ bool UPathFollowingComponent::IsBlocked() const
 
 bool UPathFollowingComponent::UpdateBlockDetection()
 {
-	const float GameTime = GetWorld()->GetTimeSeconds();
+	const double GameTime = GetWorld()->GetTimeSeconds();
 	if (bUseBlockDetection &&
 		MovementComp && 
 		GameTime > (LastSampleTime + BlockDetectionInterval) &&
diff --git a/Engine/Source/Runtime/AIModule/Public/HierarchicalHashGrid2D.h b/Engine/Source/Runtime/AIModule/Public/HierarchicalHashGrid2D.h
--- a/Engine/Source/Runtime/AIModule/Public/HierarchicalHashGrid2D.h
+++ b/Engine/Source/Runtime/AIModule/Public/HierarchicalHashGrid2D.h
@@ -21,8 +21,13 @@
  * - FloorToInt ends up calling floorf, which will be a function call.
  * 	 int FloorToInt(float a) { return (int)a + ((int)a > a ? -1 : 0); }
  * 	 auto vectorizes nicely on clang, but not on VC.
+ * 
  * - Add helper function to allow to tweak the cell size to reset the grid when spill list gets too large
  */
+
+// LWC_TODO_AI Note we are using int32 here for X and Y which does mean that we could overflow the limit of an int for LWCoords
+// unless fairly large grid cell sizes are used.As WORLD_MAX is currently in flux until we have a better idea of what we are
+// going to be able to support its probably not worth investing time in this potential issue right now.
 template <int32 InNumLevels = 3, int32 InLevelRatio = 4, typename InItemIDType = uint32>
 class THierarchicalHashGrid2D
 {
@@ -460,10 +465,10 @@ public:
 	FCellRect CalcQueryBounds(const FBox& Bounds, const int32 Level) const
 	{
 		FCellRect Result;
-		Result.MinX = FMath::FloorToInt(Bounds.Min.X * InvCellSize[Level] - 0.5f);
-		Result.MinY = FMath::FloorToInt(Bounds.Min.Y * InvCellSize[Level] - 0.5f);
-		Result.MaxX = FMath::FloorToInt(Bounds.Max.X * InvCellSize[Level] + 0.5f);
-		Result.MaxY = FMath::FloorToInt(Bounds.Max.Y * InvCellSize[Level] + 0.5f);
+		Result.MinX = IntCastChecked<int32>(FMath::FloorToInt(Bounds.Min.X * InvCellSize[Level] - 0.5f));
+		Result.MinY = IntCastChecked<int32>(FMath::FloorToInt(Bounds.Min.Y * InvCellSize[Level] - 0.5f));
+		Result.MaxX = IntCastChecked<int32>(FMath::FloorToInt(Bounds.Max.X * InvCellSize[Level] + 0.5f));
+		Result.MaxY = IntCastChecked<int32>(FMath::FloorToInt(Bounds.Max.Y * InvCellSize[Level] + 0.5f));
 		return Result;
 	}
 
@@ -501,15 +506,14 @@ public:
 		FCellLocation Location(0, 0, 0);
 
 		const FVector Center = Bounds.GetCenter();
-		const float Diameter = FMath::Max(Bounds.Max.X - Bounds.Min.X, Bounds.Max.Y - Bounds.Min.Y);
-
+		const FVector::FReal Diameter = FMath::Max(Bounds.Max.X - Bounds.Min.X, Bounds.Max.Y - Bounds.Min.Y);
 		for (Location.Level = 0; Location.Level < NumLevels; Location.Level++)
 		{
-			const int32 DiameterCells = FMath::CeilToInt(Diameter * InvCellSize[Location.Level]);
+			const int32 DiameterCells = IntCastChecked<int32>(FMath::CeilToInt(Diameter * InvCellSize[Location.Level]));
 			if (DiameterCells <= 1)
 			{
-				Location.X = FMath::FloorToInt(Center.X * InvCellSize[Location.Level]);
-				Location.Y = FMath::FloorToInt(Center.Y * InvCellSize[Location.Level]);
+				Location.X = IntCastChecked<int32>(FMath::FloorToInt(Center.X * InvCellSize[Location.Level]));
+				Location.Y = IntCastChecked<int32>(FMath::FloorToInt(Center.Y * InvCellSize[Location.Level]));
 				break;
 			}
 		}
diff --git a/Engine/Source/Runtime/AIModule/Public/IndexedHandle.h b/Engine/Source/Runtime/AIModule/Public/IndexedHandle.h
--- a/Engine/Source/Runtime/AIModule/Public/IndexedHandle.h
+++ b/Engine/Source/Runtime/AIModule/Public/IndexedHandle.h
@@ -140,7 +140,7 @@ struct AIMODULE_API FCompactIndexedHandleBase
 	}
 
 	/** @return INDEX_NONE if invalid or >=0 for a potentially valid handle */
-	int32 GetIndex() const { return Index; }
+	int16 GetIndex() const { return Index; }
 
 	/** @return true if this handle is valid, this doesn't check that it's in the correct range for the underlying resource it's providing a handle for */
 	bool IsValid() const { return Index >= 0; }
diff --git a/Engine/Source/Runtime/Engine/Classes/AI/Navigation/AvoidanceManager.h b/Engine/Source/Runtime/Engine/Classes/AI/Navigation/AvoidanceManager.h
--- a/Engine/Source/Runtime/Engine/Classes/AI/Navigation/AvoidanceManager.h
+++ b/Engine/Source/Runtime/Engine/Classes/AI/Navigation/AvoidanceManager.h
@@ -45,7 +45,7 @@ struct FNavAvoidanceData
 	float Weight;
 
 	/** Weight is treated as a hard 1.0 while this is active. This is set by code. */
-	float OverrideWeightTime;
+	double OverrideWeightTime;
 
 	/** Group data */
 	int32 GroupMask;
diff --git a/Engine/Source/Runtime/Engine/Private/AI/Navigation/AvoidanceManager.cpp b/Engine/Source/Runtime/Engine/Private/AI/Navigation/AvoidanceManager.cpp
--- a/Engine/Source/Runtime/Engine/Private/AI/Navigation/AvoidanceManager.cpp
+++ b/Engine/Source/Runtime/Engine/Private/AI/Navigation/AvoidanceManager.cpp
@@ -260,9 +260,9 @@ void UAvoidanceManager::UpdateRVO_Internal(int32 inAvoidanceUID, const FNavAvoid
 FVector AvoidCones(TArray<FVelocityAvoidanceCone>& AllCones, const FVector& BasePosition, const FVector& DesiredPosition, const int NumConesToTest)
 {
 	FVector CurrentPosition = DesiredPosition;
-	float DistanceInsidePlane_Current[2];
-	float DistanceInsidePlane_Base[2];
-	float Weighting[2];
+	FVector::FReal DistanceInsidePlane_Current[2];
+	FVector::FReal DistanceInsidePlane_Base[2];
+	FVector::FReal Weighting[2];
 	int ConePlaneIndex;
 
 	//AllCones is non-const so that it can be reordered, but nothing should be added or removed from it.
@@ -333,21 +333,21 @@ static bool AvoidsNavEdges(const FVector& OrgLocation, const FVector& TestVeloci
 		const FVector2D Seg0ToSeg1(NavEdges[Idx].P1 - NavEdges[Idx].P0);
 		const FVector2D OrgToNewPos(TestVelocity);
 		const FVector2D OrgToSeg0(NavEdges[Idx].P0 - OrgLocation);
-		const float CrossD = FVector2D::CrossProduct(Seg0ToSeg1, OrgToNewPos);
+		const FVector2D::FReal CrossD = FVector2D::CrossProduct(Seg0ToSeg1, OrgToNewPos);
 		if (FMath::Abs(CrossD) < UE_KINDA_SMALL_NUMBER)
 		{
 			continue;
 		}
 
-		const float CrossS = FVector2D::CrossProduct(OrgToNewPos, OrgToSeg0) / CrossD;
-		const float CrossT = FVector2D::CrossProduct(Seg0ToSeg1, OrgToSeg0) / CrossD;
+		const FVector2D::FReal CrossS = FVector2D::CrossProduct(OrgToNewPos, OrgToSeg0) / CrossD;
+		const FVector2D::FReal CrossT = FVector2D::CrossProduct(Seg0ToSeg1, OrgToSeg0) / CrossD;
 		if (CrossS < 0.0f || CrossS > 1.0f || CrossT < 0.0f || CrossT > 1.0f)
 		{
 			continue;
 		}
 
 		const FVector CrossPt = FMath::Lerp(NavEdges[Idx].P0, NavEdges[Idx].P1, CrossS);
-		const float ZDiff = FMath::Abs(OrgLocation.Z - CrossPt.Z);
+		const FVector2D::FReal ZDiff = FMath::Abs(OrgLocation.Z - CrossPt.Z);
 		if (ZDiff > MaxZDiff)
 		{
 			continue;
@@ -374,8 +374,8 @@ FVector UAvoidanceManager::GetAvoidanceVelocity_Internal(const FNavAvoidanceData
 	}
 
 	FVector ReturnVelocity = inAvoidanceData.Velocity * DeltaTime;
-	float MaxSpeed = ReturnVelocity.Size2D();
-	float CurrentTime;
+	FVector::FReal MaxSpeed = ReturnVelocity.Size2D();
+	double CurrentTime;
 
 	UWorld* MyWorld = Cast<UWorld>(GetOuter());
 	if (MyWorld)
@@ -532,13 +532,13 @@ FVector UAvoidanceManager::GetAvoidanceVelocity_Internal(const FNavAvoidanceData
 	//Find a good velocity that isn't inside a cone.
 	if (AllCones.Num())
 	{
-		float AngleCurrent;
-		float AngleF = ReturnVelocity.HeadingAngle();
-		float BestScore = 0.0f;
-		float BestScorePotential;
+		FVector::FReal AngleCurrent;
+		FVector::FReal AngleF = ReturnVelocity.HeadingAngle();
+		FVector::FReal BestScore = 0.0f;
+		FVector::FReal BestScorePotential;
 		FVector BestVelocity = FVector::ZeroVector;		//Worst case is we just stand completely still. Should we also allow backing up? Should we test standing still?
 		const int AngleCount = 4;		//Every angle will be tested both right and left.
-		float AngleOffset[AngleCount] = {FMath::DegreesToRadians<float>(23.0f), FMath::DegreesToRadians<float>(40.0f), FMath::DegreesToRadians<float>(55.0f), FMath::DegreesToRadians<float>(85.0f)};
+		FVector::FReal AngleOffset[AngleCount] = {FMath::DegreesToRadians<float>(23.0f), FMath::DegreesToRadians<float>(40.0f), FMath::DegreesToRadians<float>(55.0f), FMath::DegreesToRadians<float>(85.0f)};
 		FVector AngleVector[AngleCount<<1];
 
 		//Determine check angles
@@ -564,7 +564,7 @@ FVector UAvoidanceManager::GetAvoidanceVelocity_Internal(const FNavAvoidanceData
 				if (bAvoidsNavEdges)
 				{
 					FVector CandidateVelocity = AvoidCones(AllCones, FVector::ZeroVector, VelSpacePoint, AllCones.Num());
-					float CandidateScore = (CandidateVelocity|ReturnVelocity) * (CandidateVelocity|CandidateVelocity);
+					FVector::FReal CandidateScore = (CandidateVelocity|ReturnVelocity) * (CandidateVelocity|CandidateVelocity);
 
 					//Vectors are rated by their length and their overall forward movement.
 					if (CandidateScore > BestScore)
diff --git a/Engine/Source/Runtime/Engine/Private/AI/Navigation/NavigationModifier.cpp b/Engine/Source/Runtime/Engine/Private/AI/Navigation/NavigationModifier.cpp
--- a/Engine/Source/Runtime/Engine/Private/AI/Navigation/NavigationModifier.cpp
+++ b/Engine/Source/Runtime/Engine/Private/AI/Navigation/NavigationModifier.cpp
@@ -294,16 +294,16 @@ FAreaNavModifier::FAreaNavModifier(float Radius, float Height, const FTransform&
 	Init(InAreaClass);
 	
 	FVector Scale3D = LocalToWorld.GetScale3D().GetAbs();
-	Radius *= FMath::Max(Scale3D.X, Scale3D.Y);
-	Height *= Scale3D.Z;
+	const FVector::FReal RadiusScaled = Radius * FMath::Max(Scale3D.X, Scale3D.Y);
+	const FVector::FReal HeightScaled = Height * Scale3D.Z;
 
 	Points.SetNumUninitialized(2);
 	Points[0] = LocalToWorld.GetLocation();
-	Points[1].X = Radius;
-	Points[1].Z = Height;
+	Points[1].X = RadiusScaled;
+	Points[1].Z = HeightScaled;
 	ShapeType = ENavigationShapeType::Cylinder;
 
-	Bounds = FBox::BuildAABB(LocalToWorld.GetLocation(), FVector(Radius, Radius, Height));
+	Bounds = FBox::BuildAABB(LocalToWorld.GetLocation(), FVector(RadiusScaled, RadiusScaled, HeightScaled));
 }
 
 FAreaNavModifier::FAreaNavModifier(const FVector& Extent, const FTransform& LocalToWorld, const TSubclassOf<UNavAreaBase> InAreaClass)
@@ -381,8 +381,8 @@ void FAreaNavModifier::GetCylinder(FCylinderNavAreaData& Data) const
 {
 	check(Points.Num() == 2 && ShapeType == ENavigationShapeType::Cylinder);
 	Data.Origin = Points[0];
-	Data.Radius = Points[1].X;
-	Data.Height = Points[1].Z;
+	Data.Radius = FloatCastChecked<float>(Points[1].X, UE::LWC::DefaultFloatPrecision);
+	Data.Height = FloatCastChecked<float>(Points[1].Z, UE::LWC::DefaultFloatPrecision);
 }
 
 void FAreaNavModifier::GetBox(FBoxNavAreaData& Data) const
@@ -464,7 +464,7 @@ void FAreaNavModifier::SetApplyMode(ENavigationAreaMode::Type InApplyMode)
 	bIsLowAreaModifier = (InApplyMode == ENavigationAreaMode::ApplyInLowPass) || (InApplyMode == ENavigationAreaMode::ReplaceInLowPass);
 }
 
-bool IsAngleMatching(float Angle)
+bool IsAngleMatching(FRotator::FReal Angle)
 {
 	const float AngleThreshold = 1.0f; // degrees
 	return (Angle < AngleThreshold) || ((90.0f - Angle) < AngleThreshold);
@@ -484,9 +484,9 @@ void FAreaNavModifier::SetBox(const FBox& Box, const FTransform& LocalToWorld)
 
 	// check if it can be used as AABB
 	const FRotator Rotation = LocalToWorld.GetRotation().Rotator();
-	const float PitchMod = FMath::Fmod(FMath::Abs(Rotation.Pitch), 90.0f);
-	const float YawMod = FMath::Fmod(FMath::Abs(Rotation.Yaw), 90.0f);
-	const float RollMod = FMath::Fmod(FMath::Abs(Rotation.Roll), 90.0f);
+	const FRotator::FReal PitchMod = FMath::Fmod(FMath::Abs(Rotation.Pitch), 90.0f);
+	const FRotator::FReal YawMod = FMath::Fmod(FMath::Abs(Rotation.Yaw), 90.0f);
+	const FRotator::FReal RollMod = FMath::Fmod(FMath::Abs(Rotation.Roll), 90.0f);
 	if (IsAngleMatching(PitchMod) && IsAngleMatching(YawMod) && IsAngleMatching(RollMod))
 	{
 		Bounds = FBox(ForceInit);
@@ -971,7 +971,7 @@ void FCompositeNavModifier::CreateAreaModifiers(const FCollisionShape& Collision
 
 uint32 FCompositeNavModifier::GetAllocatedSize() const
 {
-	uint32 MemUsed = Areas.GetAllocatedSize() + SimpleLinks.GetAllocatedSize() + CustomLinks.GetAllocatedSize();
+	SIZE_T MemUsed = Areas.GetAllocatedSize() + SimpleLinks.GetAllocatedSize() + CustomLinks.GetAllocatedSize();
 
 	const FSimpleLinkNavModifier* SimpleLink = SimpleLinks.GetData();
 	for (int32 Index = 0; Index < SimpleLinks.Num(); ++Index, ++SimpleLink)
@@ -979,7 +979,7 @@ uint32 FCompositeNavModifier::GetAllocatedSize() const
 		MemUsed += SimpleLink->Links.GetAllocatedSize();
 	}
 
-	return MemUsed;
+	return IntCastChecked<uint32>(MemUsed);
 }
 
 bool FCompositeNavModifier::HasPerInstanceTransforms() const
diff --git a/Engine/Source/Runtime/Engine/Private/AI/Navigation/NavigationTypes.cpp b/Engine/Source/Runtime/Engine/Private/AI/Navigation/NavigationTypes.cpp
--- a/Engine/Source/Runtime/Engine/Private/AI/Navigation/NavigationTypes.cpp
+++ b/Engine/Source/Runtime/Engine/Private/AI/Navigation/NavigationTypes.cpp
@@ -328,7 +328,7 @@ FNavAgentProperties::FNavAgentProperties(const FNavAgentProperties& Other)
 void FNavAgentProperties::UpdateWithCollisionComponent(UShapeComponent* CollisionComponent)
 {
 	check(CollisionComponent != NULL);
-	AgentRadius = CollisionComponent->Bounds.SphereRadius;
+	AgentRadius = FloatCastChecked<float>(CollisionComponent->Bounds.SphereRadius, UE::LWC::DefaultFloatPrecision);
 }
 
 bool FNavAgentProperties::IsNavDataMatching(const FNavAgentProperties& Other) const
diff --git a/Engine/Source/Runtime/NavigationSystem/Private/NavCollision.cpp b/Engine/Source/Runtime/NavigationSystem/Private/NavCollision.cpp
--- a/Engine/Source/Runtime/NavigationSystem/Private/NavCollision.cpp
+++ b/Engine/Source/Runtime/NavigationSystem/Private/NavCollision.cpp
@@ -264,7 +264,12 @@ void UNavCollision::GatherCollision()
 	{
 		const FNavCollisionBox& BoxInfo = BoxCollision[Idx];
 
-		FKBoxElem BoxElem(BoxInfo.Extent.X * 2.0f, BoxInfo.Extent.Y * 2.0f, BoxInfo.Extent.Z * 2.0f);
+		const float X = FloatCastChecked<float>(BoxInfo.Extent.X * 2.0f, UE::LWC::DefaultFloatPrecision);
+		const float Y = FloatCastChecked<float>(BoxInfo.Extent.Y * 2.0f, UE::LWC::DefaultFloatPrecision);
+		const float Z = FloatCastChecked<float>(BoxInfo.Extent.Z * 2.0f, UE::LWC::DefaultFloatPrecision);
+
+		FKBoxElem BoxElem(X, Y, Z);
+
 		BoxElem.SetTransform(FTransform(BoxInfo.Offset));
 
 		SimpleGeom.BoxElems.Add(BoxElem);
diff --git a/Engine/Source/Runtime/NavigationSystem/Private/NavFilters/NavigationQueryFilter.cpp b/Engine/Source/Runtime/NavigationSystem/Private/NavFilters/NavigationQueryFilter.cpp
--- a/Engine/Source/Runtime/NavigationSystem/Private/NavFilters/NavigationQueryFilter.cpp
+++ b/Engine/Source/Runtime/NavigationSystem/Private/NavFilters/NavigationQueryFilter.cpp
@@ -70,18 +70,18 @@ void UNavigationQueryFilter::InitializeFilter(const ANavigationData& NavData, co
 
 		if (AreaData.bIsExcluded)
 		{
-			Filter.SetExcludedArea(AreaId);
+			Filter.SetExcludedArea(IntCastChecked<uint8>(AreaId));
 		}
 		else
 		{
 			if (AreaData.bOverrideTravelCost)
 			{
-				Filter.SetAreaCost(AreaId, FMath::Max(1.0f, AreaData.TravelCostOverride));
+				Filter.SetAreaCost(IntCastChecked<uint8>(AreaId), FMath::Max(1.0f, AreaData.TravelCostOverride));
 			}
 
 			if (AreaData.bOverrideEnteringCost)
 			{
-				Filter.SetFixedAreaEnteringCost(AreaId, FMath::Max(0.0f, AreaData.EnteringCostOverride));
+				Filter.SetFixedAreaEnteringCost(IntCastChecked<uint8>(AreaId), FMath::Max(0.0f, AreaData.EnteringCostOverride));
 			}
 		}
 	}
diff --git a/Engine/Source/Runtime/NavigationSystem/Private/NavLinkCustomComponent.cpp b/Engine/Source/Runtime/NavigationSystem/Private/NavLinkCustomComponent.cpp
--- a/Engine/Source/Runtime/NavigationSystem/Private/NavLinkCustomComponent.cpp
+++ b/Engine/Source/Runtime/NavigationSystem/Private/NavLinkCustomComponent.cpp
@@ -323,8 +323,8 @@ void UNavLinkCustomComponent::CollectNearbyAgents(TArray<UObject*>& NotifyList)
 
 	const FVector LocationL = GetStartPoint();
 	const FVector LocationR = GetEndPoint();
-	const float LinkDistSq = (LocationL - LocationR).SizeSquared();
-	const float DistThresholdSq = FMath::Square(BroadcastRadius * 0.25f);
+	const FVector::FReal LinkDistSq = (LocationL - LocationR).SizeSquared();
+	const FVector::FReal DistThresholdSq = FMath::Square(BroadcastRadius * 0.25);
 	if (LinkDistSq > DistThresholdSq)
 	{
 		GetWorld()->OverlapMultiByChannel(OverlapsL, LocationL, FQuat::Identity, BroadcastChannel, FCollisionShape::MakeSphere(BroadcastRadius), Params);
@@ -332,7 +332,7 @@ void UNavLinkCustomComponent::CollectNearbyAgents(TArray<UObject*>& NotifyList)
 	}
 	else
 	{
-		const FVector MidPoint = (LocationL + LocationR) * 0.5f;
+		const FVector MidPoint = (LocationL + LocationR) * 0.5;
 		GetWorld()->OverlapMultiByChannel(OverlapsL, MidPoint, FQuat::Identity, BroadcastChannel, FCollisionShape::MakeSphere(BroadcastRadius), Params);
 	}
 
diff --git a/Engine/Source/Runtime/NavigationSystem/Private/NavLinkRenderingComponent.cpp b/Engine/Source/Runtime/NavigationSystem/Private/NavLinkRenderingComponent.cpp
--- a/Engine/Source/Runtime/NavigationSystem/Private/NavLinkRenderingComponent.cpp
+++ b/Engine/Source/Runtime/NavigationSystem/Private/NavLinkRenderingComponent.cpp
@@ -240,7 +240,9 @@ void FNavLinkRenderingProxy::GetLinkMeshes(const TArray<FNavLinkDrawing>& OffMes
 		{
 			continue;
 		}
-		const uint32 Segments = FPlatformMath::Max<uint32>(LinkArcHeight*(Link.Right - Link.Left).Size() / 10, 8);
+		const FVector::FReal RealSegments = FMath::Max(LinkArcHeight * (Link.Right - Link.Left).Size() / 10., 8.);
+		check(RealSegments >= 0 && RealSegments <= (FVector::FReal)TNumericLimits<uint32>::Max());
+		const uint32 Segments = static_cast<uint32>(RealSegments);
 		DrawArc(PDI, Link.Left, Link.Right, LinkArcHeight, Segments, Link.Color, SDPG_World, 3.5f);
 		const FVector VOffset(0,0,FVector::Dist(Link.Left, Link.Right)*1.333f);
 
@@ -283,8 +285,12 @@ void FNavLinkRenderingProxy::GetLinkMeshes(const TArray<FNavLinkDrawing>& OffMes
 		{
 			continue;
 		}
-		const uint32 SegmentsStart = FPlatformMath::Max<uint32>(SegmentArcHeight*(Link.RightStart - Link.LeftStart).Size() / 10, 8);
-		const uint32 SegmentsEnd = FPlatformMath::Max<uint32>(SegmentArcHeight*(Link.RightEnd-Link.LeftEnd).Size()/10, 8);
+		const FVector::FReal RealSegmentsStart = FMath::Max(SegmentArcHeight * (Link.RightStart - Link.LeftStart).Size() / 10., 8.);
+		const FVector::FReal RealSegmentsEnd= FMath::Max(SegmentArcHeight * (Link.RightEnd - Link.LeftEnd).Size() / 10., 8.);
+		check(RealSegmentsStart >= 0 && RealSegmentsStart <= (FVector::FReal)TNumericLimits<uint32>::Max());
+		check(RealSegmentsEnd >= 0 && RealSegmentsEnd <= (FVector::FReal)TNumericLimits<uint32>::Max());
+		const uint32 SegmentsStart = static_cast<uint32>(RealSegmentsStart);
+		const uint32 SegmentsEnd = static_cast<uint32>(RealSegmentsEnd);
 		DrawArc(PDI, Link.LeftStart, Link.RightStart, SegmentArcHeight, SegmentsStart, Link.Color, SDPG_World, 3.5f);
 		DrawArc(PDI, Link.LeftEnd, Link.RightEnd, SegmentArcHeight, SegmentsEnd, Link.Color, SDPG_World, 3.5f);
 		const FVector VOffset(0,0,FVector::Dist(Link.LeftStart, Link.RightStart)*1.333f);
@@ -348,7 +354,9 @@ void FNavLinkRenderingProxy::DrawLinks(FPrimitiveDrawInterface* PDI, TArray<FNav
 			continue;
 		}
 
-		const uint32 Segments = FPlatformMath::Max<uint32>(LinkArcHeight*(Link.Right-Link.Left).Size()/10, 8);
+		const FVector::FReal RealSegments = FMath::Max(LinkArcHeight * (Link.Right - Link.Left).Size() / 10., 8.);
+		check(RealSegments >= 0 && RealSegments <= (FVector::FReal)TNumericLimits<uint32>::Max());
+		const uint32 Segments = static_cast<uint32>(RealSegments);
 		DrawArc(PDI, Link.Left, Link.Right, LinkArcHeight, Segments, Link.Color, SDPG_World, 3.5f);
 		const FVector VOffset(0,0,FVector::Dist(Link.Left, Link.Right)*1.333f);
 
@@ -392,8 +400,12 @@ void FNavLinkRenderingProxy::DrawLinks(FPrimitiveDrawInterface* PDI, TArray<FNav
 			continue;
 		}
 
-		const uint32 SegmentsStart = FPlatformMath::Max<uint32>(SegmentArcHeight*(Link.RightStart-Link.LeftStart).Size()/10, 8);
-		const uint32 SegmentsEnd = FPlatformMath::Max<uint32>(SegmentArcHeight*(Link.RightEnd-Link.LeftEnd).Size()/10, 8);
+		const FVector::FReal RealSegmentsStart = FMath::Max(SegmentArcHeight * (Link.RightStart - Link.LeftStart).Size() / 10., 8.);
+		const FVector::FReal RealSegmentsEnd = FMath::Max(SegmentArcHeight * (Link.RightEnd - Link.LeftEnd).Size() / 10., 8.);
+		check(RealSegmentsStart >= 0 && RealSegmentsStart <= (FVector::FReal)TNumericLimits<uint32>::Max());
+		check(RealSegmentsEnd >= 0 && RealSegmentsEnd <= (FVector::FReal)TNumericLimits<uint32>::Max());
+		const uint32 SegmentsStart = static_cast<uint32>(RealSegmentsStart);
+		const uint32 SegmentsEnd = static_cast<uint32>(RealSegmentsEnd);
 		DrawArc(PDI, Link.LeftStart, Link.RightStart, SegmentArcHeight, SegmentsStart, Link.Color, SDPG_World, 3.5f);
 		DrawArc(PDI, Link.LeftEnd, Link.RightEnd, SegmentArcHeight, SegmentsEnd, Link.Color, SDPG_World, 3.5f);
 		const FVector VOffset(0,0,FVector::Dist(Link.LeftStart, Link.RightStart)*1.333f);
@@ -457,6 +469,6 @@ uint32 FNavLinkRenderingProxy::GetMemoryFootprint( void ) const
 
 uint32 FNavLinkRenderingProxy::GetAllocatedSize( void ) const 
 { 
-	return(FPrimitiveSceneProxy::GetAllocatedSize() + OffMeshPointLinks.GetAllocatedSize() + OffMeshSegmentLinks.GetAllocatedSize());
+	return IntCastChecked<uint32>(FPrimitiveSceneProxy::GetAllocatedSize() + OffMeshPointLinks.GetAllocatedSize() + OffMeshSegmentLinks.GetAllocatedSize());
 }
 
diff --git a/Engine/Source/Runtime/NavigationSystem/Private/NavMesh/NavMeshPath.cpp b/Engine/Source/Runtime/NavigationSystem/Private/NavMesh/NavMeshPath.cpp
--- a/Engine/Source/Runtime/NavigationSystem/Private/NavMesh/NavMeshPath.cpp
+++ b/Engine/Source/Runtime/NavigationSystem/Private/NavMesh/NavMeshPath.cpp
@@ -825,7 +825,7 @@ void FNavMeshPath::DescribeSelfToVisLog(FVisualLogEntry* Snapshot) const
 	TArray<FVector> Verts;
 	for (int32 Idx = 0; Idx < PathCorridor.Num(); Idx++)
 	{
-		const uint8 AreaID = NavMesh->GetPolyAreaID(PathCorridor[Idx]);
+		const int32 AreaID = IntCastChecked<int32>(NavMesh->GetPolyAreaID(PathCorridor[Idx]));
 		const UClass* AreaClass = NavMesh->GetAreaClass(AreaID);
 		
 		Verts.Reset();
diff --git a/Engine/Source/Runtime/NavigationSystem/Private/NavMesh/NavMeshRenderingComponent.cpp b/Engine/Source/Runtime/NavigationSystem/Private/NavMesh/NavMeshRenderingComponent.cpp
--- a/Engine/Source/Runtime/NavigationSystem/Private/NavMesh/NavMeshRenderingComponent.cpp
+++ b/Engine/Source/Runtime/NavigationSystem/Private/NavMesh/NavMeshRenderingComponent.cpp
@@ -105,14 +105,14 @@ namespace FNavMeshRenderingHelpers
 		return true;
 	}
 
-	bool LineInCorrectDistance(const FVector& Start, const FVector& End, const FSceneView* View, float CorrectDistance = -1)
+	bool LineInCorrectDistance(const FVector& Start, const FVector& End, const FSceneView* View, FVector::FReal CorrectDistance = -1.)
 	{
-		const float MaxDistanceSq = (CorrectDistance > 0) ? FMath::Square(CorrectDistance) : ARecastNavMesh::GetDrawDistanceSq();
+		const FVector::FReal MaxDistanceSq = (CorrectDistance > 0.) ? FMath::Square(CorrectDistance) : ARecastNavMesh::GetDrawDistanceSq();
 		return	FVector::DistSquaredXY(Start, View->ViewMatrices.GetViewOrigin()) < MaxDistanceSq &&
 				FVector::DistSquaredXY(End, View->ViewMatrices.GetViewOrigin()) < MaxDistanceSq;
 	}
 
-	FVector EvalArc(const FVector& Org, const FVector& Dir, const float h, const float u)
+	FVector EvalArc(const FVector& Org, const FVector& Dir, const FVector::FReal h, const FVector::FReal u)
 	{
 		FVector Pt = Org + Dir * u;
 		Pt.Z += h * (1 - (u * 2 - 1)*(u * 2 - 1));
@@ -127,14 +127,14 @@ namespace FNavMeshRenderingHelpers
 			return;
 		}
 
-		const float ArcPtsScale = 1.0f / (float)Segments;
+		const FVector::FReal ArcPtsScale = 1. / (FVector::FReal)Segments;
 		const FVector Dir = End - Start;
-		const float Length = Dir.Size();
+		const FVector::FReal Length = Dir.Size();
 
 		FVector Prev = Start;
 		for (uint32 i = 1; i <= Segments; ++i)
 		{
-			const float u = i * ArcPtsScale;
+			const FVector::FReal u = (FVector::FReal)i * ArcPtsScale;
 			const FVector Pt = EvalArc(Start, Dir, Length*Height, u);
 
 			DebugLines.Add(FDebugRenderSceneProxy::FDebugLine(Prev, Pt, Color.ToFColor(true)));
@@ -466,7 +466,8 @@ void FNavMeshSceneProxyData::Serialize(FArchive& Ar)
 
 uint32 FNavMeshSceneProxyData::GetAllocatedSize() const
 {
-	return MeshBuilders.GetAllocatedSize() +
+	return IntCastChecked<uint32>(
+		MeshBuilders.GetAllocatedSize() +
 		ThickLineItems.GetAllocatedSize() +
 		TileEdgeLines.GetAllocatedSize() +
 		NavMeshEdgeLines.GetAllocatedSize() +
@@ -476,7 +477,7 @@ uint32 FNavMeshSceneProxyData::GetAllocatedSize() const
 		AuxPoints.GetAllocatedSize() +
 		AuxBoxes.GetAllocatedSize() +
 		DebugLabels.GetAllocatedSize() +
-		OctreeBounds.GetAllocatedSize();
+		OctreeBounds.GetAllocatedSize());
 }
 
 #if WITH_RECAST
@@ -893,7 +894,7 @@ void FNavMeshSceneProxyData::GatherData(const ARecastNavMesh* NavMesh, int32 InN
 					if (bGatherTileBounds)
 					{
 						const FBox TileBox = NavMesh->GetNavMeshTileBounds(TileIndex);
-						const float DrawZ = (TileBox.Min.Z + TileBox.Max.Z) * 0.5f;
+						const FVector::FReal DrawZ = (TileBox.Min.Z + TileBox.Max.Z) * 0.5;
 						const FVector LL(TileBox.Min.X, TileBox.Min.Y, DrawZ);
 						const FVector UR(TileBox.Max.X, TileBox.Max.Y, DrawZ);
 						const FVector UL(LL.X, UR.Y, DrawZ);
@@ -1134,14 +1135,14 @@ void FNavMeshSceneProxyData::GatherData(const ARecastNavMesh* NavMesh, int32 InN
 					const FVector B1 = SegInfo.RightEnd + NavMeshDrawOffset;
 					const FVector Edge0 = B0 - A0;
 					const FVector Edge1 = B1 - A1;
-					const float Len0 = Edge0.Size();
-					const float Len1 = Edge1.Size();
+					const FVector::FReal Len0 = Edge0.Size();
+					const FVector::FReal Len1 = Edge1.Size();
 					const FColor SegColor = FNavMeshRenderingHelpers::DarkenColor(NavMeshColors[SegInfo.AreaID]);
 					const FColor ColA = (SegInfo.ValidEnds & FRecastDebugGeometry::OMLE_Left) ? FColor::White : FColor::Black;
 					const FColor ColB = (SegInfo.ValidEnds & FRecastDebugGeometry::OMLE_Right) ? FColor::White : FColor::Black;
 
 					constexpr int32 NumArcPoints = 8;
-					constexpr float ArcPtsScale = 1.0f / NumArcPoints;
+					constexpr FVector::FReal ArcPtsScale = 1. / NumArcPoints;
 
 					FVector Prev0 = FNavMeshRenderingHelpers::EvalArc(A0, Edge0, Len0*0.25f, 0);
 					FVector Prev1 = FNavMeshRenderingHelpers::EvalArc(A1, Edge1, Len1*0.25f, 0);
@@ -1149,7 +1150,7 @@ void FNavMeshSceneProxyData::GatherData(const ARecastNavMesh* NavMesh, int32 InN
 					FNavMeshRenderingHelpers::AddVertex(DebugMeshData, Prev1, ColA);
 					for (int32 ArcIdx = 1; ArcIdx <= NumArcPoints; ArcIdx++)
 					{
-						const float u = ArcIdx * ArcPtsScale;
+						const FVector::FReal u = ArcIdx * ArcPtsScale;
 						FVector Pt0 = FNavMeshRenderingHelpers::EvalArc(A0, Edge0, Len0*0.25f, u);
 						FVector Pt1 = FNavMeshRenderingHelpers::EvalArc(A1, Edge1, Len1*0.25f, u);
 
@@ -1589,7 +1590,9 @@ void FNavMeshDebugDrawDelegateHelper::DrawDebugLabels(UCanvas* Canvas, APlayerCo
 			if (FNavMeshRenderingHelpers::PointInView(DebugText->Location, View))
 			{
 				const FVector ScreenLoc = Canvas->Project(DebugText->Location);
-				Canvas->DrawText(Font, DebugText->Text, ScreenLoc.X, ScreenLoc.Y);
+				Canvas->DrawText(Font, DebugText->Text
+				, FloatCastChecked<float>(ScreenLoc.X, UE::LWC::DefaultFloatPrecision)
+				, FloatCastChecked<float>(ScreenLoc.Y, UE::LWC::DefaultFloatPrecision));
 			}
 		}
 	}
@@ -1611,14 +1614,15 @@ FPrimitiveViewRelevance FNavMeshSceneProxy::GetViewRelevance(const FSceneView* V
 
 uint32 FNavMeshSceneProxy::GetAllocatedSizeInternal() const
 {
-	return FDebugRenderSceneProxy::GetAllocatedSize() +
+	return IntCastChecked<uint32>(
+		FDebugRenderSceneProxy::GetAllocatedSize() +
 		ProxyData.GetAllocatedSize() +
 		IndexBuffer.Indices.GetAllocatedSize() +
 		VertexBuffers.PositionVertexBuffer.GetNumVertices() * VertexBuffers.PositionVertexBuffer.GetStride() +
 		VertexBuffers.StaticMeshVertexBuffer.GetResourceSize() +
 		VertexBuffers.ColorVertexBuffer.GetNumVertices() * VertexBuffers.ColorVertexBuffer.GetStride() +
 		MeshColors.GetAllocatedSize() + MeshColors.Num() * sizeof(FColoredMaterialRenderProxy) +
-		MeshBatchElements.GetAllocatedSize();
+		MeshBatchElements.GetAllocatedSize());
 }
 
 //////////////////////////////////////////////////////////////////////////
diff --git a/Engine/Source/Runtime/NavigationSystem/Private/NavMesh/NavTestRenderingComponent.cpp b/Engine/Source/Runtime/NavigationSystem/Private/NavMesh/NavTestRenderingComponent.cpp
--- a/Engine/Source/Runtime/NavigationSystem/Private/NavMesh/NavTestRenderingComponent.cpp
+++ b/Engine/Source/Runtime/NavigationSystem/Private/NavMesh/NavTestRenderingComponent.cpp
@@ -110,7 +110,7 @@ void FNavTestSceneProxy::GetDynamicMeshElements(const TArray<const FSceneView*>&
 					FDynamicMeshBuilder	MeshBuilder(View->GetFeatureLevel());
 					MeshBuilder.AddVertices(ClosedSetVerts);
 					MeshBuilder.AddTriangles(ClosedSetIndices);
-					MeshBuilder.GetMesh(FMatrix::Identity, MeshColorInstance, GetDepthPriorityGroup(View), false, false, ViewIndex, Collector);
+					MeshBuilder.GetMesh(FMatrix::Identity, MeshColorInstance, IntCastChecked<uint8>((int32)GetDepthPriorityGroup(View)), false, false, ViewIndex, Collector);
 				}
 
 				if (OpenSetIndices.Num())
@@ -119,7 +119,7 @@ void FNavTestSceneProxy::GetDynamicMeshElements(const TArray<const FSceneView*>&
 					FDynamicMeshBuilder	MeshBuilder(View->GetFeatureLevel());
 					MeshBuilder.AddVertices(OpenSetVerts);
 					MeshBuilder.AddTriangles(OpenSetIndices);
-					MeshBuilder.GetMesh(FMatrix::Identity, MeshColorInstance, GetDepthPriorityGroup(View), false, false, ViewIndex, Collector);
+					MeshBuilder.GetMesh(FMatrix::Identity, MeshColorInstance, IntCastChecked<uint8>((int32)GetDepthPriorityGroup(View)), false, false, ViewIndex, Collector);
 				}
 			}
 
@@ -303,17 +303,17 @@ FPrimitiveViewRelevance FNavTestSceneProxy::GetViewRelevance(const FSceneView* V
 
 uint32 FNavTestSceneProxy::GetAllocatedSizeInternal() const
 {
-	int32 InternalAllocSize = 0;
+	SIZE_T InternalAllocSize = 0;
 	for (TSet<FNodeDebugData>::TConstIterator It(NodeDebug); It; ++It)
 	{
 		InternalAllocSize += (*It).Desc.GetAllocatedSize();
 	}
 
-	return FDebugRenderSceneProxy::GetAllocatedSize() + PathPoints.GetAllocatedSize()
+	return IntCastChecked<uint32>(FDebugRenderSceneProxy::GetAllocatedSize() + PathPoints.GetAllocatedSize()
 		+ PathPointFlags.GetAllocatedSize()
 		+ OpenSetVerts.GetAllocatedSize() + OpenSetIndices.GetAllocatedSize()
 		+ ClosedSetVerts.GetAllocatedSize() + ClosedSetIndices.GetAllocatedSize()
-		+ NodeDebug.GetAllocatedSize() + InternalAllocSize;
+		+ NodeDebug.GetAllocatedSize() + InternalAllocSize);
 
 }
 
@@ -365,7 +365,7 @@ void FNavTestDebugDrawDelegateHelper::DrawDebugLabels(UCanvas* Canvas, APlayerCo
 
 				Canvas->SetDrawColor(MyColor);
 
-				const FVector ScreenLoc = Canvas->Project(NodeData.Position) + FVector(NavTestActor->TextCanvasOffset, 0.f);
+				const FVector3f ScreenLoc(Canvas->Project(NodeData.Position) + FVector(NavTestActor->TextCanvasOffset, 0.f));
 				Canvas->DrawText(RenderFont, NodeData.Desc, ScreenLoc.X, ScreenLoc.Y);
 			}
 		}
@@ -377,9 +377,10 @@ void FNavTestDebugDrawDelegateHelper::DrawDebugLabels(UCanvas* Canvas, APlayerCo
 		{
 			if (FNavTestSceneProxy::LocationInView(PathPoints[PointIndex], View))
 			{
-				const FVector PathPointLoc = Canvas->Project(PathPoints[PointIndex]);
+				const FVector3f PathPointLoc(Canvas->Project(PathPoints[PointIndex]));
 				const UFont* RenderFont = GEngine->GetSmallFont();
 				Canvas->DrawText(RenderFont, PathPointFlags[PointIndex], PathPointLoc.X, PathPointLoc.Y);
+
 			}
 		}
 
diff --git a/Engine/Source/Runtime/NavigationSystem/Private/NavMesh/PImplRecastNavMesh.cpp b/Engine/Source/Runtime/NavigationSystem/Private/NavMesh/PImplRecastNavMesh.cpp
--- a/Engine/Source/Runtime/NavigationSystem/Private/NavMesh/PImplRecastNavMesh.cpp
+++ b/Engine/Source/Runtime/NavigationSystem/Private/NavMesh/PImplRecastNavMesh.cpp
@@ -51,7 +51,7 @@ static void* DetourMalloc(int Size, dtAllocHint Hint)
 	LLM_SCOPE(ELLMTag::NavigationRecast);
 	void* Result = FMemory::Malloc(uint32(Size));
 #if STATS
-	const uint32 ActualSize = FMemory::GetAllocSize(Result);
+	const SIZE_T ActualSize = FMemory::GetAllocSize(Result);
 
 	switch (Hint)
 	{
@@ -127,7 +127,7 @@ static void* RecastMalloc(int Size, rcAllocHint)
 	LLM_SCOPE(ELLMTag::NavigationRecast);
 	void* Result = FMemory::Malloc(uint32(Size));
 #if STATS
-	const uint32 ActualSize = FMemory::GetAllocSize(Result);
+	const SIZE_T ActualSize = FMemory::GetAllocSize(Result);
 	INC_DWORD_STAT_BY(STAT_NavigationMemory, ActualSize);
 	INC_MEMORY_STAT_BY(STAT_Navigation_RecastMemory, ActualSize);
 #endif // STATS
@@ -137,7 +137,7 @@ static void* RecastMalloc(int Size, rcAllocHint)
 static void DetourFree(void* Original, dtAllocHint Hint)
 {
 #if STATS
-	const uint32 Size = FMemory::GetAllocSize(Original);
+	const SIZE_T Size = FMemory::GetAllocSize(Original);
 
 	switch (Hint)
 	{
@@ -212,7 +212,7 @@ static void DetourFree(void* Original, dtAllocHint Hint)
 static void RecastFree(void* Original)
 {
 #if STATS
-	const uint32 Size = FMemory::GetAllocSize(Original);
+	const SIZE_T Size = FMemory::GetAllocSize(Original);
 	DEC_DWORD_STAT_BY(STAT_NavigationMemory, Size);
 	DEC_MEMORY_STAT_BY(STAT_Navigation_RecastMemory, Size);
 #endif // STATS
@@ -1231,7 +1231,8 @@ ENavigationQueryResult::Type FPImplRecastNavMesh::TestClusterPath(const FVector&
 	NavNodeRef StartPolyID, EndPolyID;
 	const dtQueryFilter* ClusterFilter = ((const FRecastQueryFilter*)NavMeshOwner->GetDefaultQueryFilterImpl())->GetAsDetourQueryFilter();
 
-	INITIALIZE_NAVQUERY_SIMPLE(ClusterQuery, NavMeshOwner->DefaultMaxHierarchicalSearchNodes);
+	check(NavMeshOwner->DefaultMaxHierarchicalSearchNodes >= 0. && NavMeshOwner->DefaultMaxHierarchicalSearchNodes <= (float)TNumericLimits<int32>::Max());
+	INITIALIZE_NAVQUERY_SIMPLE(ClusterQuery, static_cast<int32>(NavMeshOwner->DefaultMaxHierarchicalSearchNodes));
 
 	const bool bCanSearch = InitPathfinding(StartLoc, EndLoc, ClusterQuery, ClusterFilter, RecastStartPos, StartPolyID, RecastEndPos, EndPolyID);
 	if (!bCanSearch)
@@ -1441,7 +1442,7 @@ bool FPImplRecastNavMesh::FindStraightPath(const FVector& StartLoc, const FVecto
 			CurVert->NodeRef = StringPullResult.getRef(VertIdx);
 
 			FNavMeshNodeFlags CurNodeFlags(0);
-			CurNodeFlags.PathFlags = StringPullResult.getFlag(VertIdx);
+			CurNodeFlags.PathFlags = IntCastChecked<uint8>(StringPullResult.getFlag(VertIdx));
 
 			uint8 AreaID = RECAST_DEFAULT_AREA;
 			DetourNavMesh->getPolyArea(CurVert->NodeRef, &AreaID);
@@ -2651,7 +2652,7 @@ bool FPImplRecastNavMesh::GetDebugGeometryForTile(FRecastDebugGeometry& OutGeome
 	int32 NumVertsToReserve = 0;
 	int32 NumIndicesToReserve = 0;
 
-	int32 ForbiddenFlags = OutGeometry.bMarkForbiddenPolys 
+	uint16 ForbiddenFlags = OutGeometry.bMarkForbiddenPolys 
 		? GetFilterForbiddenFlags((const FRecastQueryFilter*)NavMeshOwner->GetDefaultQueryFilterImpl()) 
 		: 0;
 
@@ -2828,7 +2829,7 @@ int32 FPImplRecastNavMesh::GetTilesDebugGeometry(const FRecastNavMeshGenerator*
 				int32 Rank = 0;
 				if (Range != 0.)
 				{
-					Rank = FRecastDebugGeometry::BuildTimeBucketsCount * ((DebugData->BuildTime - OutGeometry.MinTileBuildTime) / Range);
+					Rank = static_cast<int32>(FRecastDebugGeometry::BuildTimeBucketsCount * ((DebugData->BuildTime - OutGeometry.MinTileBuildTime) / Range));
 					Rank = FMath::Clamp(Rank, 0, FRecastDebugGeometry::BuildTimeBucketsCount-1);
 				}
 				Indices = &OutGeometry.TileBuildTimesIndices[Rank];
@@ -3257,7 +3258,7 @@ void FPImplRecastNavMesh::OnAreaCostChanged()
 		uint8 AreaCostOrder[RECAST_MAX_AREAS];
 		for (int32 Idx = 0; Idx < RECAST_MAX_AREAS; Idx++)
 		{
-			AreaCostOrder[AreaData[Idx].Index] = Idx;
+			AreaCostOrder[AreaData[Idx].Index] = static_cast<uint8>(Idx);
 		}
 
 		DetourNavMesh->applyAreaCostOrder(AreaCostOrder);
diff --git a/Engine/Source/Runtime/NavigationSystem/Private/NavMesh/RecastHelpers.cpp b/Engine/Source/Runtime/NavigationSystem/Private/NavMesh/RecastHelpers.cpp
--- a/Engine/Source/Runtime/NavigationSystem/Private/NavMesh/RecastHelpers.cpp
+++ b/Engine/Source/Runtime/NavigationSystem/Private/NavMesh/RecastHelpers.cpp
@@ -74,7 +74,7 @@ FColor Recast2UnrealColor(const unsigned int RecastColor)
 {
 	const uint8 R = (RecastColor & 0xFF);
 	const uint8 G = (RecastColor & 0xFF00) >> 8;
-	const uint8 B = (RecastColor & 0xFF0000) >> 16;
-	const uint8 A = (RecastColor & 0xFF000000) >> 24;
+	const uint8 B = static_cast<uint8>((RecastColor & 0xFF0000) >> 16);
+	const uint8 A = static_cast<uint8>((RecastColor & 0xFF000000) >> 24);
 	return FColor(R,G,B,A);	// can't be direct assignation since internally FColor is not RGBA
 }
\ No newline at end of file
diff --git a/Engine/Source/Runtime/NavigationSystem/Private/NavMesh/RecastNavMesh.cpp b/Engine/Source/Runtime/NavigationSystem/Private/NavMesh/RecastNavMesh.cpp
--- a/Engine/Source/Runtime/NavigationSystem/Private/NavMesh/RecastNavMesh.cpp
+++ b/Engine/Source/Runtime/NavigationSystem/Private/NavMesh/RecastNavMesh.cpp
@@ -229,7 +229,7 @@ void FNavMeshTileData::MakeUnique()
 	}
 }
 
-float ARecastNavMesh::DrawDistanceSq = 0.0f;
+FVector::FReal ARecastNavMesh::DrawDistanceSq = 0.;
 float ARecastNavMesh::MinimumSizeForChaosNavMeshInfluenceSq = 0.0f;
 #if !WITH_RECAST
 
@@ -277,7 +277,7 @@ ARecastNavMesh::~ARecastNavMesh()
 //----------------------------------------------------------------------//
 uint32 FRecastDebugGeometry::GetAllocatedSize() const
 {
-	uint32 Size = sizeof(*this) + MeshVerts.GetAllocatedSize()
+	SIZE_T Size = sizeof(*this) + MeshVerts.GetAllocatedSize()
 		+ BuiltMeshIndices.GetAllocatedSize()
 		+ PolyEdges.GetAllocatedSize()
 		+ NavMeshEdges.GetAllocatedSize()
@@ -308,7 +308,7 @@ uint32 FRecastDebugGeometry::GetAllocatedSize() const
 	}
 #endif // WITH_NAVMESH_CLUSTER_LINKS
 
-	return Size;
+	return IntCastChecked<uint32>(Size);
 }
 
 //----------------------------------------------------------------------//
@@ -752,7 +752,9 @@ void ARecastNavMesh::PostInitProperties()
 void ARecastNavMesh::RecreateDefaultFilter()
 {
 	DefaultQueryFilter->SetFilterType<FRecastQueryFilter>();
-	DefaultQueryFilter->SetMaxSearchNodes(DefaultMaxSearchNodes);
+
+	check(DefaultMaxSearchNodes >= 0 && DefaultMaxSearchNodes <= (float)TNumericLimits<uint32>::Max());
+	DefaultQueryFilter->SetMaxSearchNodes(static_cast<int32>(DefaultMaxSearchNodes));
 
 	FRecastQueryFilter* DetourFilter = static_cast<FRecastQueryFilter*>(DefaultQueryFilter->GetImplementation());
 	DetourFilter->SetIsVirtual(bUseVirtualFilters);
@@ -774,8 +776,8 @@ void ARecastNavMesh::RecreateDefaultFilter()
 
 		if (DefArea)
 		{
-			DetourFilter->SetAreaCost(AreaData.AreaID, DefArea->DefaultCost);
-			DetourFilter->SetFixedAreaEnteringCost(AreaData.AreaID, DefArea->GetFixedAreaEnteringCost());
+			DetourFilter->SetAreaCost(IntCastChecked<uint8>(AreaData.AreaID), DefArea->DefaultCost);
+			DetourFilter->SetFixedAreaEnteringCost(IntCastChecked<uint8>(AreaData.AreaID), DefArea->GetFixedAreaEnteringCost());
 		}
 	}
 }
@@ -798,8 +800,8 @@ void ARecastNavMesh::OnNavAreaAdded(const UClass* NavAreaClass, int32 AgentIndex
 	{
 		UNavArea* DefArea = ((UClass*)NavAreaClass)->GetDefaultObject<UNavArea>();
 
-		DefaultQueryFilter->SetAreaCost(AreaID, DefArea->DefaultCost);
-		DefaultQueryFilter->SetFixedAreaEnteringCost(AreaID, DefArea->GetFixedAreaEnteringCost());
+		DefaultQueryFilter->SetAreaCost(IntCastChecked<uint8>(AreaID), DefArea->DefaultCost);
+		DefaultQueryFilter->SetFixedAreaEnteringCost(IntCastChecked<uint8>(AreaID), DefArea->GetFixedAreaEnteringCost());
 	}
 
 	// update generator's cached data
@@ -1025,7 +1027,7 @@ void ARecastNavMesh::Serialize( FArchive& Ar )
 		if (Ar.IsSaving())
 		{
 			int64 CurPos = Ar.Tell();
-			RecastNavMeshSizeBytes = CurPos - RecastNavMeshSizePos;
+			RecastNavMeshSizeBytes = IntCastChecked<uint32>(CurPos - RecastNavMeshSizePos);
 			Ar.Seek(RecastNavMeshSizePos);
 			Ar << RecastNavMeshSizeBytes;
 			Ar.Seek(CurPos);
@@ -1964,9 +1966,9 @@ void ARecastNavMesh::UpdateCustomLink(const INavLinkCustomInterface* CustomLink)
 		UNavArea* DefArea = (UNavArea*)(AreaClass->GetDefaultObject());
 		const uint16 PolyFlags = DefArea->GetAreaFlags() | ARecastNavMesh::GetNavLinkFlag();
 
-		RecastNavMeshImpl->UpdateNavigationLinkArea(UserId, AreaId, PolyFlags);
+		RecastNavMeshImpl->UpdateNavigationLinkArea(UserId, IntCastChecked<uint8>(AreaId), PolyFlags);
 #if WITH_NAVMESH_SEGMENT_LINKS
-		RecastNavMeshImpl->UpdateSegmentLinkArea(UserId, AreaId, PolyFlags);
+		RecastNavMeshImpl->UpdateSegmentLinkArea(UserId, IntCastChecked<uint8>(AreaId), PolyFlags);
 #endif // WITH_NAVMESH_SEGMENT_LINKS
 
 #if !UE_BUILD_SHIPPING
@@ -1983,7 +1985,7 @@ void ARecastNavMesh::UpdateNavigationLinkArea(int32 UserId, TSubclassOf<UNavArea
 		UNavArea* DefArea = (UNavArea*)(AreaClass->GetDefaultObject());
 		const uint16 PolyFlags = DefArea->GetAreaFlags() | ARecastNavMesh::GetNavLinkFlag();
 
-		RecastNavMeshImpl->UpdateNavigationLinkArea(UserId, AreaId, PolyFlags);
+		RecastNavMeshImpl->UpdateNavigationLinkArea(UserId, IntCastChecked<uint8>(AreaId), PolyFlags);
 	}
 }
 
@@ -1996,7 +1998,7 @@ void ARecastNavMesh::UpdateSegmentLinkArea(int32 UserId, TSubclassOf<UNavArea> A
 		UNavArea* DefArea = (UNavArea*)(AreaClass->GetDefaultObject());
 		const uint16 PolyFlags = DefArea->GetAreaFlags() | ARecastNavMesh::GetNavLinkFlag();
 
-		RecastNavMeshImpl->UpdateSegmentLinkArea(UserId, AreaId, PolyFlags);
+		RecastNavMeshImpl->UpdateSegmentLinkArea(UserId, IntCastChecked<uint8>(AreaId), PolyFlags);
 	}
 }
 #endif // WITH_NAVMESH_SEGMENT_LINKS
@@ -2039,7 +2041,7 @@ bool ARecastNavMesh::SetPolyArea(NavNodeRef PolyID, TSubclassOf<UNavArea> AreaCl
 		if (AreaId != INDEX_NONE && NavMesh)
 		{
 			// @todo implement a single detour function that would do both
-			bSuccess = dtStatusSucceed(NavMesh->setPolyArea(PolyID, AreaId));
+			bSuccess = dtStatusSucceed(NavMesh->setPolyArea(PolyID, IntCastChecked<unsigned char>(AreaId)));
 			bSuccess = (bSuccess && dtStatusSucceed(NavMesh->setPolyFlags(PolyID, AreaFlags)));
 		}
 	}
@@ -2058,7 +2060,7 @@ void ARecastNavMesh::SetPolyArrayArea(const TArray<FNavPoly>& Polys, TSubclassOf
 		{
 			for (int32 Idx = 0; Idx < Polys.Num(); Idx++)
 			{
-				NavMesh->setPolyArea(Polys[Idx].Ref, AreaId);
+				NavMesh->setPolyArea(Polys[Idx].Ref, IntCastChecked<unsigned char>(AreaId));
 				NavMesh->setPolyFlags(Polys[Idx].Ref, AreaFlags);
 			}
 		}
@@ -2084,8 +2086,8 @@ int32 ARecastNavMesh::ReplaceAreaInTileBounds(const FBox& Bounds, TSubclassOf<UN
 		dtNavMesh const* const ConstDetourNavMesh = RecastNavMeshImpl->GetRecastMesh();
 
 		const FVector RcNavMeshOrigin = Unreal2RecastPoint(NavMeshOriginOffset);
-		const FVector::FReal RcTileSize = FMath::TruncToInt(TileSizeUU / CellSize);
-		const FVector::FReal TileSizeInWorldUnits = RcTileSize * CellSize;
+		const float RcTileSize = FMath::TruncToFloat(TileSizeUU / CellSize);
+		const float TileSizeInWorldUnits = RcTileSize * CellSize;
 		const FRcTileBox TileBox(Bounds, RcNavMeshOrigin, TileSizeInWorldUnits);
 
 		for (int32 TileY = TileBox.YMin; TileY <= TileBox.YMax; ++TileY)
@@ -2119,7 +2121,7 @@ int32 ARecastNavMesh::ReplaceAreaInTileBounds(const FBox& Bounds, TSubclassOf<UN
 							{
 								if (Poly->getArea() == OldAreaID)
 								{
-									Poly->setArea(NewAreaID);
+									Poly->setArea(IntCastChecked<unsigned char>(NewAreaID));
 									++PolysTouched;
 								}
 							}
@@ -3047,7 +3049,7 @@ void ARecastNavMesh::PostEditChangeProperty(FPropertyChangedEvent& PropertyChang
 				TileSizeUU = UE::NavMesh::Private::GetClampedTileSizeUU(TileSizeUU, CellSize, AgentRadius);
 						
 				// Match cell size to tile size.
-				CellSize = TileSizeUU / FMath::TruncToInt(TileSizeUU / CellSize);
+				CellSize = TileSizeUU / FMath::TruncToFloat(TileSizeUU / CellSize);
 
 				// update config
 				FillConfig(NavDataConfig);
@@ -3058,7 +3060,7 @@ void ARecastNavMesh::PostEditChangeProperty(FPropertyChangedEvent& PropertyChang
 				TileSizeUU = UE::NavMesh::Private::GetClampedTileSizeUU(TileSizeUU, CellSize, AgentRadius);
 				
 				// Match tile size to cell size.
-				TileSizeUU = CellSize * FMath::TruncToInt(TileSizeUU / CellSize);
+				TileSizeUU = CellSize * FMath::TruncToFloat(TileSizeUU / CellSize);
 
 				// update config
 				FillConfig(NavDataConfig);
@@ -3242,16 +3244,16 @@ bool ARecastNavMesh::HasCompleteDataInRadius(const FVector& TestLocation, FVecto
 
 	const dtNavMesh* NavMesh = RecastNavMeshImpl->DetourNavMesh;
 	const dtNavMeshParams* NavParams = RecastNavMeshImpl->DetourNavMesh->getParams();
-	const FVector::FReal NavTileSize = CellSize * FMath::TruncToInt(TileSizeUU / CellSize);
+	const float NavTileSize = CellSize * FMath::TruncToFloat(TileSizeUU / CellSize);
 	const FVector RcNavOrigin(NavParams->orig[0], NavParams->orig[1], NavParams->orig[2]);
 
 	const FBox RcBounds = Unreal2RecastBox(FBox::BuildAABB(TestLocation, FVector(TestRadius, TestRadius, 0)));
 	const FVector RcTestLocation = Unreal2RecastPoint(TestLocation);
 
-	const int32 MinTileX = FMath::FloorToInt((RcBounds.Min.X - RcNavOrigin.X) / NavTileSize);
-	const int32 MaxTileX = FMath::CeilToInt((RcBounds.Max.X - RcNavOrigin.X) / NavTileSize);
-	const int32 MinTileY = FMath::FloorToInt((RcBounds.Min.Z - RcNavOrigin.Z) / NavTileSize);
-	const int32 MaxTileY = FMath::CeilToInt((RcBounds.Max.Z - RcNavOrigin.Z) / NavTileSize);
+	const int32 MinTileX = IntCastChecked<int32>(FMath::FloorToInt((RcBounds.Min.X - RcNavOrigin.X) / NavTileSize));
+	const int32 MaxTileX = IntCastChecked<int32>(FMath::CeilToInt((RcBounds.Max.X - RcNavOrigin.X) / NavTileSize));
+	const int32 MinTileY = IntCastChecked<int32>(FMath::FloorToInt((RcBounds.Min.Z - RcNavOrigin.Z) / NavTileSize));
+	const int32 MaxTileY = IntCastChecked<int32>(FMath::CeilToInt((RcBounds.Max.Z - RcNavOrigin.Z) / NavTileSize));
 	const FVector RcTileExtent2D(NavTileSize * 0.5f, 0.f, NavTileSize * 0.5f);
 	const FVector::FReal RadiusSq = FMath::Square(TestRadius);
 
@@ -3317,10 +3319,10 @@ void ARecastNavMesh::UpdateActiveTiles(const TArray<FNavigationInvokerRaw>& Invo
 		const  FVector::FReal TileCenterDistanceToRemoveSq = FMath::Square(TileDim * UE_SQRT_2 / 2 + Invoker.RadiusMax);
 		const  FVector::FReal TileCenterDistanceToAddSq = FMath::Square(TileDim * UE_SQRT_2 / 2 + Invoker.RadiusMin);
 
-		const int32 MinTileX = FMath::FloorToInt((InvokerRelativeLocation.X - Invoker.RadiusMax) / TileDim);
-		const int32 MaxTileX = FMath::CeilToInt((InvokerRelativeLocation.X + Invoker.RadiusMax) / TileDim);
-		const int32 MinTileY = FMath::FloorToInt((InvokerRelativeLocation.Y - Invoker.RadiusMax) / TileDim);
-		const int32 MaxTileY = FMath::CeilToInt((InvokerRelativeLocation.Y + Invoker.RadiusMax) / TileDim);
+		const int32 MinTileX = IntCastChecked<int32>(FMath::FloorToInt((InvokerRelativeLocation.X - Invoker.RadiusMax) / TileDim));
+		const int32 MaxTileX = IntCastChecked<int32>(FMath::CeilToInt((InvokerRelativeLocation.X + Invoker.RadiusMax) / TileDim));
+		const int32 MinTileY = IntCastChecked<int32>(FMath::FloorToInt((InvokerRelativeLocation.Y - Invoker.RadiusMax) / TileDim));
+		const int32 MaxTileY = IntCastChecked<int32>(FMath::CeilToInt((InvokerRelativeLocation.Y + Invoker.RadiusMax) / TileDim));
 
 		for (int32 X = MinTileX; X <= MaxTileX; ++X)
 		{
@@ -3441,7 +3443,7 @@ void ARecastNavMesh::DirtyTilesInBounds(const FBox& Bounds)
 		// Add tiles within the overlapping bounds
 		TArray<FIntPoint> Points;
 		const FVector RcNavMeshOrigin = Unreal2RecastPoint(NavMeshOriginOffset);
-		const float RcTileSize = FMath::TruncToInt(TileSizeUU / CellSize);
+		const float RcTileSize = FMath::TruncToFloat(TileSizeUU / CellSize);
 		const float TileSizeInWorldUnits = RcTileSize * CellSize;
 		const FRcTileBox TileBox(OverlappingBounds, RcNavMeshOrigin, TileSizeInWorldUnits);
 
diff --git a/Engine/Source/Runtime/NavigationSystem/Private/NavMesh/RecastNavMeshGenerator.cpp b/Engine/Source/Runtime/NavigationSystem/Private/NavMesh/RecastNavMeshGenerator.cpp
--- a/Engine/Source/Runtime/NavigationSystem/Private/NavMesh/RecastNavMeshGenerator.cpp
+++ b/Engine/Source/Runtime/NavigationSystem/Private/NavMesh/RecastNavMeshGenerator.cpp
@@ -1331,7 +1331,7 @@ struct FOffMeshData
 			const int32* AreaID = AreaClassToIdMap->Find(AreaClass);
 			if (AreaID != NULL)
 			{
-				NewInfo.area = *AreaID;
+				NewInfo.area = IntCastChecked<unsigned char>(*AreaID);
 				NewInfo.polyFlag = FlagsPerArea[*AreaID];
 			}
 			else
@@ -1372,7 +1372,7 @@ struct FOffMeshData
 			const int32* AreaID = AreaClassToIdMap->Find(AreaClass);
 			if (AreaID != NULL)
 			{
-				NewInfo.area = *AreaID;
+				NewInfo.area = IntCastChecked<unsigned char>(*AreaID);
 				NewInfo.polyFlag = FlagsPerArea[*AreaID];
 			}
 			else
@@ -1478,7 +1478,7 @@ struct FTileCacheCompressor : public dtTileCacheCompressor
 			const int64 CompressedSize = FOodleDataCompression::Compress((void*)DataPtr, DataSize, (const void*)buffer, bufferSize, GNavmeshTileCacheCompressor, GNavmeshTileCacheCompressionLevel);
 			if (CompressedSize > 0)
 			{
-				*outCompressedSize = CompressedSize + HeaderSize;
+				*outCompressedSize = IntCastChecked<int32>(CompressedSize + HeaderSize);
 				return DT_SUCCESS;
 			}
 			else
@@ -2048,12 +2048,12 @@ void FRecastTileGenerator::GatherNavigationDataGeometry(const TSharedRef<FNaviga
 				{
 					LLM_SCOPE_BYTAG(NavigationOctree);
 
-					const int32 PrevElementMemory = ElementData->GetAllocatedSize();
+					const SIZE_T PrevElementMemory = ElementData->GetAllocatedSize();
 					FNavigationRelevantData* ModData = (FNavigationRelevantData*)&ElementData;
 					AddVoxelCache(ModData->VoxelData, CachedVoxels, NumCachedVoxels);
 
-					const int32 NewElementMemory = ElementData->GetAllocatedSize();
-					const int32 ElementMemoryDelta = NewElementMemory - PrevElementMemory;
+					const SIZE_T NewElementMemory = ElementData->GetAllocatedSize();
+					const SIZE_T ElementMemoryDelta = NewElementMemory - PrevElementMemory;
 					INC_MEMORY_STAT_BY(STAT_Navigation_CollisionTreeMemory, ElementMemoryDelta);
 				}
 			}
@@ -2838,7 +2838,7 @@ void FRecastTileGenerator::GenerateRecastFilter(FNavMeshBuildContext& BuildConte
 	SCOPE_CYCLE_COUNTER(STAT_Navigation_RecastFilter)
 
 	// TileConfig.walkableHeight is set to 1 when marking low spans, calculate real value for filtering
-	const int32 FilterWalkableHeight = FMath::CeilToInt(TileConfig.AgentHeight / TileConfig.ch);
+	const int32 FilterWalkableHeight = FMath::CeilToInt(TileConfig.AgentHeight / static_cast<float>(TileConfig.ch));
 
 	// Once all geometry is rasterized, we do initial pass of filtering to
 	// remove unwanted overhangs caused by the conservative rasterization
@@ -2920,7 +2920,7 @@ ETimeSliceWorkResult FRecastTileGenerator::GenerateRecastFilterTimeSliced(FNavMe
 	case EGenerateRecastFilterTimeSlicedState::FilterWalkableLowHeightSpans:
 	{
 		// TileConfig.walkableHeight is set to 1 when marking low spans, calculate real value for filtering
-		const int32 FilterWalkableHeight = FMath::CeilToInt(TileConfig.AgentHeight / TileConfig.ch);
+		const int32 FilterWalkableHeight = FMath::CeilToInt(TileConfig.AgentHeight / static_cast<float>(TileConfig.ch));
 
 		if (!TileConfig.bMarkLowHeightAreas)
 		{
@@ -2987,7 +2987,7 @@ bool FRecastTileGenerator::RecastErodeWalkable(FNavMeshBuildContext& BuildContex
 	SCOPE_CYCLE_COUNTER(STAT_Navigation_RecastErodeWalkable);
 
 	// TileConfig.walkableHeight is set to 1 when marking low spans, calculate real value for filtering
-	const int32 FilterWalkableHeight = FMath::CeilToInt(TileConfig.AgentHeight / TileConfig.ch);
+	const int32 FilterWalkableHeight = FMath::CeilToInt(TileConfig.AgentHeight / static_cast<float>(TileConfig.ch));
 
 	if (TileConfig.walkableRadius > RECAST_VERY_SMALL_AGENT_RADIUS)
 	{
@@ -3702,7 +3702,7 @@ bool FRecastTileGenerator::GenerateNavigationDataLayer(FNavMeshBuildContext& Bui
 			OffMeshData.AreaClassToIdMap = &AdditionalCachedData.AreaClassToIdMap;
 			OffMeshData.FlagsPerArea = AdditionalCachedData.FlagsPerOffMeshLinkArea;
 			const FSimpleLinkNavModifier* LinkModifier = OffmeshLinks.GetData();
-			const float DefaultSnapHeight = TileConfig.walkableClimb * (float)TileConfig.ch;
+			const float DefaultSnapHeight = TileConfig.walkableClimb * static_cast<float>(TileConfig.ch);
 
 			for (int32 LinkModifierIndex = 0; LinkModifierIndex < OffmeshLinks.Num(); ++LinkModifierIndex, ++LinkModifier)
 			{
@@ -3752,7 +3752,7 @@ bool FRecastTileGenerator::GenerateNavigationDataLayer(FNavMeshBuildContext& Bui
 		Params.buildBvTree = TileConfig.bGenerateBVTree;
 
 #if WITH_NAVMESH_CLUSTER_LINKS
-		Params.clusterCount = GenerationContext.ClusterSet->nclusters;
+		Params.clusterCount = IntCastChecked<unsigned short>(GenerationContext.ClusterSet->nclusters);
 		Params.polyClusters = GenerationContext.ClusterSet->polyMap;
 #endif // WITH_NAVMESH_CLUSTER_LINKS
 
diff --git a/Engine/Source/Runtime/NavigationSystem/Private/NavigationData.cpp b/Engine/Source/Runtime/NavigationSystem/Private/NavigationData.cpp
--- a/Engine/Source/Runtime/NavigationSystem/Private/NavigationData.cpp
+++ b/Engine/Source/Runtime/NavigationSystem/Private/NavigationData.cpp
@@ -183,7 +183,7 @@ ANavigationData::ANavigationData(const FObjectInitializer& ObjectInitializer)
 uint16 ANavigationData::GetNextUniqueID()
 {
 	static FThreadSafeCounter StaticID(INVALID_NAVDATA);
-	return StaticID.Increment();
+	return IntCastChecked<uint16>(StaticID.Increment());
 }
 
 void ANavigationData::PostInitProperties()
@@ -308,7 +308,7 @@ void ANavigationData::TickActor(float DeltaTime, enum ELevelTick TickType, FActo
 
 	if (RepathRequests.Num() > 0)
 	{
-		float TimeStamp = GetWorldTimeStamp();
+		double TimeStamp = GetWorldTimeStamp();
 		const UWorld* World = GetWorld();
 
 		// @todo batch-process it!
@@ -673,7 +673,7 @@ void ANavigationData::DrawDebugPath(FNavigationPath* Path, const FColor PathColo
 	Path->DebugDraw(this, PathColor, Canvas, bPersistent, LifeTime, NextPathPointIndex);
 }
 
-float ANavigationData::GetWorldTimeStamp() const
+double ANavigationData::GetWorldTimeStamp() const
 {
 	const UWorld* World = GetWorld();
 	return World ? World->GetTimeSeconds() : 0.f;
@@ -852,7 +852,7 @@ void ANavigationData::RemoveQueryFilter(TSubclassOf<UNavigationQueryFilter> Filt
 
 uint32 ANavigationData::LogMemUsed() const
 {
-	uint32 ActivePathsMemSize = 0;
+	SIZE_T ActivePathsMemSize = 0;
 	{
 		// Paths can be registered from async pathfinding thread
 		// while logging is requested on main thread (console command)
@@ -860,8 +860,8 @@ uint32 ANavigationData::LogMemUsed() const
 		ActivePathsMemSize = ActivePaths.GetAllocatedSize();
 	}
 
-	const uint32 MemUsed = ActivePathsMemSize + SupportedAreas.GetAllocatedSize() +
-		QueryFilters.GetAllocatedSize() + AreaClassToIdMap.GetAllocatedSize();
+	const uint32 MemUsed = IntCastChecked<uint32>(ActivePathsMemSize + SupportedAreas.GetAllocatedSize() +
+		QueryFilters.GetAllocatedSize() + AreaClassToIdMap.GetAllocatedSize());
 
 	UE_VLOG_UELOG(this, LogNavigation, Display, TEXT("%s: ANavigationData: %u\n    self: %d"), *GetName(), MemUsed, sizeof(ANavigationData));
 
diff --git a/Engine/Source/Runtime/NavigationSystem/Private/NavigationOctree.cpp b/Engine/Source/Runtime/NavigationSystem/Private/NavigationOctree.cpp
--- a/Engine/Source/Runtime/NavigationSystem/Private/NavigationOctree.cpp
+++ b/Engine/Source/Runtime/NavigationSystem/Private/NavigationOctree.cpp
@@ -12,7 +12,7 @@ LLM_DEFINE_TAG(NavigationOctree);
 //----------------------------------------------------------------------//
 // FNavigationOctree
 //----------------------------------------------------------------------//
-FNavigationOctree::FNavigationOctree(const FVector& Origin, float Radius)
+FNavigationOctree::FNavigationOctree(const FVector& Origin, FVector::FReal Radius)
 	: TOctree2<FNavigationOctreeElement, FNavigationOctreeSemantics>(Origin, Radius)
 	, DefaultGeometryGatheringMode(ENavDataGatheringMode::Instant)
 	, bGatherGeometry(false)
@@ -54,7 +54,7 @@ void FNavigationOctree::DemandLazyDataGathering(FNavigationRelevantData& Element
 	}
 
 	bool bShrink = false;
-	const int32 OrgElementMemory = ElementData.GetGeometryAllocatedSize();
+	const int32 OrgElementMemory = IntCastChecked<int32>(ElementData.GetGeometryAllocatedSize());
 
 	if (ElementData.IsPendingLazyGeometryGathering() == true && ElementData.SupportsGatheringGeometrySlices() == false)
 	{
@@ -106,7 +106,7 @@ void FNavigationOctree::DemandLazyDataGathering(FNavigationRelevantData& Element
 		ElementData.ValidateAndShrink();
 	}
 
-	const int32 ElementMemoryChange = ElementData.GetGeometryAllocatedSize() - OrgElementMemory;
+	const int32 ElementMemoryChange = IntCastChecked<int32>(ElementData.GetGeometryAllocatedSize()) - OrgElementMemory;
 	const_cast<FNavigationOctree*>(this)->NodesMemory += ElementMemoryChange;
 	INC_MEMORY_STAT_BY(STAT_Navigation_CollisionTreeMemory, ElementMemoryChange);
 }
diff --git a/Engine/Source/Runtime/NavigationSystem/Private/NavigationPath.cpp b/Engine/Source/Runtime/NavigationSystem/Private/NavigationPath.cpp
--- a/Engine/Source/Runtime/NavigationSystem/Private/NavigationPath.cpp
+++ b/Engine/Source/Runtime/NavigationSystem/Private/NavigationPath.cpp
@@ -410,7 +410,7 @@ void FNavigationPath::DescribeSelfToVisLog(FVisualLogEntry* Snapshot) const
 	Element.Category = LogNavigation.GetCategoryName();
 	Element.SetColor(FColorList::Green);
 	Element.Points.Reserve(NumPathVerts);
-	Element.Thicknes = 3.f;
+	Element.Thicknes = 3;
 	
 	for (int32 VertIdx = 0; VertIdx < NumPathVerts; ++VertIdx)
 	{
diff --git a/Engine/Source/Runtime/NavigationSystem/Private/NavigationSystem.cpp b/Engine/Source/Runtime/NavigationSystem/Private/NavigationSystem.cpp
--- a/Engine/Source/Runtime/NavigationSystem/Private/NavigationSystem.cpp
+++ b/Engine/Source/Runtime/NavigationSystem/Private/NavigationSystem.cpp
@@ -340,7 +340,8 @@ void FNavRegenTimeSliceManager::CalcAverageDeltaTime(uint64 FrameNum)
 
 void FNavRegenTimeSliceManager::CalcTimeSliceDuration(int32 NumTilesToRegen, const TArray<double>& CurrentTileRegenDurations)
 {
-	const float DeltaTimesAverage = (MovingWindowDeltaTime.GetAverage() > 0.f) ? MovingWindowDeltaTime.GetAverage() : (1.f / 30.f); //use default 33 ms
+	const float RawDeltaTimesAverage = FloatCastChecked<float>(MovingWindowDeltaTime.GetAverage(), UE::LWC::DefaultFloatPrecision);
+	const float DeltaTimesAverage = (RawDeltaTimesAverage > 0.f) ? RawDeltaTimesAverage : (1.f / 30.f); //use default 33 ms
 
 	const double TileRegenTimesAverage = (MovingWindowTileRegenTime.GetAverage() > 0.) ? MovingWindowTileRegenTime.GetAverage() : 0.0025; //use default of 2.5 milli secs to regen a full tile
 
@@ -571,7 +572,7 @@ UNavigationSystemV1::UNavigationSystemV1(const FObjectInitializer& ObjectInitial
 	CrowdManagerClass = FSoftObjectPath(TEXT("/Script/AIModule.CrowdManager"));
 
 	// active tiles
-	NextInvokersUpdateTime = 0.f;
+	NextInvokersUpdateTime = 0.;
 	ActiveTilesUpdateInterval = 1.f;
 	bGenerateNavigationOnlyAroundNavigationInvokers = false;
 	DataGatheringMode = ENavDataGatheringModeConfig::Instant;
@@ -4732,7 +4733,7 @@ void UNavigationSystemV1::UnregisterInvoker(AActor& Invoker)
 void UNavigationSystemV1::UpdateInvokers()
 {
 	UWorld* World = GetWorld();
-	const float CurrentTime = World->GetTimeSeconds();
+	const double CurrentTime = World->GetTimeSeconds();
 	if (CurrentTime >= NextInvokersUpdateTime)
 	{
 		InvokerLocations.Reset();
diff --git a/Engine/Source/Runtime/NavigationSystem/Private/NavigationSystemTypes.cpp b/Engine/Source/Runtime/NavigationSystem/Private/NavigationSystemTypes.cpp
--- a/Engine/Source/Runtime/NavigationSystem/Private/NavigationSystemTypes.cpp
+++ b/Engine/Source/Runtime/NavigationSystem/Private/NavigationSystemTypes.cpp
@@ -12,16 +12,16 @@
 
 namespace
 {
-	FORCEINLINE_DEBUGGABLE float RawGeometryFall(const AActor* Querier, const FVector& FallStart, const float FallLimit)
+	FORCEINLINE_DEBUGGABLE FVector::FReal RawGeometryFall(const AActor* Querier, const FVector& FallStart, const FVector::FReal FallLimit)
 	{
-		float FallDownHeight = 0.f;
+		FVector::FReal FallDownHeight = 0.;
 
 		UE_VLOG_SEGMENT(Querier, LogNavigation, Log, FallStart, FallStart + FVector(0, 0, -FallLimit)
 			, FColor::Red, TEXT("TerrainTrace"));
 
 		FCollisionQueryParams TraceParams(SCENE_QUERY_STAT(RawGeometryFall), true, Querier);
 		FHitResult Hit;
-		const bool bHit = Querier->GetWorld()->LineTraceSingleByChannel(Hit, FallStart, FallStart + FVector(0, 0, -FallLimit), ECC_WorldStatic, TraceParams);
+		const bool bHit = Querier->GetWorld()->LineTraceSingleByChannel(Hit, FallStart, FallStart + FVector(0., 0., -FallLimit), ECC_WorldStatic, TraceParams);
 		if (bHit)
 		{
 			UE_VLOG_LOCATION(Querier, LogNavigation, Log, Hit.Location, 15, FColor::Red, TEXT("%s")
@@ -103,9 +103,9 @@ namespace NavigationHelper
 			if (Link.MaxFallDownLength > 0)
 			{
 				const FVector WorldRight = OwnerData.LinkToWorld.TransformPosition(Link.Right);
-				const float FallDownHeight = RawGeometryFall(OwnerData.Actor, WorldRight, Link.MaxFallDownLength);
+				const FVector::FReal FallDownHeight = RawGeometryFall(OwnerData.Actor, WorldRight, Link.MaxFallDownLength);
 
-				if (FallDownHeight > 0.f)
+				if (FallDownHeight > 0.)
 				{
 					// @todo maybe it's a good idea to clear ModifiedLink.MaxFallDownLength here
 					UE_VLOG_SEGMENT(OwnerData.Actor, LogNavigation, Log, WorldRight, WorldRight + FVector(0, 0, -FallDownHeight), FColor::Green, TEXT("FallDownHeight %d"), LinkIndex);
@@ -117,9 +117,9 @@ namespace NavigationHelper
 			if (Link.LeftProjectHeight > 0)
 			{
 				const FVector WorldLeft = OwnerData.LinkToWorld.TransformPosition(Link.Left);
-				const float FallDownHeight = RawGeometryFall(OwnerData.Actor, WorldLeft, Link.LeftProjectHeight);
+				const FVector::FReal FallDownHeight = RawGeometryFall(OwnerData.Actor, WorldLeft, Link.LeftProjectHeight);
 
-				if (FallDownHeight > 0.f)
+				if (FallDownHeight > 0.)
 				{
 					// @todo maybe it's a good idea to clear ModifiedLink.LeftProjectHeight here
 					UE_VLOG_SEGMENT(OwnerData.Actor, LogNavigation, Log, WorldLeft, WorldLeft + FVector(0, 0, -FallDownHeight), FColor::Green, TEXT("LeftProjectHeight %d"), LinkIndex);
@@ -153,17 +153,17 @@ namespace NavigationHelper
 				const FVector WorldRightStart = OwnerData.LinkToWorld.TransformPosition(Link.RightStart);
 				const FVector WorldRightEnd = OwnerData.LinkToWorld.TransformPosition(Link.RightEnd);
 
-				const float FallDownHeightStart = RawGeometryFall(OwnerData.Actor, WorldRightStart, Link.MaxFallDownLength);
-				const float FallDownHeightEnd = RawGeometryFall(OwnerData.Actor, WorldRightEnd, Link.MaxFallDownLength);
+				const FVector::FReal FallDownHeightStart = RawGeometryFall(OwnerData.Actor, WorldRightStart, Link.MaxFallDownLength);
+				const FVector::FReal FallDownHeightEnd = RawGeometryFall(OwnerData.Actor, WorldRightEnd, Link.MaxFallDownLength);
 
-				if (FallDownHeightStart > 0.f)
+				if (FallDownHeightStart > 0.)
 				{
 					// @todo maybe it's a good idea to clear ModifiedLink.MaxFallDownLength here
 					UE_VLOG_SEGMENT(OwnerData.Actor, LogNavigation, Log, WorldRightStart, WorldRightStart + FVector(0, 0, -FallDownHeightStart), FColor::Green, TEXT("FallDownHeightStart %d"), LinkIndex);
 
 					Link.RightStart.Z -= FallDownHeightStart;
 				}
-				if (FallDownHeightEnd > 0.f)
+				if (FallDownHeightEnd > 0.)
 				{
 					// @todo maybe it's a good idea to clear ModifiedLink.MaxFallDownLength here
 					UE_VLOG_SEGMENT(OwnerData.Actor, LogNavigation, Log, WorldRightEnd, WorldRightEnd + FVector(0, 0, -FallDownHeightEnd), FColor::Green, TEXT("FallDownHeightEnd %d"), LinkIndex);
diff --git a/Engine/Source/Runtime/NavigationSystem/Private/NavigationTestingActor.cpp b/Engine/Source/Runtime/NavigationSystem/Private/NavigationTestingActor.cpp
--- a/Engine/Source/Runtime/NavigationSystem/Private/NavigationTestingActor.cpp
+++ b/Engine/Source/Runtime/NavigationSystem/Private/NavigationTestingActor.cpp
@@ -304,7 +304,7 @@ void ANavigationTestingActor::UpdateNavData()
 void ANavigationTestingActor::UpdatePathfinding()
 {
 	PathfindingTime = 0.0f;
-	PathCost = 0.0f;
+	PathCost = 0.;
 	bPathSearchOutOfNodes = false;
 	bPathIsPartial = false;
 	bPathExist = false;
@@ -436,13 +436,13 @@ void ANavigationTestingActor::SearchPathTo(ANavigationTestingActor* Goal)
 	FPathFindingResult Result = NavSys->FindPathSync(NavAgentProps, Query, Mode);
 
 	const double EndTime = FPlatformTime::Seconds();
-	const float Duration = (EndTime - StartTime);
-	PathfindingTime = Duration * 1000000.0f;			// in micro seconds [us]
+	const double Duration = (EndTime - StartTime);
+	PathfindingTime = static_cast<float>(Duration * 1000000.);			// in micro seconds [us]
 	bPathIsPartial = Result.IsPartial();
 	bPathExist = Result.IsSuccessful();
 	bPathSearchOutOfNodes = bPathExist ? Result.Path->DidSearchReachedLimit() : false;
 	LastPath = Result.Path;
-	PathCost = bPathExist ? Result.Path->GetCost() : 0.0f;
+	PathCost = bPathExist ? Result.Path->GetCost() : 0.;
 
 	if (bPathExist)
 	{
diff --git a/Engine/Source/Runtime/NavigationSystem/Public/NavMesh/NavTestRenderingComponent.h b/Engine/Source/Runtime/NavigationSystem/Public/NavMesh/NavTestRenderingComponent.h
--- a/Engine/Source/Runtime/NavigationSystem/Public/NavMesh/NavTestRenderingComponent.h
+++ b/Engine/Source/Runtime/NavigationSystem/Public/NavMesh/NavTestRenderingComponent.h
@@ -38,7 +38,7 @@ public:
 		}
 		FORCEINLINE friend uint32 GetTypeHash(const FNodeDebugData& Other)
 		{
-			return Other.PolyRef;
+			return ::GetTypeHash(Other.PolyRef);
 		}
 	};
 
diff --git a/Engine/Source/Runtime/NavigationSystem/Public/NavMesh/RecastNavMesh.h b/Engine/Source/Runtime/NavigationSystem/Public/NavMesh/RecastNavMesh.h
--- a/Engine/Source/Runtime/NavigationSystem/Public/NavMesh/RecastNavMesh.h
+++ b/Engine/Source/Runtime/NavigationSystem/Public/NavMesh/RecastNavMesh.h
@@ -38,6 +38,10 @@
 //TIME_SLICE_NAV_REGEN must be 0 if we are async rebuilding recast
 #define TIME_SLICE_NAV_REGEN (ALLOW_TIME_SLICE_NAV_REGEN && !RECAST_ASYNC_REBUILDING)
 
+// LWC_TODO_AI Note we are using int32 here for X and Y which does mean that we could overflow the limit of an int for LWCoords
+// unless fairly large tile sizes are used.As WORLD_MAX is currently in flux until we have a better idea of what we are
+// going to be able to support its probably not worth investing time in this potential issue right now.
+
 class FPImplRecastNavMesh;
 class FRecastQueryFilter;
 class INavLinkCustomInterface;
@@ -864,7 +868,7 @@ private:
 	static FNavPolyFlags NavLinkFlag;
 
 	/** Squared draw distance */
-	static float DrawDistanceSq;
+	static FVector::FReal DrawDistanceSq;
 
 	/** MinimumSizeForChaosNavMeshInfluence*/
 	static float MinimumSizeForChaosNavMeshInfluenceSq;
@@ -914,8 +918,8 @@ public:
 	/** broadcast for navmesh updates */
 	FOnNavMeshUpdate OnNavMeshUpdate;
 
-	FORCEINLINE static void SetDrawDistance(float NewDistance) { DrawDistanceSq = NewDistance * NewDistance; }
-	FORCEINLINE static float GetDrawDistanceSq() { return DrawDistanceSq; }
+	FORCEINLINE static void SetDrawDistance(FVector::FReal NewDistance) { DrawDistanceSq = NewDistance * NewDistance; }
+	FORCEINLINE static FVector::FReal GetDrawDistanceSq() { return DrawDistanceSq; }
 
 	FORCEINLINE static void SetMinimumSizeForChaosNavMeshInfluence(float NewSize) { MinimumSizeForChaosNavMeshInfluenceSq = NewSize * NewSize; }
 	FORCEINLINE static float GetMinimumSizeForChaosNavMeshInfluenceSq() { return MinimumSizeForChaosNavMeshInfluenceSq; }
diff --git a/Engine/Source/Runtime/NavigationSystem/Public/NavMesh/RecastNavMeshGenerator.h b/Engine/Source/Runtime/NavigationSystem/Public/NavMesh/RecastNavMeshGenerator.h
--- a/Engine/Source/Runtime/NavigationSystem/Public/NavMesh/RecastNavMeshGenerator.h
+++ b/Engine/Source/Runtime/NavigationSystem/Public/NavMesh/RecastNavMeshGenerator.h
@@ -190,16 +190,16 @@ struct FRcTileBox
 			if (FMath::Modf(MaxAsFloat, &UnusedIntPart) == 0)
 			{
 				// Return the lower tile
-				return FMath::Max(FMath::FloorToInt(MaxAsFloat) - 1, MinCoord);
+				return FMath::Max(IntCastChecked<int32>(FMath::FloorToInt(MaxAsFloat) - 1), MinCoord);
 			}
 			// Otherwise use default behaviour
-			return FMath::FloorToInt(MaxAsFloat);
+			return IntCastChecked<int32>(FMath::FloorToInt(MaxAsFloat));
 		};
 
 		const FBox RcAreaBounds = Unreal2RecastBox(UnrealBounds);
-		XMin = FMath::FloorToInt((RcAreaBounds.Min.X - RcNavMeshOrigin.X) / TileSizeInWorldUnits);
+		XMin = IntCastChecked<int32>(FMath::FloorToInt((RcAreaBounds.Min.X - RcNavMeshOrigin.X) / TileSizeInWorldUnits));
 		XMax = CalcMaxCoordExclusive((RcAreaBounds.Max.X - RcNavMeshOrigin.X) / TileSizeInWorldUnits, XMin);
-		YMin = FMath::FloorToInt((RcAreaBounds.Min.Z - RcNavMeshOrigin.Z) / TileSizeInWorldUnits);
+		YMin = IntCastChecked<int32>(FMath::FloorToInt((RcAreaBounds.Min.Z - RcNavMeshOrigin.Z) / TileSizeInWorldUnits));
 		YMax = CalcMaxCoordExclusive((RcAreaBounds.Max.Z - RcNavMeshOrigin.Z) / TileSizeInWorldUnits, YMin);
 	}
 
diff --git a/Engine/Source/Runtime/NavigationSystem/Public/NavigationData.h b/Engine/Source/Runtime/NavigationSystem/Public/NavigationData.h
--- a/Engine/Source/Runtime/NavigationSystem/Public/NavigationData.h
+++ b/Engine/Source/Runtime/NavigationSystem/Public/NavigationData.h
@@ -214,8 +214,8 @@ struct NAVIGATIONSYSTEM_API FNavigationPath : public TSharedFromThis<FNavigation
 		ObserverDelegate.Broadcast(this, PathUpdateTypeToPathEvent[uint8(UpdateType)]);
 	}
 
-	FORCEINLINE float GetTimeStamp() const { return LastUpdateTimeStamp; }
-	FORCEINLINE void SetTimeStamp(float TimeStamp) { LastUpdateTimeStamp = TimeStamp; }
+	FORCEINLINE double GetTimeStamp() const { return LastUpdateTimeStamp; }
+	FORCEINLINE void SetTimeStamp(double TimeStamp) { LastUpdateTimeStamp = TimeStamp; }
 
 	void Invalidate();
 	void RePathFailed();
@@ -336,7 +336,7 @@ public:
 	FVector GetLastRepathGoalLocation() const { return GoalActorLastLocation; }
 	void UpdateLastRepathGoalLocation();
 	
-	float GetLastUpdateTime() const { return LastUpdateTimeStamp; }
+	double GetLastUpdateTime() const { return LastUpdateTimeStamp; }
 	float GetGoalActorTetherDistance() const { return FMath::Sqrt(GoalActorLocationTetherDistanceSq); }
 
 	/** if enabled path will request recalculation if it gets invalidated due to a change to underlying navigation */
@@ -461,7 +461,7 @@ protected:
 	FPathFindingQueryData PathFindingQueryData;
 
 	/** gets set during path creation and on subsequent path's updates */
-	float LastUpdateTimeStamp;
+	double LastUpdateTimeStamp;
 
 private:
 	/* if GoalActor is set this is the distance we'll try to keep GoalActor from end of path. If GoalActor
@@ -907,7 +907,7 @@ public:
 	 *	@param Location is expressed in WorldSpace, navigation data is responsible for tansforming if need be */
 	virtual bool DoesNodeContainLocation(NavNodeRef NodeRef, const FVector& WorldSpaceLocation) const PURE_VIRTUAL(ANavigationData::DoesNodeContainLocation, return false;);
 
-	float GetWorldTimeStamp() const;
+	double GetWorldTimeStamp() const;
 
 	//----------------------------------------------------------------------//
 	// Areas
diff --git a/Engine/Source/Runtime/NavigationSystem/Public/NavigationOctree.h b/Engine/Source/Runtime/NavigationSystem/Public/NavigationOctree.h
--- a/Engine/Source/Runtime/NavigationSystem/Public/NavigationOctree.h
+++ b/Engine/Source/Runtime/NavigationSystem/Public/NavigationOctree.h
@@ -159,7 +159,7 @@ public:
 		INC_MEMORY_STAT_BY(STAT_Navigation_CollisionTreeMemory, OctreeSizeBytes);
 	}
 
-	FNavigationOctree(const FVector& Origin, float Radius);
+	FNavigationOctree(const FVector& Origin, FVector::FReal Radius);
 	virtual ~FNavigationOctree();
 
 	/** Add new node and fill it with navigation export data */
diff --git a/Engine/Source/Runtime/NavigationSystem/Public/NavigationSystem.h b/Engine/Source/Runtime/NavigationSystem/Public/NavigationSystem.h
--- a/Engine/Source/Runtime/NavigationSystem/Public/NavigationSystem.h
+++ b/Engine/Source/Runtime/NavigationSystem/Public/NavigationSystem.h
@@ -397,7 +397,7 @@ private:
 	/** Contains pre-digested and cached invokers' info. Generated by UpdateInvokers */
 	TArray<FNavigationInvokerRaw> InvokerLocations;
 
-	float NextInvokersUpdateTime;
+	double NextInvokersUpdateTime;
 	void UpdateInvokers();
 
 	void DirtyTilesInBuildBounds();
diff --git a/Engine/Source/Runtime/NavigationSystem/Public/NavigationTestingActor.h b/Engine/Source/Runtime/NavigationSystem/Public/NavigationTestingActor.h
--- a/Engine/Source/Runtime/NavigationSystem/Public/NavigationTestingActor.h
+++ b/Engine/Source/Runtime/NavigationSystem/Public/NavigationTestingActor.h
@@ -147,7 +147,7 @@ public:
 	float PathfindingTime;
 
 	UPROPERTY(transient, VisibleAnywhere, BlueprintReadOnly, Category=PathfindingStatus)
-	float PathCost;
+	double PathCost;
 
 	UPROPERTY(transient, VisibleAnywhere, BlueprintReadOnly, Category=PathfindingStatus)
 	int32 PathfindingSteps;
diff --git a/Engine/Source/Runtime/Navmesh/Private/Detour/DetourNavMesh.cpp b/Engine/Source/Runtime/Navmesh/Private/Detour/DetourNavMesh.cpp
--- a/Engine/Source/Runtime/Navmesh/Private/Detour/DetourNavMesh.cpp
+++ b/Engine/Source/Runtime/Navmesh/Private/Detour/DetourNavMesh.cpp
@@ -1198,7 +1198,7 @@ void dtNavMesh::connectExtOffMeshLinks(dtMeshTile* tile, dtMeshTile* target, int
 		linkOffMeshHelper(target, targetCon->poly, tile, landPolyIdx, linkSide, 1);
 
 		// Link target poly to off-mesh connection.
-		linkSide = (side == -1 ? DT_CONNECTION_INTERNAL : side) | DT_LINK_FLAG_OFFMESH_CON | biDirFlag;
+		linkSide = (unsigned char)(side == -1 ? DT_CONNECTION_INTERNAL : side) | DT_LINK_FLAG_OFFMESH_CON | biDirFlag;
 		if (tile != target)
 		{
 			linkSide &= ~DT_CONNECTION_INTERNAL;
diff --git a/Engine/Source/Runtime/Navmesh/Private/Detour/DetourNavMeshBuilder.cpp b/Engine/Source/Runtime/Navmesh/Private/Detour/DetourNavMeshBuilder.cpp
--- a/Engine/Source/Runtime/Navmesh/Private/Detour/DetourNavMeshBuilder.cpp
+++ b/Engine/Source/Runtime/Navmesh/Private/Detour/DetourNavMeshBuilder.cpp
@@ -707,7 +707,7 @@ bool dtCreateNavMeshData(dtNavMeshCreateParams* params, unsigned char** outData,
 			const int nv = navPolys[i].vertCount;
 			dtl.vertBase = (unsigned int)vbase;
 			dtl.vertCount = (unsigned char)(ndv-nv);
-			dtl.triBase = (unsigned int)params->detailMeshes[i*4+2];
+			dtl.triBase = (unsigned short)params->detailMeshes[i*4+2];
 			dtl.triCount = (unsigned char)params->detailMeshes[i*4+3];
 			// Copy vertices except the first 'nv' verts which are equal to nav poly verts.
 			if (ndv-nv)
@@ -729,7 +729,7 @@ bool dtCreateNavMeshData(dtNavMeshCreateParams* params, unsigned char** outData,
 			const int nv = navPolys[i].vertCount;
 			dtl.vertBase = 0;
 			dtl.vertCount = 0;
-			dtl.triBase = (unsigned int)tbase;
+			dtl.triBase = (unsigned short)tbase;
 			dtl.triCount = (unsigned char)(nv-2);
 			// Triangulate polygon (local indices).
 			for (int j = 2; j < nv; ++j)
@@ -1197,7 +1197,7 @@ bool dtTransformTileData(unsigned char* data, const int dataSize, const int offs
 
 		const unsigned short realSide = tile.offMeshCons[j].side & DT_LINK_FLAG_SIDE_MASK;
 		const unsigned short newSide = (realSide + (2*rot)) % 8; //rot [0..3], side [0,1,2,3,4,5,6,7]
-		tile.offMeshCons[j].side = (tile.offMeshCons[j].side & ~DT_LINK_FLAG_SIDE_MASK) | newSide;
+		tile.offMeshCons[j].side = (unsigned char)((tile.offMeshCons[j].side & ~DT_LINK_FLAG_SIDE_MASK) | newSide);
 	}
 
 #if WITH_NAVMESH_SEGMENT_LINKS
@@ -1236,7 +1236,7 @@ void dtComputeTileOffsetFromRotation(const dtReal* position, const dtReal* rotat
 
 	dtReal RcTilePosChangedVector[3];
 	dtVsub(RcTilePosChangedVector, newRelativeTilePos, relativeTilePos);
-	deltaX = RcTilePosChangedVector[0] / tileWidth;
-	deltaY = RcTilePosChangedVector[2] / tileHeight;
+	deltaX = (int)(RcTilePosChangedVector[0] / tileWidth);
+	deltaY = (int)(RcTilePosChangedVector[2] / tileHeight);
 }
 // @UE END
diff --git a/Engine/Source/Runtime/Navmesh/Private/Detour/DetourNavMeshQuery.cpp b/Engine/Source/Runtime/Navmesh/Private/Detour/DetourNavMeshQuery.cpp
--- a/Engine/Source/Runtime/Navmesh/Private/Detour/DetourNavMeshQuery.cpp
+++ b/Engine/Source/Runtime/Navmesh/Private/Detour/DetourNavMeshQuery.cpp
@@ -4750,7 +4750,7 @@ dtStatus dtNavMeshQuery::findWallsOverlappingShape(dtPolyRef startRef, const dtR
 }
 
 dtStatus dtNavMeshQuery::findWallsAroundPath(const dtPolyRef* path, const int pathCount, const dtReal* searchAreaPoly, const int searchAreaPolyCount,
-												   const float maxAreaEnterCost, const dtQueryFilter* filter,
+												   const dtReal maxAreaEnterCost, const dtQueryFilter* filter,
 												   dtPolyRef* neiRefs, int* neiCount, const int maxNei,
 												   dtReal* resultWalls, dtPolyRef* resultRefs, int* resultCount, const int maxResult) const
 {
@@ -4844,7 +4844,7 @@ dtStatus dtNavMeshQuery::findWallsAroundPath(const dtPolyRef* path, const int pa
 				continue;
 
 			// Do not advance if the polygon is excluded by the filter, or if the area enter cost is too high.
-			const float enterCost = filter->getAreaFixedCost(neighbourPoly->getArea());
+			const dtReal enterCost = filter->getAreaFixedCost(neighbourPoly->getArea());
 
 			if (!filter->passFilter(neighbourRef, neighbourTile, neighbourPoly)
 				|| !passLinkFilterByRef(neighbourTile, neighbourRef)
diff --git a/Engine/Source/Runtime/Navmesh/Private/DetourCrowd/DetourCrowd.cpp b/Engine/Source/Runtime/Navmesh/Private/DetourCrowd/DetourCrowd.cpp
--- a/Engine/Source/Runtime/Navmesh/Private/DetourCrowd/DetourCrowd.cpp
+++ b/Engine/Source/Runtime/Navmesh/Private/DetourCrowd/DetourCrowd.cpp
@@ -542,7 +542,7 @@ bool dtCrowd::updateAgentFilter(const int idx, const dtQueryFilter* filter)
 		const bool bMatching = filter->equals(m_filters[i]);
 		if (bMatching)
 		{
-			m_agents[idx].params.filter = i;
+			m_agents[idx].params.filter = (unsigned char)i;
 			return true;
 		}
 	}
@@ -562,7 +562,7 @@ bool dtCrowd::updateAgentFilter(const int idx, const dtQueryFilter* filter)
 		if (used[i] == 0)
 		{
 			m_filters[i].copyFrom(filter);
-			m_agents[idx].params.filter = i;
+			m_agents[idx].params.filter = (unsigned char)i;
 			return true;
 		}
 	}
@@ -1393,7 +1393,7 @@ void dtCrowd::updateStepNextMovePoint(const dtReal dt, dtCrowdAgentDebugInfo* de
 		if (overOffmeshConnection(ag, triggerRadius))
 		{
 			// Prepare to off-mesh connection.
-			const int idx = ag - m_agents;
+			const int idx = (int)(ag - m_agents);
 			dtCrowdAgentAnimation* anim = &m_agentAnims[idx];
 			m_navquery->updateLinkFilter(ag->params.linkFilter.Get());
 
diff --git a/Engine/Source/Runtime/Navmesh/Private/DetourTileCache/DetourTileCache.cpp b/Engine/Source/Runtime/Navmesh/Private/DetourTileCache/DetourTileCache.cpp
--- a/Engine/Source/Runtime/Navmesh/Private/DetourTileCache/DetourTileCache.cpp
+++ b/Engine/Source/Runtime/Navmesh/Private/DetourTileCache/DetourTileCache.cpp
@@ -266,14 +266,14 @@ dtCompressedTile* dtTileCache::getTileAt(const int tx, const int ty, const int t
 dtCompressedTileRef dtTileCache::getTileRef(const dtCompressedTile* tile) const
 {
 	if (!tile) return 0;
-	const unsigned int it = tile - m_tiles;
+	const unsigned int it = (unsigned int)(tile - m_tiles);
 	return (dtCompressedTileRef)encodeTileId(tile->salt, it);
 }
 
 dtObstacleRef dtTileCache::getObstacleRef(const dtTileCacheObstacle* ob) const
 {
 	if (!ob) return 0;
-	const unsigned int idx = ob - m_obstacles;
+	const unsigned int idx = (unsigned int)(ob - m_obstacles);
 	return encodeObstacleId(ob->salt, idx);
 }
 
diff --git a/Engine/Source/Runtime/Navmesh/Private/DetourTileCache/DetourTileCacheBuilder.cpp b/Engine/Source/Runtime/Navmesh/Private/DetourTileCache/DetourTileCacheBuilder.cpp
--- a/Engine/Source/Runtime/Navmesh/Private/DetourTileCache/DetourTileCacheBuilder.cpp
+++ b/Engine/Source/Runtime/Navmesh/Private/DetourTileCache/DetourTileCacheBuilder.cpp
@@ -23,6 +23,7 @@
 #include "Detour/DetourCommon.h"
 #include "Detour/DetourAssert.h"
 #include "Stats/Stats.h"
+#include <limits>
 
 static const int MAX_VERTS_PER_POLY = 6;	// TODO: use the DT_VERTS_PER_POLYGON
 static const int MAX_REM_EDGES = 48;		// TODO: make this an expression.
@@ -999,7 +1000,8 @@ dtStatus dtBuildTileCacheContours(dtTileCacheAlloc* alloc, dtTileCacheLayer& lay
 
 			if (bCanAddCluster)
 			{
-				clusters.regMap[i] = newClusterId;
+				dtAssert(newClusterId < std::numeric_limits<unsigned short>::max());
+				clusters.regMap[i] = (unsigned short)newClusterId;
 				clusters.nclusters++;
 				bCanAddCluster = false;
 			}
@@ -1030,7 +1032,8 @@ dtStatus dtBuildTileCacheContours(dtTileCacheAlloc* alloc, dtTileCacheLayer& lay
 						continue;
 					}
 
-					clusters.regMap[cont.reg] = newClusterId;
+					dtAssert(newClusterId < std::numeric_limits<unsigned short>::max());
+					clusters.regMap[cont.reg] = (unsigned short)newClusterId;
 					for (int j = 0; j < nlinks[ic]; j++)
 					{
 						unsigned short neiReg = (unsigned short)(links[linksBase[ic] + j]);
diff --git a/Engine/Source/Runtime/Navmesh/Private/Recast/RecastArea.cpp b/Engine/Source/Runtime/Navmesh/Private/Recast/RecastArea.cpp
--- a/Engine/Source/Runtime/Navmesh/Private/Recast/RecastArea.cpp
+++ b/Engine/Source/Runtime/Navmesh/Private/Recast/RecastArea.cpp
@@ -237,7 +237,7 @@ void SeedWalkableAreasErode(int w, int h, int lowSpanHeight, int lowSpanSeed, co
 				else if (chf.spans[i].h < lowSpanHeight)
 				{
 					dist[i] = 0;
-					seed[i] = lowSpanSeed;
+					seed[i] = (unsigned char)lowSpanSeed;
 				}
 				else
 				{
@@ -290,7 +290,7 @@ void SeedWalkableAreasErodeFiltered(int w, int h, int lowSpanHeight, int lowSpan
 
 					if (s.y < NextAllowedLowBase)
 					{
-						seed[i] = lowSpanSeed;
+						seed[i] = (unsigned char)lowSpanSeed;
 
 						// next low span is allowed after full agent height (or valid area from condition below)
 						NextAllowedLowBase = rcMax(0, s.y - lowSpanHeight);
diff --git a/Engine/Source/Runtime/Navmesh/Private/Recast/RecastRasterization.cpp b/Engine/Source/Runtime/Navmesh/Private/Recast/RecastRasterization.cpp
--- a/Engine/Source/Runtime/Navmesh/Private/Recast/RecastRasterization.cpp
+++ b/Engine/Source/Runtime/Navmesh/Private/Recast/RecastRasterization.cpp
@@ -562,7 +562,7 @@ static void rasterizeTri(const rcReal* v0, const rcReal* v1, const rcReal* v2,
 				int loop0 = intMax(edge0 + 1, y0);
 				int loop1 = intMin(edge1, y1_edge);
 
-				unsigned char edgeBits = (edge << 4) | (othervert << 2) | basevert;
+				unsigned char edgeBits = (unsigned char)((edge << 4) | (othervert << 2) | basevert);
 				for (int y = loop0; y <= loop1; y++)
 				{
 					int HitIndex = !!hfEdgeHits[y].Hits[0];
@@ -732,7 +732,7 @@ static void rasterizeTri(const rcReal* v0, const rcReal* v1, const rcReal* v2,
 				int loop0 = intMax(edge0 + 1, y0);
 				int loop1 = intMin(edge1, y1_edge);
 
-				unsigned char edgeBits = (edge << 4) | (othervert << 2) | basevert;
+				unsigned char edgeBits = (unsigned char)((edge << 4) | (othervert << 2) | basevert);
 				for (int y = loop0; y <= loop1; y++)
 				{
 					int HitIndex = !!hfEdgeHits[y].Hits[0];
@@ -864,8 +864,8 @@ static void rasterizeTri(const rcReal* v0, const rcReal* v1, const rcReal* v2,
 				// Skip the span if it is outside the heightfield bbox
 				if (smin >= RC_SPAN_MAX_HEIGHT || smax < 0) continue;
 
-				smin = intMax(smin, 0);
-				smax = intMin(intMax(smax,smin+1), RC_SPAN_MAX_HEIGHT);
+				smin = (short)intMax(smin, 0);
+				smax = (short)intMin(intMax(smax,smin+1), RC_SPAN_MAX_HEIGHT);
 				const int projectSpanToBottom = rasterizationMasks != nullptr ? (projectTriToBottom & rasterizationMasks[x+y*w]) : projectTriToBottom; //UE
 				if (projectSpanToBottom) //UE
 				{
diff --git a/Engine/Source/Runtime/Navmesh/Public/Detour/DetourNavMeshQuery.h b/Engine/Source/Runtime/Navmesh/Public/Detour/DetourNavMeshQuery.h
--- a/Engine/Source/Runtime/Navmesh/Public/Detour/DetourNavMeshQuery.h
+++ b/Engine/Source/Runtime/Navmesh/Public/Detour/DetourNavMeshQuery.h
@@ -567,7 +567,7 @@ public:
 
 	/// [UE] Finds the wall segments that overlap the polygon shape.
 	dtStatus findWallsAroundPath(const dtPolyRef* path, const int pathCount, const dtReal* searchAreaPoly, const int searchAreaPolyCount,
-									   const float maxAreaEnterCost, const dtQueryFilter* filter,
+									   const dtReal maxAreaEnterCost, const dtQueryFilter* filter,
 									   dtPolyRef* neiRefs, int* neiCount, const int maxNei,
 									   dtReal* resultWalls, dtPolyRef* resultRefs, int* resultCount, const int maxResult) const;
 	
