Refactor AISense_Sight to extract duplicated code for query registration

diff --git a/Engine/Source/Runtime/AIModule/Classes/Perception/AISense_Sight.h b/Engine/Source/Runtime/AIModule/Classes/Perception/AISense_Sight.h
--- a/Engine/Source/Runtime/AIModule/Classes/Perception/AISense_Sight.h
+++ b/Engine/Source/Runtime/AIModule/Classes/Perception/AISense_Sight.h
@@ -316,6 +316,7 @@ protected:
 	bool RegisterTarget(AActor& TargetActor, const TFunction<void(FAISightQuery&)>& OnAddedFunc = nullptr);
 
 	float CalcQueryImportance(const FPerceptionListener& Listener, const FVector& TargetLocation, const float SightRadiusSq) const;
+	bool RegisterNewQuery(const FPerceptionListener& Listener, const IGenericTeamAgentInterface* ListenersTeamAgent, const AActor& TargetActor, const FAISightTarget::FTargetId& TargetId, const FVector& TargetLocation, const FDigestedSightProperties& PropDigest, const TFunction<void(FAISightQuery&)>& OnAddedFunc);
 
 	// Deprecated methods
 public:
diff --git a/Engine/Source/Runtime/AIModule/Private/Perception/AISense_Sight.cpp b/Engine/Source/Runtime/AIModule/Private/Perception/AISense_Sight.cpp
--- a/Engine/Source/Runtime/AIModule/Private/Perception/AISense_Sight.cpp
+++ b/Engine/Source/Runtime/AIModule/Private/Perception/AISense_Sight.cpp
@@ -781,7 +781,7 @@ bool UAISense_Sight::RegisterTarget(AActor& TargetActor, const TFunction<void(FA
 
 	// set/update data
 	SightTarget->TeamId = FGenericTeamId::GetTeamIdentifier(&TargetActor);
-	
+
 	// generate all pairs and add them to current Sight Queries
 	bool bNewQueriesAdded = false;
 	AIPerception::FListenerMap& ListenersMap = *GetListeners();
@@ -790,31 +790,16 @@ bool UAISense_Sight::RegisterTarget(AActor& TargetActor, const TFunction<void(FA
 	for (AIPerception::FListenerMap::TConstIterator ItListener(ListenersMap); ItListener; ++ItListener)
 	{
 		const FPerceptionListener& Listener = ItListener->Value;
-		const IGenericTeamAgentInterface* ListenersTeamAgent = Listener.GetTeamAgent();
+		if (!Listener.HasSense(GetSenseID()) || Listener.GetBodyActor() == &TargetActor)
+		{
+			continue;
+		}
 
-		if (Listener.HasSense(GetSenseID()) && Listener.GetBodyActor() != &TargetActor)
+		const FDigestedSightProperties& PropDigest = DigestedProperties[Listener.GetListenerID()];
+		const IGenericTeamAgentInterface* ListenersTeamAgent = Listener.GetTeamAgent();
+		if (RegisterNewQuery(Listener, ListenersTeamAgent, TargetActor, SightTarget->TargetId, TargetLocation, PropDigest, OnAddedFunc))
 		{
-			const FDigestedSightProperties& PropDigest = DigestedProperties[Listener.GetListenerID()];
-			if (FAISenseAffiliationFilter::ShouldSenseTeam(ListenersTeamAgent, TargetActor, PropDigest.AffiliationFlags))
-			{
-				// create a sight query		
-				const float Importance = CalcQueryImportance(ItListener->Value, TargetLocation, PropDigest.SightRadiusSq);
-				const bool bInRange = Importance > 0.0f;
-				if (!bInRange)
-				{
-					bSightQueriesOutOfRangeDirty = true;
-				}
-				FAISightQuery& AddedQuery = bInRange ? SightQueriesInRange.AddDefaulted_GetRef() : SightQueriesOutOfRange.AddDefaulted_GetRef();
-				AddedQuery.ObserverId = ItListener->Key;
-				AddedQuery.TargetId = SightTarget->TargetId;
-				AddedQuery.Importance = Importance;
-				
-				if (OnAddedFunc)
-				{
-					OnAddedFunc(AddedQuery);
-				}
-				bNewQueriesAdded = true;
-			}
+			bNewQueriesAdded = true;
 		}
 	}
 
@@ -849,29 +834,14 @@ void UAISense_Sight::GenerateQueriesForListener(const FPerceptionListener& Liste
 	for (FTargetsContainer::TConstIterator ItTarget(ObservedTargets); ItTarget; ++ItTarget)
 	{
 		const AActor* TargetActor = ItTarget->Value.GetTargetActor();
-		if (TargetActor == NULL || TargetActor == Avatar)
+		if (TargetActor == nullptr || TargetActor == Avatar)
 		{
 			continue;
 		}
 
-		if (FAISenseAffiliationFilter::ShouldSenseTeam(ListenersTeamAgent, *TargetActor, PropertyDigest.AffiliationFlags))
+		const FVector TargetLocation = TargetActor->GetActorLocation();
+		if (RegisterNewQuery(Listener, ListenersTeamAgent, *TargetActor, ItTarget->Key, TargetLocation, PropertyDigest, OnAddedFunc))
 		{
-			// create a sight query		
-			const float Importance = CalcQueryImportance(Listener, ItTarget->Value.GetLocationSimple(), PropertyDigest.SightRadiusSq);
-			const bool bInRange = Importance > 0.0f;
-			if (!bInRange)
-			{
-				bSightQueriesOutOfRangeDirty = true;
-			}
-			FAISightQuery& AddedQuery = bInRange ? SightQueriesInRange.AddDefaulted_GetRef() : SightQueriesOutOfRange.AddDefaulted_GetRef();
-			AddedQuery.ObserverId = Listener.GetListenerID();
-			AddedQuery.TargetId = ItTarget->Key;
-			AddedQuery.Importance = Importance;
-
-			if (OnAddedFunc)
-			{
-				OnAddedFunc(AddedQuery);
-			}
 			bNewQueriesAdded = true;
 		}
 	}
@@ -883,6 +853,34 @@ void UAISense_Sight::GenerateQueriesForListener(const FPerceptionListener& Liste
 	}
 }
 
+bool UAISense_Sight::RegisterNewQuery(const FPerceptionListener& Listener, const IGenericTeamAgentInterface* ListenersTeamAgent, const AActor& TargetActor, const FAISightTarget::FTargetId& TargetId, const FVector& TargetLocation, const FDigestedSightProperties& PropDigest, const TFunction<void(FAISightQuery&)>& OnAddedFunc)
+{
+	if (!FAISenseAffiliationFilter::ShouldSenseTeam(ListenersTeamAgent, TargetActor, PropDigest.AffiliationFlags))
+	{
+		return false;
+	}
+
+	// create a sight query
+	const float Importance = CalcQueryImportance(Listener, TargetLocation, PropDigest.SightRadiusSq);
+	const bool bInRange = Importance > 0.0f;
+	if (!bInRange)
+	{
+		bSightQueriesOutOfRangeDirty = true;
+	}
+
+	FAISightQuery& AddedQuery = bInRange ? SightQueriesInRange.AddDefaulted_GetRef() : SightQueriesOutOfRange.AddDefaulted_GetRef();
+	AddedQuery.ObserverId = Listener.GetListenerID();
+	AddedQuery.TargetId = TargetId;
+	AddedQuery.Importance = Importance;
+
+	if (OnAddedFunc)
+	{
+		OnAddedFunc(AddedQuery);
+	}
+
+	return true;
+}
+
 void UAISense_Sight::OnListenerUpdateImpl(const FPerceptionListener& UpdatedListener)
 {
 	SCOPE_CYCLE_COUNTER(STAT_AI_Sense_Sight_ListenerUpdate);
