[AIModule]
* Moved Blackboard struct out into its own file
* Moved Blackboard key enums out into their own file

diff --git a/Engine/Source/Runtime/AIModule/Classes/BehaviorTree/BehaviorTreeTypes.h b/Engine/Source/Runtime/AIModule/Classes/BehaviorTree/BehaviorTreeTypes.h
--- a/Engine/Source/Runtime/AIModule/Classes/BehaviorTree/BehaviorTreeTypes.h
+++ b/Engine/Source/Runtime/AIModule/Classes/BehaviorTree/BehaviorTreeTypes.h
@@ -2,6 +2,7 @@
 
 #pragma once
 
+#include "Blackboard.h"
 #include "CoreMinimal.h"
 #include "Stats/Stats.h"
 #include "UObject/ObjectMacros.h"
@@ -43,38 +44,6 @@ DECLARE_DWORD_ACCUMULATOR_STAT_EXTERN(TEXT("Num Templates"),STAT_AI_BehaviorTree
 DECLARE_DWORD_ACCUMULATOR_STAT_EXTERN(TEXT("Num Instances"),STAT_AI_BehaviorTree_NumInstances,STATGROUP_AIBehaviorTree, );
 DECLARE_MEMORY_STAT_EXTERN(TEXT("Instance memory"),STAT_AI_BehaviorTree_InstanceMemory,STATGROUP_AIBehaviorTree, AIMODULE_API);
 
-namespace FBlackboard
-{
-	const FName KeySelf = TEXT("SelfActor");
-
-#ifdef AI_BLACKBOARD_KEY_SIZE_8
-	// this is the legacy BB key size. Add AI_BLACKBOARD_KEY_SIZE_8 to your *.target.cs to enable it.
-	using FKey = uint8;
-	inline constexpr FKey InvalidKey = FKey(-1);
-#else
-	//the default BB key size is now 16
-	struct FKey
-	{
-		constexpr FKey() = default;
-		FKey(int32 InKey) {Key = IntCastChecked<uint16>(InKey); }
-		constexpr FKey(uint16 InKey) : Key(InKey) {}
-		constexpr FKey(uint8 InKey) = delete;
-		constexpr operator int32() const { return Key; }
-		constexpr operator uint16() const { return Key; }
-		constexpr operator uint8() const = delete;
-		constexpr bool operator==(const FKey& Other) const {return Key == Other.Key;}
-		constexpr bool operator!=(const FKey& Other) const {return Key != Other.Key;}
-	private:
-		friend uint32 GetTypeHash(const FKey& Key);
-		uint16 Key = static_cast<uint16>(-1);
-	};
-
-	inline constexpr FKey InvalidKey = FKey();
-
-	inline uint32 GetTypeHash(const FKey& Key) { return ::GetTypeHash(Key.Key);}
-#endif
-}
-
 #ifndef AI_BLACKBOARD_KEY_SIZE_8
 template <> struct TIsValidVariadicFunctionArg<FBlackboard::FKey>
 {
diff --git a/Engine/Source/Runtime/AIModule/Classes/BehaviorTree/Blackboard.h b/Engine/Source/Runtime/AIModule/Classes/BehaviorTree/Blackboard.h
new file mode 100644
--- /dev/null
+++ b/Engine/Source/Runtime/AIModule/Classes/BehaviorTree/Blackboard.h
@@ -0,0 +1,37 @@
+// Copyright Epic Games, Inc. All Rights Reserved.
+
+#pragma once
+
+#include "UObject/NameTypes.h"
+
+namespace FBlackboard
+{
+	const FName KeySelf = TEXT("SelfActor");
+
+#ifdef AI_BLACKBOARD_KEY_SIZE_8
+	// this is the legacy BB key size. Add AI_BLACKBOARD_KEY_SIZE_8 to your *.target.cs to enable it.
+	using FKey = uint8;
+	inline constexpr FKey InvalidKey = FKey(-1);
+#else
+	//the default BB key size is now 16
+	struct FKey
+	{
+		constexpr FKey() = default;
+		FKey(int32 InKey) {Key = IntCastChecked<uint16>(InKey); }
+		constexpr FKey(uint16 InKey) : Key(InKey) {}
+		constexpr FKey(uint8 InKey) = delete;
+		constexpr operator int32() const { return Key; }
+		constexpr operator uint16() const { return Key; }
+		constexpr operator uint8() const = delete;
+		constexpr bool operator==(const FKey& Other) const {return Key == Other.Key;}
+		constexpr bool operator!=(const FKey& Other) const {return Key != Other.Key;}
+	private:
+		friend uint32 GetTypeHash(const FKey& Key);
+		uint16 Key = static_cast<uint16>(-1);
+	};
+
+	inline constexpr FKey InvalidKey = FKey();
+
+	inline uint32 GetTypeHash(const FKey& Key) { return ::GetTypeHash(Key.Key);}
+#endif
+}
\ No newline at end of file
diff --git a/Engine/Source/Runtime/AIModule/Classes/BehaviorTree/Blackboard/BlackboardKeyEnums.h b/Engine/Source/Runtime/AIModule/Classes/BehaviorTree/Blackboard/BlackboardKeyEnums.h
new file mode 100644
--- /dev/null
+++ b/Engine/Source/Runtime/AIModule/Classes/BehaviorTree/Blackboard/BlackboardKeyEnums.h
@@ -0,0 +1,64 @@
+// Copyright Epic Games, Inc. All Rights Reserved.
+
+#pragma once
+
+#include "CoreTypes.h"
+
+
+namespace EBlackboardCompare
+{
+	enum Type
+	{
+		Less = -1, 	
+		Equal = 0, 
+		Greater = 1,
+
+		NotEqual = 1,	// important, do not change
+	};
+}
+
+namespace EBlackboardKeyOperation
+{
+	enum Type
+	{
+		Basic,
+		Arithmetic,
+		Text,
+	};
+}
+
+UENUM()
+namespace EBasicKeyOperation
+{
+	enum Type
+	{
+		Set				UMETA(DisplayName="Is Set"),
+		NotSet			UMETA(DisplayName="Is Not Set"),
+	};
+}
+
+UENUM()
+namespace EArithmeticKeyOperation
+{
+	enum Type
+	{
+		Equal			UMETA(DisplayName="Is Equal To"),
+		NotEqual		UMETA(DisplayName="Is Not Equal To"),
+		Less			UMETA(DisplayName="Is Less Than"),
+		LessOrEqual		UMETA(DisplayName="Is Less Than Or Equal To"),
+		Greater			UMETA(DisplayName="Is Greater Than"),
+		GreaterOrEqual	UMETA(DisplayName="Is Greater Than Or Equal To"),
+	};
+}
+
+UENUM()
+namespace ETextKeyOperation
+{
+	enum Type
+	{
+		Equal			UMETA(DisplayName="Is Equal To"),
+		NotEqual		UMETA(DisplayName="Is Not Equal To"),
+		Contain			UMETA(DisplayName="Contains"),
+		NotContain		UMETA(DisplayName="Not Contains"),
+	};
+}
diff --git a/Engine/Source/Runtime/AIModule/Classes/BehaviorTree/Blackboard/BlackboardKeyType.h b/Engine/Source/Runtime/AIModule/Classes/BehaviorTree/Blackboard/BlackboardKeyType.h
--- a/Engine/Source/Runtime/AIModule/Classes/BehaviorTree/Blackboard/BlackboardKeyType.h
+++ b/Engine/Source/Runtime/AIModule/Classes/BehaviorTree/Blackboard/BlackboardKeyType.h
@@ -6,68 +6,11 @@
 #include "UObject/Object.h"
 #include "InputCoreTypes.h"
 #include "BehaviorTree/BehaviorTreeTypes.h"
+#include "BlackboardKeyEnums.h"
 #include "BlackboardKeyType.generated.h"
 
 class UBlackboardComponent;
 
-namespace EBlackboardCompare
-{
-	enum Type
-	{
-		Less = -1, 	
-		Equal = 0, 
-		Greater = 1,
-
-		NotEqual = 1,	// important, do not change
-	};
-}
-
-namespace EBlackboardKeyOperation
-{
-	enum Type
-	{
-		Basic,
-		Arithmetic,
-		Text,
-	};
-}
-
-UENUM()
-namespace EBasicKeyOperation
-{
-	enum Type
-	{
-		Set				UMETA(DisplayName="Is Set"),
-		NotSet			UMETA(DisplayName="Is Not Set"),
-	};
-}
-
-UENUM()
-namespace EArithmeticKeyOperation
-{
-	enum Type
-	{
-		Equal			UMETA(DisplayName="Is Equal To"),
-		NotEqual		UMETA(DisplayName="Is Not Equal To"),
-		Less			UMETA(DisplayName="Is Less Than"),
-		LessOrEqual		UMETA(DisplayName="Is Less Than Or Equal To"),
-		Greater			UMETA(DisplayName="Is Greater Than"),
-		GreaterOrEqual	UMETA(DisplayName="Is Greater Than Or Equal To"),
-	};
-}
-
-UENUM()
-namespace ETextKeyOperation
-{
-	enum Type
-	{
-		Equal			UMETA(DisplayName="Is Equal To"),
-		NotEqual		UMETA(DisplayName="Is Not Equal To"),
-		Contain			UMETA(DisplayName="Contains"),
-		NotContain		UMETA(DisplayName="Not Contains"),
-	};
-}
-
 struct FBlackboardInstancedKeyMemory
 {
 	/** index of instanced key in UBlackboardComponent::InstancedKeys */
