Call UpdateDecelerationData when we call SetMoveSegment in PathFollowingComponent

diff --git a/Engine/Source/Runtime/AIModule/Private/Navigation/PathFollowingComponent.cpp b/Engine/Source/Runtime/AIModule/Private/Navigation/PathFollowingComponent.cpp
--- a/Engine/Source/Runtime/AIModule/Private/Navigation/PathFollowingComponent.cpp
+++ b/Engine/Source/Runtime/AIModule/Private/Navigation/PathFollowingComponent.cpp
@@ -877,6 +877,7 @@ void UPathFollowingComponent::SetMoveSegment(int32 SegmentStartIndex)
 
 		// update move focus in owning AI
 		UpdateMoveFocus();
+		UpdateDecelerationData();
 	}
 }
 
