[BehaviorTree]
* Renamed Blackboard.h to BlackboardKey.h and moved it into the Blackboard folder

diff --git a/Engine/Source/Runtime/AIModule/Classes/BehaviorTree/BehaviorTreeTypes.h b/Engine/Source/Runtime/AIModule/Classes/BehaviorTree/BehaviorTreeTypes.h
--- a/Engine/Source/Runtime/AIModule/Classes/BehaviorTree/BehaviorTreeTypes.h
+++ b/Engine/Source/Runtime/AIModule/Classes/BehaviorTree/BehaviorTreeTypes.h
@@ -2,7 +2,7 @@
 
 #pragma once
 
-#include "Blackboard.h"
+#include "Blackboard/BlackboardKey.h"
 #include "CoreMinimal.h"
 #include "Stats/Stats.h"
 #include "UObject/ObjectMacros.h"
diff --git a/Engine/Source/Runtime/AIModule/Classes/BehaviorTree/Blackboard.h b/Engine/Source/Runtime/AIModule/Classes/BehaviorTree/Blackboard.h
deleted file mode 100644
diff --git a/Engine/Source/Runtime/AIModule/Classes/BehaviorTree/Blackboard/BlackboardKey.h b/Engine/Source/Runtime/AIModule/Classes/BehaviorTree/Blackboard/BlackboardKey.h
new file mode 100644
--- /dev/null
+++ b/Engine/Source/Runtime/AIModule/Classes/BehaviorTree/Blackboard/BlackboardKey.h
@@ -0,0 +1,37 @@
+// Copyright Epic Games, Inc. All Rights Reserved.
+
+#pragma once
+
+#include "UObject/NameTypes.h"
+
+namespace FBlackboard
+{
+	const FName KeySelf = TEXT("SelfActor");
+
+#ifdef AI_BLACKBOARD_KEY_SIZE_8
+	// this is the legacy BB key size. Add AI_BLACKBOARD_KEY_SIZE_8 to your *.target.cs to enable it.
+	using FKey = uint8;
+	inline constexpr FKey InvalidKey = FKey(-1);
+#else
+	//the default BB key size is now 16
+	struct FKey
+	{
+		constexpr FKey() = default;
+		FKey(int32 InKey) {Key = IntCastChecked<uint16>(InKey); }
+		constexpr FKey(uint16 InKey) : Key(InKey) {}
+		constexpr FKey(uint8 InKey) = delete;
+		constexpr operator int32() const { return Key; }
+		constexpr operator uint16() const { return Key; }
+		constexpr operator uint8() const = delete;
+		constexpr bool operator==(const FKey& Other) const {return Key == Other.Key;}
+		constexpr bool operator!=(const FKey& Other) const {return Key != Other.Key;}
+	private:
+		friend uint32 GetTypeHash(const FKey& Key);
+		uint16 Key = static_cast<uint16>(-1);
+	};
+
+	inline constexpr FKey InvalidKey = FKey();
+
+	inline uint32 GetTypeHash(const FKey& Key) { return ::GetTypeHash(Key.Key);}
+#endif
+}
\ No newline at end of file
