Addressing remaining numeric conversion warnings in EQS.

diff --git a/Engine/Source/Runtime/AIModule/Private/EnvironmentQuery/Tests/EnvQueryTest_Distance.cpp b/Engine/Source/Runtime/AIModule/Private/EnvironmentQuery/Tests/EnvQueryTest_Distance.cpp
--- a/Engine/Source/Runtime/AIModule/Private/EnvironmentQuery/Tests/EnvQueryTest_Distance.cpp
+++ b/Engine/Source/Runtime/AIModule/Private/EnvironmentQuery/Tests/EnvQueryTest_Distance.cpp
@@ -12,22 +12,26 @@ namespace
 {
 	FORCEINLINE float CalcDistance3D(const FVector& PosA, const FVector& PosB)
 	{
-		return (PosB - PosA).Size();
+		// Static cast this to a float, for EQS scoring purposes float precision is OK.
+		return static_cast<float>(FVector::Distance(PosA, PosB));
 	}
 
 	FORCEINLINE float CalcDistance2D(const FVector& PosA, const FVector& PosB)
 	{
-		return (PosB - PosA).Size2D();
+		// Static cast this to a float, for EQS scoring purposes float precision is OK.
+		return static_cast<float>(FVector::Dist2D(PosA, PosB));
 	}
 
 	FORCEINLINE float CalcDistanceZ(const FVector& PosA, const FVector& PosB)
 	{
-		return PosB.Z - PosA.Z;
+		// Static cast this to a float, for EQS scoring purposes float precision is OK.
+		return static_cast<float>(PosB.Z - PosA.Z);
 	}
 
 	FORCEINLINE float CalcDistanceAbsoluteZ(const FVector& PosA, const FVector& PosB)
 	{
-		return FMath::Abs(PosB.Z - PosA.Z);
+		// Static cast this to a float, for EQS scoring purposes float precision is OK.
+		return static_cast<float>(FMath::Abs(PosB.Z - PosA.Z));
 	}
 
 	FORCEINLINE void CheckItemLocationForNaN(const FVector& ItemLocation, UObject* QueryOwner, int32 Index, uint8 TestMode)
diff --git a/Engine/Source/Runtime/AIModule/Private/EnvironmentQuery/Tests/EnvQueryTest_Pathfinding.cpp b/Engine/Source/Runtime/AIModule/Private/EnvironmentQuery/Tests/EnvQueryTest_Pathfinding.cpp
--- a/Engine/Source/Runtime/AIModule/Private/EnvironmentQuery/Tests/EnvQueryTest_Pathfinding.cpp
+++ b/Engine/Source/Runtime/AIModule/Private/EnvironmentQuery/Tests/EnvQueryTest_Pathfinding.cpp
@@ -192,7 +192,8 @@ float UEnvQueryTest_Pathfinding::FindPathCostFrom(const FVector& ItemPos, const
 	Query.SetAllowPartialPaths(false);
 
 	FPathFindingResult Result = NavSys.FindPathSync(Query, Mode);
-	return (Result.IsSuccessful()) ? Result.Path->GetCost() : BIG_NUMBER;
+	// Static cast this to a float, for EQS scoring purposes float precision is OK.
+	return (Result.IsSuccessful()) ? static_cast<float>(Result.Path->GetCost()) : BIG_NUMBER;
 }
 
 float UEnvQueryTest_Pathfinding::FindPathCostTo(const FVector& ItemPos, const FVector& ContextPos, EPathFindingMode::Type Mode, const ANavigationData& NavData, UNavigationSystemV1& NavSys, FSharedConstNavQueryFilter NavFilter, const UObject* PathOwner) const
@@ -201,7 +202,8 @@ float UEnvQueryTest_Pathfinding::FindPathCostTo(const FVector& ItemPos, const FV
 	Query.SetAllowPartialPaths(false);
 
 	FPathFindingResult Result = NavSys.FindPathSync(Query, Mode);
-	return (Result.IsSuccessful()) ? Result.Path->GetCost() : BIG_NUMBER;
+	// Static cast this to a float, for EQS scoring purposes float precision is OK.
+	return (Result.IsSuccessful()) ? static_cast<float>(Result.Path->GetCost()) : BIG_NUMBER;
 }
 
 float UEnvQueryTest_Pathfinding::FindPathLengthFrom(const FVector& ItemPos, const FVector& ContextPos, EPathFindingMode::Type Mode, const ANavigationData& NavData, UNavigationSystemV1& NavSys, FSharedConstNavQueryFilter NavFilter, const UObject* PathOwner) const
@@ -210,7 +212,8 @@ float UEnvQueryTest_Pathfinding::FindPathLengthFrom(const FVector& ItemPos, cons
 	Query.SetAllowPartialPaths(false);
 
 	FPathFindingResult Result = NavSys.FindPathSync(Query, Mode);
-	return (Result.IsSuccessful()) ? Result.Path->GetLength() : BIG_NUMBER;
+	// Static cast this to a float, for EQS scoring purposes float precision is OK.
+	return (Result.IsSuccessful()) ? static_cast<float>(Result.Path->GetLength()) : BIG_NUMBER;
 }
 
 float UEnvQueryTest_Pathfinding::FindPathLengthTo(const FVector& ItemPos, const FVector& ContextPos, EPathFindingMode::Type Mode, const ANavigationData& NavData, UNavigationSystemV1& NavSys, FSharedConstNavQueryFilter NavFilter, const UObject* PathOwner) const
@@ -219,7 +222,8 @@ float UEnvQueryTest_Pathfinding::FindPathLengthTo(const FVector& ItemPos, const
 	Query.SetAllowPartialPaths(false);
 
 	FPathFindingResult Result = NavSys.FindPathSync(Query, Mode);
-	return (Result.IsSuccessful()) ? Result.Path->GetLength() : BIG_NUMBER;
+	// Static cast this to a float, for EQS scoring purposes float precision is OK.
+	return (Result.IsSuccessful()) ? static_cast<float>(Result.Path->GetLength()) : BIG_NUMBER;
 }
 
 ANavigationData* UEnvQueryTest_Pathfinding::FindNavigationData(UNavigationSystemV1& NavSys, UObject* Owner) const
diff --git a/Engine/Source/Runtime/AIModule/Private/EnvironmentQuery/Tests/EnvQueryTest_PathfindingBatch.cpp b/Engine/Source/Runtime/AIModule/Private/EnvironmentQuery/Tests/EnvQueryTest_PathfindingBatch.cpp
--- a/Engine/Source/Runtime/AIModule/Private/EnvironmentQuery/Tests/EnvQueryTest_PathfindingBatch.cpp
+++ b/Engine/Source/Runtime/AIModule/Private/EnvironmentQuery/Tests/EnvQueryTest_PathfindingBatch.cpp
@@ -31,8 +31,9 @@ namespace NodePoolHelpers
 		const FRecastDebugPathfindingNode* MyNode = NodePool.Nodes.Find(SearchKey);
 		if (MyNode)
 		{
-			float LastSegmentLength = FVector::Dist(MyNode->NodePos, TestPt.OutLocation.Location);
-			return MyNode->Length + LastSegmentLength;
+			const FVector::FReal LastSegmentLength = FVector::Dist(MyNode->NodePos, TestPt.OutLocation.Location);
+			// Static cast this to a float, for EQS scoring purposes float precision is OK.
+			return static_cast<float>(MyNode->Length + LastSegmentLength);
 		}
 
 		return BIG_NUMBER;
@@ -44,7 +45,8 @@ namespace NodePoolHelpers
 		const FRecastDebugPathfindingNode* MyNode = NodePool.Nodes.Find(SearchKey);
 		if (MyNode)
 		{
-			return MyNode->TotalCost;
+			// Static cast this to a float, for EQS scoring purposes float precision is OK.
+			return static_cast<float>(MyNode->TotalCost);
 		}
 
 		return BIG_NUMBER;
@@ -90,7 +92,7 @@ void UEnvQueryTest_PathfindingBatch::RunTest(FEnvQueryInstance& QueryInstance) c
 	}
 
 	TArray<FNavigationProjectionWork> TestPoints;
-	TArray<float> CollectDistanceSq;
+	TArray<FVector::FReal> CollectDistanceSq;
 	CollectDistanceSq.Init(0.0f, ContextLocations.Num());
 
 	TSubclassOf<UNavigationQueryFilter> NavFilterToUse = GetNavFilterClass(QueryInstance);
@@ -115,7 +117,7 @@ void UEnvQueryTest_PathfindingBatch::RunTest(FEnvQueryInstance& QueryInstance) c
 
 				for (int32 ContextIdx = 0; ContextIdx < ContextLocations.Num(); ContextIdx++)
 				{
-					const float TestDistanceSq = FVector::DistSquared(ItemLocation, ContextLocations[ContextIdx]);
+					const FVector::FReal TestDistanceSq = FVector::DistSquared(ItemLocation, ContextLocations[ContextIdx]);
 					CollectDistanceSq[ContextIdx] = FMath::Max(CollectDistanceSq[ContextIdx], TestDistanceSq);
 				}
 			}
@@ -133,7 +135,7 @@ void UEnvQueryTest_PathfindingBatch::RunTest(FEnvQueryInstance& QueryInstance) c
 		TArray<NavNodeRef> Polys;
 		for (int32 ContextIdx = 0; ContextIdx < ContextLocations.Num(); ContextIdx++)
 		{
-			const float MaxPathDistance = FMath::Sqrt(CollectDistanceSq[ContextIdx]) * RangeMultiplierValue;
+			const FVector::FReal MaxPathDistance = FMath::Sqrt(CollectDistanceSq[ContextIdx]) * RangeMultiplierValue;
 
 			Polys.Reset();
 			NodePoolData[ContextIdx].Flags = ERecastDebugPathfindingFlags::PathLength;
