PR #9819: Add OnTargetPerceptionForgetten Delegate to UAIPerceptionComponent (Contributed by [REDACTED_PERSON])

diff --git a/Engine/Source/Runtime/AIModule/Classes/Perception/AIPerceptionComponent.h b/Engine/Source/Runtime/AIModule/Classes/Perception/AIPerceptionComponent.h
--- a/Engine/Source/Runtime/AIModule/Classes/Perception/AIPerceptionComponent.h
+++ b/Engine/Source/Runtime/AIModule/Classes/Perception/AIPerceptionComponent.h
@@ -21,6 +21,7 @@ struct FVisualLogEntry;
 
 DECLARE_DYNAMIC_MULTICAST_DELEGATE_OneParam(FPerceptionUpdatedDelegate, const TArray<AActor*>&, UpdatedActors);
 DECLARE_DYNAMIC_MULTICAST_DELEGATE_TwoParams(FActorPerceptionUpdatedDelegate, AActor*, Actor, FAIStimulus, Stimulus);
+DECLARE_DYNAMIC_MULTICAST_DELEGATE_OneParam(FActorPerceptionForgetUpdatedDelegate, AActor*, Actor);
 
 USTRUCT(BlueprintType, meta = (DisplayName = "Sensed Actor's Update Data"))
 struct FActorPerceptionUpdateInfo
@@ -379,6 +380,17 @@ public:
 	UPROPERTY(BlueprintAssignable)
 	FPerceptionUpdatedDelegate OnPerceptionUpdated;
 
+	/**
+	 * Notifies all bound delegates that the perception info has been forgotten for a given target.
+	 * The notification get broadcast when all stimuli of a given target expire. Note that this
+	 * functionality requires the the actor forgetting must be enabled via AIPerceptionSystem.bForgetStaleActors.
+	 *
+	 * @param	SourceActor	Actor associated to the stimulus (can not be null)
+	 * @param	Stimulus	Updated stimulus
+	 */
+	UPROPERTY(BlueprintAssignable)
+	FActorPerceptionForgetUpdatedDelegate OnTargetPerceptionForgotten;
+
 	/**
 	 * Notifies all bound objects that perception info has been updated for a given target.
 	 * The notification is broadcasted for any received stimulus or on change of state
diff --git a/Engine/Source/Runtime/AIModule/Private/Perception/AIPerceptionComponent.cpp b/Engine/Source/Runtime/AIModule/Private/Perception/AIPerceptionComponent.cpp
--- a/Engine/Source/Runtime/AIModule/Private/Perception/AIPerceptionComponent.cpp
+++ b/Engine/Source/Runtime/AIModule/Private/Perception/AIPerceptionComponent.cpp
@@ -607,6 +607,15 @@ void UAIPerceptionComponent::ProcessStimuli()
 		ForgetActor(ActorToForget);
 	}
 
+	// notify anyone interested
+	if (OnTargetPerceptionForgotten.IsBound())
+	{
+		for (AActor* ActorToForget : ActorsToForget)
+		{
+			OnTargetPerceptionForgotten.Broadcast(ActorToForget);
+		}
+	}
+
 	// remove perceptual info related to stale actors
 	for (const TObjectKey<AActor>& SourceKey : DataToRemove)
 	{
