BVTree:
* Expose elements array to allow callers to re-use its existing allocation
* Use Count to compute NodeIndex instead of waiting after all child have been visited

#rnx

diff --git a/Engine/Source/Runtime/AIModule/Public/BVTree.h b/Engine/Source/Runtime/AIModule/Public/BVTree.h
--- a/Engine/Source/Runtime/AIModule/Public/BVTree.h
+++ b/Engine/Source/Runtime/AIModule/Public/BVTree.h
@@ -5,6 +5,17 @@
 #include "Containers/ContainersFwd.h"
 #include "Math/Box.h"
 
+/**
+* BVTree - Bounding Volumes Tree
+*
+* A binary search tree where each node has a Bounding Volume which is either an element (leaf nodes)
+* or a volume encompassing all of its children (non-leaf nodes).
+*
+* It provides fast gathering of all bounding volumes overlapping a given FBox (see GetOverlapping).
+* 
+* Users are expected to provide an implementation of CalcElementBounds for ElementType:
+*	static FBox CalcElementBounds(const FElement& Element) const
+*/
 template<typename InElementType, typename InAllocator = FDefaultAllocator>
 struct TBVTree
 {
@@ -34,6 +45,15 @@ struct TBVTree
 		Create(MoveTemp(InElements));
 	}
 
+	// Recreate the tree based on the current content of Elements.
+	// Callers are expected to have filled the GetElements() array before calling this function.
+	void RecreateTree()
+	{
+		Nodes.Reset();
+		NodeBoundingBoxes.Reset();
+		CreateCommonInternal();
+	}
+
 	void RecreateTree(const TArray<FElement>& InElements)
 	{
 		Reset();
@@ -69,6 +89,7 @@ struct TBVTree
 	const TArray<int16, InAllocator>& GetNodes() const { return Nodes; }
 	const TArray<FBox, InAllocator>& GetBoundingBoxes() const { return NodeBoundingBoxes; }
 	const TArray<FElement>& GetElements() const { return Elements; }
+	TArray<FElement>& GetElements() { return Elements; }
 
 	bool IsEmpty() const { return Nodes.Num() == 0; }
 
@@ -92,6 +113,13 @@ protected:
 			// Needs splitting
 			const FBox TempNodeBounds = CalcNodeBounds(ElementBBoxes, StartIndex, LimitIndex);
 			NodeBounds = TempNodeBounds;
+
+			const int NumChildNodes = Count * 2 - 1;
+
+			// A negative index means this is a non-leaf node, and the value is the number of nodes
+			// to skip to exit the current branch and jump to the next unvisited node during a search.
+			NodeIndex = -NumChildNodes;
+
 			const int Axis = GetLongestAxis(NodeBounds);
 
 			struct FAxisSort
@@ -113,11 +141,6 @@ protected:
 			Subdivide(ElementBBoxes, StartIndex, SplitIndex, CurrentNode);
 			// Right
 			Subdivide(ElementBBoxes, SplitIndex, LimitIndex, CurrentNode);
-
-			// Negative index means escape.
-			const int EscapeIndex = ThisCurrentNode - CurrentNode;
-			// node.i = -iescape;
-			NodeIndex = EscapeIndex;
 		}
 	}
 
