HearingSense: Deprecate config values for unimplemented LOSHearingRange

#rnx

diff --git a/Engine/Config/BaseGame.ini b/Engine/Config/BaseGame.ini
--- a/Engine/Config/BaseGame.ini
+++ b/Engine/Config/BaseGame.ini
@@ -176,7 +176,6 @@ MaxGlobalTimeDilation=20.0
 HearingRange=768
 SightRadius=3000
 LoseSightRadius=3500
-LoSHearingRange=1500
 PeripheralVisionAngle=90
 
 [/Script/AIModule.AISense_Hearing]
@@ -185,7 +184,6 @@ SpeedOfSoundSq=0
 [/Script/AIModule.AISenseConfig_Hearing]
 Implementation=Class'/Script/AIModule.AISense_Hearing'
 HearingRange=768
-LoSHearingRange=1500
 DetectionByAffiliation=(bDetectEnemies=true)
 
 [/Script/AIModule.AISenseConfig_Sight]
diff --git a/Engine/Source/Runtime/AIModule/Classes/Perception/AISenseConfig_Hearing.h b/Engine/Source/Runtime/AIModule/Classes/Perception/AISenseConfig_Hearing.h
--- a/Engine/Source/Runtime/AIModule/Classes/Perception/AISenseConfig_Hearing.h
+++ b/Engine/Source/Runtime/AIModule/Classes/Perception/AISenseConfig_Hearing.h
@@ -26,10 +26,11 @@ public:
 	UPROPERTY(EditDefaultsOnly, BlueprintReadOnly, Category = "Sense")
 	float HearingRange;
 
+	UE_DEPRECATED(5.1, "LoSHearingRange is deprecated. Use HearingRage instead.")
 	UPROPERTY(EditDefaultsOnly, BlueprintReadOnly, Category = "Sense", meta = (EditCondition = "bUseLoSHearing"))
 	float LoSHearingRange;
 
-	/** Warning: has significant runtime cost */
+	UE_DEPRECATED(5.1, "bUseLoSHearing is deprecated.")
 	UPROPERTY(EditDefaultsOnly, Category = "Sense", meta = (InlineEditConditionToggle))
 	uint32 bUseLoSHearing : 1;
 
diff --git a/Engine/Source/Runtime/AIModule/Classes/Perception/AISense_Hearing.h b/Engine/Source/Runtime/AIModule/Classes/Perception/AISense_Hearing.h
--- a/Engine/Source/Runtime/AIModule/Classes/Perception/AISense_Hearing.h
+++ b/Engine/Source/Runtime/AIModule/Classes/Perception/AISense_Hearing.h
@@ -27,7 +27,7 @@ struct AIMODULE_API FAINoiseEvent
 	/**
 	 * Loudness modifier of the sound.
 	 * If MaxRange is non-zero, this modifies the range (by multiplication).
-	 * If there is no MaxRange, then if Square(DistanceToSound) <= Square(HearingRange) * Loudness, the sound is heard, false otherwise.
+	 * If there is no MaxRange, then if Square(DistanceToSound) <= Square(HearingRange * Loudness), the sound is heard, false otherwise.
 	 */
 	UPROPERTY(EditAnywhere, BlueprintReadWrite, Category = "Sense", meta = (UIMin = 0, ClampMin = 0))
 	float Loudness;
@@ -77,9 +77,7 @@ protected:
 	struct FDigestedHearingProperties
 	{
 		float HearingRangeSq;
-		float LoSHearingRangeSq;
 		uint8 AffiliationFlags;
-		uint32 bUseLoSHearing : 1;
 
 		FDigestedHearingProperties(const UAISenseConfig_Hearing& SenseConfig);
 		FDigestedHearingProperties();
diff --git a/Engine/Source/Runtime/AIModule/Private/Perception/AISense.cpp b/Engine/Source/Runtime/AIModule/Private/Perception/AISense.cpp
--- a/Engine/Source/Runtime/AIModule/Private/Perception/AISense.cpp
+++ b/Engine/Source/Runtime/AIModule/Private/Perception/AISense.cpp
@@ -281,10 +281,6 @@ void UAISenseConfig_Hearing::DescribeSelfToGameplayDebugger(const UAIPerceptionC
 		FVector OwnerLocation = BodyActor->GetActorLocation();
 		
 		DebuggerCategory->AddShape(FGameplayDebuggerShape::MakeCylinder(OwnerLocation, HearingRange, 25.0f, HearingRangeColor));
-		if (bUseLoSHearing)
-		{
-			DebuggerCategory->AddShape(FGameplayDebuggerShape::MakeCylinder(OwnerLocation, LoSHearingRange, 25.0f, LoSHearingRangeColor));
-		}
 	}
 }
 #endif // WITH_GAMEPLAY_DEBUGGER
diff --git a/Engine/Source/Runtime/AIModule/Private/Perception/AISense_Hearing.cpp b/Engine/Source/Runtime/AIModule/Private/Perception/AISense_Hearing.cpp
--- a/Engine/Source/Runtime/AIModule/Private/Perception/AISense_Hearing.cpp
+++ b/Engine/Source/Runtime/AIModule/Private/Perception/AISense_Hearing.cpp
@@ -39,13 +39,11 @@ void FAINoiseEvent::Compile()
 UAISense_Hearing::FDigestedHearingProperties::FDigestedHearingProperties(const UAISenseConfig_Hearing& SenseConfig)
 {
 	HearingRangeSq = FMath::Square(SenseConfig.HearingRange);
-	LoSHearingRangeSq = FMath::Square(SenseConfig.LoSHearingRange);
 	AffiliationFlags = SenseConfig.DetectionByAffiliation.GetAsFlags();
-	bUseLoSHearing = SenseConfig.bUseLoSHearing;
 }
 
 UAISense_Hearing::FDigestedHearingProperties::FDigestedHearingProperties()
-	: HearingRangeSq(-1.f), LoSHearingRangeSq(-1.f), AffiliationFlags(-1), bUseLoSHearing(false)
+	: HearingRangeSq(-1.f), AffiliationFlags(-1)
 {
 
 }
