HearingSense: Use proper UE version number in deprecation macro

#rnx

diff --git a/Engine/Source/Runtime/AIModule/Classes/Perception/AISenseConfig_Hearing.h b/Engine/Source/Runtime/AIModule/Classes/Perception/AISenseConfig_Hearing.h
--- a/Engine/Source/Runtime/AIModule/Classes/Perception/AISenseConfig_Hearing.h
+++ b/Engine/Source/Runtime/AIModule/Classes/Perception/AISenseConfig_Hearing.h
@@ -26,11 +26,11 @@ public:
 	UPROPERTY(EditDefaultsOnly, BlueprintReadOnly, Category = "Sense")
 	float HearingRange;
 
-	UE_DEPRECATED(5.1, "LoSHearingRange is deprecated. Use HearingRage instead.")
+	UE_DEPRECATED(5.2, "LoSHearingRange is deprecated. Use HearingRange instead.")
 	UPROPERTY(EditDefaultsOnly, BlueprintReadOnly, Category = "Sense", meta = (EditCondition = "bUseLoSHearing"))
 	float LoSHearingRange;
 
-	UE_DEPRECATED(5.1, "bUseLoSHearing is deprecated.")
+	UE_DEPRECATED(5.2, "bUseLoSHearing is deprecated.")
 	UPROPERTY(EditDefaultsOnly, Category = "Sense", meta = (InlineEditConditionToggle))
 	uint32 bUseLoSHearing : 1;
 
