fix crash while inserting element in OutOfRange sight query list

diff --git a/Engine/Source/Runtime/AIModule/Private/Perception/AISense_Sight.cpp b/Engine/Source/Runtime/AIModule/Private/Perception/AISense_Sight.cpp
--- a/Engine/Source/Runtime/AIModule/Private/Perception/AISense_Sight.cpp
+++ b/Engine/Source/Runtime/AIModule/Private/Perception/AISense_Sight.cpp
@@ -697,8 +697,15 @@ void UAISense_Sight::OnPendingQueryProcessed(const int32 SightQueryIndex, const
 	}
 	else
 	{
-		SightQueriesOutOfRange.Insert(SightQuery, NextOutOfRangeIndex);
-		++NextOutOfRangeIndex;
+		if (bSightQueriesOutOfRangeDirty)
+		{
+			SightQueriesOutOfRange.Add(SightQuery); 
+		}
+		else
+		{
+			SightQueriesOutOfRange.Insert(SightQuery, NextOutOfRangeIndex);
+			++NextOutOfRangeIndex;
+		}
 	}
 }
 
