[BlackboardData] initialize FirstKeyId to 0 instead of invalid key so newly created asset will be valid on creation which will prevent a check when calling InternalGetKeyID.
#rnx

diff --git a/Engine/Source/Runtime/AIModule/Private/BehaviorTree/BlackboardData.cpp b/Engine/Source/Runtime/AIModule/Private/BehaviorTree/BlackboardData.cpp
--- a/Engine/Source/Runtime/AIModule/Private/BehaviorTree/BlackboardData.cpp
+++ b/Engine/Source/Runtime/AIModule/Private/BehaviorTree/BlackboardData.cpp
@@ -52,7 +52,7 @@ bool FBlackboardEntry::operator==(const FBlackboardEntry& Other) const
 		((KeyType && Other.KeyType && KeyType->GetClass() == Other.KeyType->GetClass()) || (KeyType == NULL && Other.KeyType == NULL));
 }
 
-UBlackboardData::UBlackboardData(const FObjectInitializer& ObjectInitializer) : Super(ObjectInitializer)
+UBlackboardData::UBlackboardData(const FObjectInitializer& ObjectInitializer) : Super(ObjectInitializer), FirstKeyID(0)
 {
 }
 
