fix crash caused by missed removal of pending requests in Sight sense

diff --git a/Engine/Source/Runtime/AIModule/Private/Perception/AISense_Sight.cpp b/Engine/Source/Runtime/AIModule/Private/Perception/AISense_Sight.cpp
--- a/Engine/Source/Runtime/AIModule/Private/Perception/AISense_Sight.cpp
+++ b/Engine/Source/Runtime/AIModule/Private/Perception/AISense_Sight.cpp
@@ -993,12 +993,6 @@ void UAISense_Sight::OnListenerRemovedImpl(const FPerceptionListener& RemovedLis
 void UAISense_Sight::RemoveAllQueriesByListener(const FPerceptionListener& Listener, const TFunction<void(const FAISightQuery&)>& OnRemoveFunc/*= nullptr */)
 {
 	SCOPE_CYCLE_COUNTER(STAT_AI_Sense_Sight_RemoveByListener);
-
-	if ((SightQueriesInRange.Num() + SightQueriesOutOfRange.Num()) == 0)
-	{
-		return;
-	}
-
 	UE_MT_SCOPED_WRITE_ACCESS(QueriesListAccessDetector);
 
 	const uint32 ListenerId = Listener.GetListenerID();
