Deprecated PawnActions

diff --git a/Engine/Source/Runtime/AIModule/Classes/AIController.h b/Engine/Source/Runtime/AIModule/Classes/AIController.h
--- a/Engine/Source/Runtime/AIModule/Classes/AIController.h
+++ b/Engine/Source/Runtime/AIModule/Classes/AIController.h
@@ -31,8 +31,8 @@ class UCanvas;
 class UGameplayTaskResource;
 class UGameplayTasksComponent;
 class UPathFollowingComponent;
-class UPawnAction;
-class UPawnActionsComponent;
+class UDEPRECATED_UPawnAction;
+class UDEPRECATED_UPawnActionsComponent;
 
 namespace EPathFollowingRequestResult {	enum Type : int; }
 namespace EPathFollowingResult { enum Type : int; }
@@ -145,8 +145,8 @@ public:
 	TObjectPtr<UAIPerceptionComponent> PerceptionComponent;
 
 private:
-	UPROPERTY(BlueprintReadOnly, Category = AI, meta = (AllowPrivateAccess = "true"))
-	TObjectPtr<UPawnActionsComponent> ActionsComp;
+	UPROPERTY(Category = AI, BlueprintGetter = GetDeprecatedActionsComponent)
+	TObjectPtr<UDEPRECATED_UPawnActionsComponent> ActionsComp_DEPRECATED;
 
 protected:
 	/** blackboard */
@@ -422,11 +422,6 @@ public:
 	UFUNCTION()
 	virtual void OnGameplayTaskResourcesClaimed(FGameplayResourceSet NewlyClaimed, FGameplayResourceSet FreshlyReleased);
 
-	//----------------------------------------------------------------------//
-	// Actions
-	//----------------------------------------------------------------------//
-	bool PerformAction(UPawnAction& Action, EAIRequestPriority::Type Priority, UObject* const Instigator = NULL);
-
 	//----------------------------------------------------------------------//
 	// debug/dev-time 
 	//----------------------------------------------------------------------//
@@ -443,8 +438,6 @@ public:
 	/** Returns PathFollowingComponent subobject **/
 	UFUNCTION(BlueprintCallable, Category="AI|Navigation")
 	UPathFollowingComponent* GetPathFollowingComponent() const { return PathFollowingComponent; }
-	/** Returns ActionsComp subobject **/
-	UPawnActionsComponent* GetActionsComp() const { return ActionsComp; }
 	UFUNCTION(BlueprintPure, Category = "AI|Perception")
 	UAIPerceptionComponent* GetAIPerceptionComponent() { return PerceptionComponent; }
 
@@ -458,6 +451,18 @@ public:
 	 *	Intended to be called as part of initialization/setup process */
 	UFUNCTION(BlueprintCallable, Category = "AI|Navigation")
 	void SetPathFollowingComponent(UPathFollowingComponent* NewPFComponent);
+
+	//----------------------------------------------------------------------//
+	// DEPRECATED
+	//----------------------------------------------------------------------//
+	UE_DEPRECATED(5.2, "PawnActions have been deprecated and are no longer being supported. It will get removed in following UE5 releases. Use GameplayTasks or AITasks instead.")
+	UDEPRECATED_UPawnActionsComponent* GetActionsComp() const { return ActionsComp_DEPRECATED; }
+	
+	UE_DEPRECATED(5.2, "PawnActions have been deprecated and are no longer being supported. It will get removed in following UE5 releases. Use GameplayTasks or AITasks instead.")
+	bool PerformAction(UDEPRECATED_UPawnAction& Action, EAIRequestPriority::Type Priority, UObject* const Instigator = NULL);
+
+	UFUNCTION(BlueprintGetter, meta = (DeprecatedFunction, DeprecationMessage = "PawnActions have been deprecated and are no longer being supported. It will get removed in following UE5 releases. Use GameplayTasks or AITasks instead."))
+	UDEPRECATED_UPawnActionsComponent* GetDeprecatedActionsComponent() const { return ActionsComp_DEPRECATED; }
 };
 
 //----------------------------------------------------------------------//
diff --git a/Engine/Source/Runtime/AIModule/Classes/Actions/PawnAction.h b/Engine/Source/Runtime/AIModule/Classes/Actions/PawnAction.h
--- a/Engine/Source/Runtime/AIModule/Classes/Actions/PawnAction.h
+++ b/Engine/Source/Runtime/AIModule/Classes/Actions/PawnAction.h
@@ -12,8 +12,8 @@
 
 class AController;
 class APawn;
-class UPawnAction;
-class UPawnActionsComponent;
+class UDEPRECATED_UPawnAction;
+class UDEPRECATED_UPawnActionsComponent;
 struct FPawnActionStack;
 
 UENUM()
@@ -27,7 +27,7 @@ namespace EPawnSubActionTriggeringPolicy
 }
 
 AIMODULE_API DECLARE_LOG_CATEGORY_EXTERN(LogPawnAction, Warning, All);
-DECLARE_DELEGATE_TwoParams(FPawnActionEventDelegate, UPawnAction&, EPawnActionEventType::Type);
+DECLARE_DELEGATE_TwoParams(FPawnActionEventDelegate, UDEPRECATED_UPawnAction&, EPawnActionEventType::Type);
 
 UENUM()
 namespace EPawnActionFailHandling
@@ -43,25 +43,25 @@ namespace EPawnActionFailHandling
  *	Things to remember:
  *	* Actions are created paused
  */
-UCLASS(abstract, EditInlineNew)
-class AIMODULE_API UPawnAction : public UObject
+UCLASS(abstract, EditInlineNew, deprecated, meta = (DeprecationMessage = "PawnActions have been deprecated and are no longer being supported. It will get removed in following UE5 releases. Use GameplayTasks or AITasks instead."))
+class AIMODULE_API UDEPRECATED_UPawnAction : public UObject
 {
 	GENERATED_UCLASS_BODY()
 
-	friend UPawnActionsComponent;
+	friend UDEPRECATED_UPawnActionsComponent;
 	friend FPawnActionStack;
 
 private:
 	/** Current child node executing on top of this Action */
 	UPROPERTY(Transient)
-	TObjectPtr<UPawnAction> ChildAction;
+	TObjectPtr<UDEPRECATED_UPawnAction> ChildAction_DEPRECATED;
 
 	UPROPERTY(Transient)
-	TObjectPtr<UPawnAction> ParentAction;
+	TObjectPtr<UDEPRECATED_UPawnAction> ParentAction_DEPRECATED;
 
 	/** Extra reference to the component this action is being governed by */
 	UPROPERTY(Transient)
-	TObjectPtr<UPawnActionsComponent> OwnerComponent;
+	TObjectPtr<UDEPRECATED_UPawnActionsComponent> OwnerComponent_DEPRECATED;
 	
 	/** indicates an object that caused this action. Used for mass removal of actions 
 	 *	by specific object */
@@ -137,9 +137,9 @@ public:
 	virtual UWorld* GetWorld() const override;
 	// End UObject
 
-	FORCEINLINE const UPawnAction* GetParentAction() const { return ParentAction; }
-	FORCEINLINE const UPawnAction* GetChildAction() const { return ChildAction; }
-	FORCEINLINE UPawnAction* GetChildAction() { return ChildAction; }
+	FORCEINLINE const UDEPRECATED_UPawnAction* GetParentAction() const { return ParentAction_DEPRECATED; }
+	FORCEINLINE const UDEPRECATED_UPawnAction* GetChildAction() const { return ChildAction_DEPRECATED; }
+	FORCEINLINE UDEPRECATED_UPawnAction* GetChildAction() { return ChildAction_DEPRECATED; }
 	FORCEINLINE bool IsPaused() const { return !!bPaused; }
 	FORCEINLINE bool IsActive() const { return FinishResult == EPawnActionResult::InProgress && IsPaused() == false && AbortState == EPawnActionAbortState::NotBeingAborted; }
 	FORCEINLINE bool IsBeingAborted() const { return AbortState != EPawnActionAbortState::NotBeingAborted; }
@@ -152,9 +152,9 @@ protected:
 	FORCEINLINE void TickAction(float DeltaTime)
 	{ 
 		// tick ChildAction 
-		if (ChildAction != NULL)
+		if (ChildAction_DEPRECATED != NULL)
 		{
-			ChildAction->Tick(DeltaTime);
+			ChildAction_DEPRECATED->Tick(DeltaTime);
 		}
 		// or self if not paused
 		else if (!!bWantsTick && IsPaused() == false)
@@ -169,12 +169,12 @@ protected:
 	 *	@NOTE do not make this virtual! Contains some essential logic. */
 	EPawnActionAbortState::Type Abort(EAIForceParam::Type ShouldForce = EAIForceParam::DoNotForce);
 	
-	FORCEINLINE UPawnActionsComponent* GetOwnerComponent() { return OwnerComponent; }
+	FORCEINLINE UDEPRECATED_UPawnActionsComponent* GetOwnerComponent() { return OwnerComponent_DEPRECATED; }
 public:
 	FORCEINLINE EAIRequestPriority::Type GetPriority() const { return ExecutionPriority; }
 	FORCEINLINE EPawnActionResult::Type GetResult() const { return FinishResult; }
 	FORCEINLINE EPawnActionAbortState::Type GetAbortState() const { return AbortState; }
-	FORCEINLINE UPawnActionsComponent* GetOwnerComponent() const { return OwnerComponent; }
+	FORCEINLINE UDEPRECATED_UPawnActionsComponent* GetOwnerComponent() const { return OwnerComponent_DEPRECATED; }
 	FORCEINLINE UObject* GetInstigator() const { return Instigator; }
 	APawn* GetPawn() const;
 	AController* GetController() const;
@@ -182,7 +182,7 @@ public:
 	template<class TActionClass>
 	static TActionClass* CreateActionInstance(UWorld& World)
 	{
-		TSubclassOf<UPawnAction> ActionClass = TActionClass::StaticClass();
+		TSubclassOf<UDEPRECATED_UPawnAction> ActionClass = TActionClass::StaticClass();
 		return NewObject<TActionClass>(&World, ActionClass);
 	}
 
@@ -202,8 +202,9 @@ public:
 	UFUNCTION(BlueprintPure, Category = "AI|PawnActions")
 	TEnumAsByte<EAIRequestPriority::Type> GetActionPriority();
 
-	UFUNCTION(BlueprintCallable, Category = "AI|PawnActions", meta = (WorldContext="WorldContextObject"))
-	static UPawnAction* CreateActionInstance(UObject* WorldContextObject, TSubclassOf<UPawnAction> ActionClass);
+	UE_DEPRECATED(5.2, "PawnActions have been deprecated and are no longer being supported. It will get removed in following UE5 releases. Use GameplayTasks or AITasks instead.")
+	UFUNCTION(BlueprintCallable, Category = "AI|PawnActions", meta = (WorldContext="WorldContextObject", DeprecatedFunction, DeprecationMessage = "PawnActions have been deprecated and are no longer being supported. It will get removed in following UE5 releases. Use GameplayTasks or AITasks instead."))
+	static UDEPRECATED_UPawnAction* CreateActionInstance(UObject* WorldContextObject, TSubclassOf<UDEPRECATED_UPawnAction> ActionClass);
 
 	//----------------------------------------------------------------------//
 	// debug
@@ -225,7 +226,7 @@ protected:
 
 	void StopWaitingForMessages();
 
-	void SetOwnerComponent(UPawnActionsComponent* Component);
+	void SetOwnerComponent(UDEPRECATED_UPawnActionsComponent* Component);
 
 	void SetInstigator(UObject* const InInstigator);
 
@@ -236,7 +237,7 @@ protected:
 	 *	@NOTE if action fails to start no finishing or aborting mechanics will be triggered */
 	virtual bool Start();
 	/** called to pause action when higher priority or child action kicks in */
-	virtual bool Pause(const UPawnAction* PausedBy);
+	virtual bool Pause(const UDEPRECATED_UPawnAction* PausedBy);
 	/** called to resume action after being paused */
 	virtual bool Resume();
 	/** called when this action is being removed from action stacks */
@@ -245,10 +246,10 @@ protected:
 	 *	@NOTE gets called _AFTER_ child's OnFinished to give child action chance 
 	 *		to prepare "finishing data" for parent to read. 
 	 *	@NOTE clears parent-child binding */
-	virtual void OnChildFinished(UPawnAction& Action, EPawnActionResult::Type WithResult);
+	virtual void OnChildFinished(UDEPRECATED_UPawnAction& Action, EPawnActionResult::Type WithResult);
 
 	/** apart from doing regular push request copies additional values from Parent, like Priority and Instigator */
-	bool PushChildAction(UPawnAction& Action);
+	bool PushChildAction(UDEPRECATED_UPawnAction& Action);
 	
 	/** performs actual work on aborting Action. Should be called exclusively by Abort function
 	 *	@return only valid return values here are LatendAbortInProgress and AbortDone */
@@ -272,7 +273,7 @@ private:
 //----------------------------------------------------------------------//
 // Blueprint inlines
 //----------------------------------------------------------------------//
-FORCEINLINE TEnumAsByte<EAIRequestPriority::Type> UPawnAction::GetActionPriority()
+FORCEINLINE TEnumAsByte<EAIRequestPriority::Type> UDEPRECATED_UPawnAction::GetActionPriority()
 {
 	return ExecutionPriority;
 }
diff --git a/Engine/Source/Runtime/AIModule/Classes/Actions/PawnAction_BlueprintBase.h b/Engine/Source/Runtime/AIModule/Classes/Actions/PawnAction_BlueprintBase.h
--- a/Engine/Source/Runtime/AIModule/Classes/Actions/PawnAction_BlueprintBase.h
+++ b/Engine/Source/Runtime/AIModule/Classes/Actions/PawnAction_BlueprintBase.h
@@ -10,7 +10,7 @@
 class APawn;
 
 UCLASS(abstract, Blueprintable)
-class AIMODULE_API UPawnAction_BlueprintBase : public UPawnAction
+class AIMODULE_API UDEPRECATED_UPawnAction_BlueprintBase : public UDEPRECATED_UPawnAction
 {
 	GENERATED_UCLASS_BODY()
 
@@ -34,7 +34,7 @@ public:
 protected:
 	virtual void Tick(float DeltaTime) override;
 	virtual bool Start() override;
-	virtual bool Pause(const UPawnAction* PausedBy) override;
+	virtual bool Pause(const UDEPRECATED_UPawnAction* PausedBy) override;
 	virtual bool Resume() override;
 	virtual void OnFinished(EPawnActionResult::Type WithResult) override;
 };
diff --git a/Engine/Source/Runtime/AIModule/Classes/Actions/PawnAction_Move.h b/Engine/Source/Runtime/AIModule/Classes/Actions/PawnAction_Move.h
--- a/Engine/Source/Runtime/AIModule/Classes/Actions/PawnAction_Move.h
+++ b/Engine/Source/Runtime/AIModule/Classes/Actions/PawnAction_Move.h
@@ -24,7 +24,7 @@ namespace EPawnActionMoveMode
 }
 
 UCLASS()
-class AIMODULE_API UPawnAction_Move : public UPawnAction
+class AIMODULE_API UDEPRECATED_UPawnAction_Move : public UDEPRECATED_UPawnAction
 {
 	GENERATED_UCLASS_BODY()
 protected:
@@ -71,8 +71,8 @@ protected:
 public:
 	virtual void BeginDestroy() override;
 
-	static UPawnAction_Move* CreateAction(UWorld& World, AActor* GoalActor, EPawnActionMoveMode::Type Mode);
-	static UPawnAction_Move* CreateAction(UWorld& World, const FVector& GoalLocation, EPawnActionMoveMode::Type Mode);
+	static UDEPRECATED_UPawnAction_Move* CreateAction(UWorld& World, AActor* GoalActor, EPawnActionMoveMode::Type Mode);
+	static UDEPRECATED_UPawnAction_Move* CreateAction(UWorld& World, const FVector& GoalLocation, EPawnActionMoveMode::Type Mode);
 
 	static bool CheckAlreadyAtGoal(AAIController& Controller, const FVector& TestLocation, float Radius);
 	static bool CheckAlreadyAtGoal(AAIController& Controller, const AActor& TestGoal, float Radius);
@@ -108,7 +108,7 @@ protected:
 
 	void ClearPath();
 	virtual bool Start() override;
-	virtual bool Pause(const UPawnAction* PausedBy) override;
+	virtual bool Pause(const UDEPRECATED_UPawnAction* PausedBy) override;
 	virtual bool Resume() override;
 	virtual void OnFinished(EPawnActionResult::Type WithResult) override;
 	virtual EPawnActionAbortState::Type PerformAbort(EAIForceParam::Type ShouldForce) override;
diff --git a/Engine/Source/Runtime/AIModule/Classes/Actions/PawnAction_Repeat.h b/Engine/Source/Runtime/AIModule/Classes/Actions/PawnAction_Repeat.h
--- a/Engine/Source/Runtime/AIModule/Classes/Actions/PawnAction_Repeat.h
+++ b/Engine/Source/Runtime/AIModule/Classes/Actions/PawnAction_Repeat.h
@@ -8,7 +8,7 @@
 #include "PawnAction_Repeat.generated.h"
 
 UCLASS()
-class AIMODULE_API UPawnAction_Repeat : public UPawnAction
+class AIMODULE_API UDEPRECATED_UPawnAction_Repeat : public UDEPRECATED_UPawnAction
 {
 	GENERATED_UCLASS_BODY()
 
@@ -19,10 +19,10 @@ class AIMODULE_API UPawnAction_Repeat : public UPawnAction
 
 	/** Action to repeat. This instance won't really be run, it's a source for copying actions to be actually performed */
 	UPROPERTY()
-	TObjectPtr<UPawnAction> ActionToRepeat;
+	TObjectPtr<UDEPRECATED_UPawnAction> ActionToRepeat_DEPRECATED;
 
 	UPROPERTY(Transient)
-	TObjectPtr<UPawnAction> RecentActionCopy;
+	TObjectPtr<UDEPRECATED_UPawnAction> RecentActionCopy_DEPRECATED;
 
 	UPROPERTY(Category = PawnAction, EditAnywhere, BlueprintReadOnly)
 	TEnumAsByte<EPawnActionFailHandling::Type> ChildFailureHandlingMode;
@@ -31,13 +31,13 @@ class AIMODULE_API UPawnAction_Repeat : public UPawnAction
 
 	EPawnSubActionTriggeringPolicy::Type SubActionTriggeringPolicy;
 
-	/** @param NumberOfRepeats number of times to repeat action. UPawnAction_Repeat::LoopForever loops forever */
-	static UPawnAction_Repeat* CreateAction(UWorld& World, UPawnAction* ActionToRepeat, int32 NumberOfRepeats, EPawnSubActionTriggeringPolicy::Type InSubActionTriggeringPolicy = EPawnSubActionTriggeringPolicy::CopyBeforeTriggering);
+	/** @param NumberOfRepeats number of times to repeat action. UDEPRECATED_UPawnAction_Repeat::LoopForever loops forever */
+	static UDEPRECATED_UPawnAction_Repeat* CreateAction(UWorld& World, UDEPRECATED_UPawnAction* ActionToRepeat, int32 NumberOfRepeats, EPawnSubActionTriggeringPolicy::Type InSubActionTriggeringPolicy = EPawnSubActionTriggeringPolicy::CopyBeforeTriggering);
 
 protected:
 	virtual bool Start() override;
 	virtual bool Resume() override;
-	virtual void OnChildFinished(UPawnAction& Action, EPawnActionResult::Type WithResult) override;
+	virtual void OnChildFinished(UDEPRECATED_UPawnAction& Action, EPawnActionResult::Type WithResult) override;
 
 	bool PushSubAction();
 };
diff --git a/Engine/Source/Runtime/AIModule/Classes/Actions/PawnAction_Sequence.h b/Engine/Source/Runtime/AIModule/Classes/Actions/PawnAction_Sequence.h
--- a/Engine/Source/Runtime/AIModule/Classes/Actions/PawnAction_Sequence.h
+++ b/Engine/Source/Runtime/AIModule/Classes/Actions/PawnAction_Sequence.h
@@ -8,29 +8,29 @@
 #include "PawnAction_Sequence.generated.h"
 
 UCLASS()
-class AIMODULE_API UPawnAction_Sequence : public UPawnAction
+class AIMODULE_API UDEPRECATED_UPawnAction_Sequence : public UDEPRECATED_UPawnAction
 {
 	GENERATED_UCLASS_BODY()
 
 	UPROPERTY()
-	TArray<TObjectPtr<UPawnAction>> ActionSequence;
+	TArray<TObjectPtr<UDEPRECATED_UPawnAction>> ActionSequence_DEPRECATED;
 
 	UPROPERTY(Category = PawnAction, EditAnywhere, BlueprintReadOnly)
 	TEnumAsByte<EPawnActionFailHandling::Type> ChildFailureHandlingMode;
 
 	UPROPERTY(Transient)
-	TObjectPtr<UPawnAction> RecentActionCopy;
+	TObjectPtr<UDEPRECATED_UPawnAction> RecentActionCopy_DEPRECATED;
 
 	uint32 CurrentActionIndex;
 
 	EPawnSubActionTriggeringPolicy::Type SubActionTriggeringPolicy;
 
-	static UPawnAction_Sequence* CreateAction(UWorld& World, TArray<UPawnAction*>& ActionSequence, EPawnSubActionTriggeringPolicy::Type InSubActionTriggeringPolicy = EPawnSubActionTriggeringPolicy::CopyBeforeTriggering);
+	static UDEPRECATED_UPawnAction_Sequence* CreateAction(UWorld& World, TArray<UDEPRECATED_UPawnAction*>& ActionSequence, EPawnSubActionTriggeringPolicy::Type InSubActionTriggeringPolicy = EPawnSubActionTriggeringPolicy::CopyBeforeTriggering);
 
 protected:
 	virtual bool Start() override;
 	virtual bool Resume() override;
-	virtual void OnChildFinished(UPawnAction& Action, EPawnActionResult::Type WithResult) override;
+	virtual void OnChildFinished(UDEPRECATED_UPawnAction& Action, EPawnActionResult::Type WithResult) override;
 
 	bool PushNextActionCopy();
 };
diff --git a/Engine/Source/Runtime/AIModule/Classes/Actions/PawnAction_Wait.h b/Engine/Source/Runtime/AIModule/Classes/Actions/PawnAction_Wait.h
--- a/Engine/Source/Runtime/AIModule/Classes/Actions/PawnAction_Wait.h
+++ b/Engine/Source/Runtime/AIModule/Classes/Actions/PawnAction_Wait.h
@@ -10,7 +10,7 @@
 
 /** uses system timers rather then ticking */
 UCLASS()
-class AIMODULE_API UPawnAction_Wait : public UPawnAction
+class AIMODULE_API UDEPRECATED_UPawnAction_Wait : public UDEPRECATED_UPawnAction
 {
 	GENERATED_UCLASS_BODY()
 		
@@ -22,10 +22,10 @@ class AIMODULE_API UPawnAction_Wait : public UPawnAction
 	FTimerHandle TimerHandle;
 
 	/** InTimeToWait < 0 (or just FAISystem::InfiniteInterval) will result in waiting forever */
-	static UPawnAction_Wait* CreateAction(UWorld& World, float InTimeToWait = FAISystem::InfiniteInterval);
+	static UDEPRECATED_UPawnAction_Wait* CreateAction(UWorld& World, float InTimeToWait = FAISystem::InfiniteInterval);
 
 	virtual bool Start() override;
-	virtual bool Pause(const UPawnAction* PausedBy) override;
+	virtual bool Pause(const UDEPRECATED_UPawnAction* PausedBy) override;
 	virtual bool Resume() override;
 	virtual EPawnActionAbortState::Type PerformAbort(EAIForceParam::Type ShouldForce) override;
 
diff --git a/Engine/Source/Runtime/AIModule/Classes/Actions/PawnActionsComponent.h b/Engine/Source/Runtime/AIModule/Classes/Actions/PawnActionsComponent.h
--- a/Engine/Source/Runtime/AIModule/Classes/Actions/PawnActionsComponent.h
+++ b/Engine/Source/Runtime/AIModule/Classes/Actions/PawnActionsComponent.h
@@ -22,7 +22,7 @@ struct AIMODULE_API FPawnActionEvent
 	static const int32 FakeActionIndex = INDEX_NONE;
 
 	UPROPERTY()
-	TObjectPtr<UPawnAction> Action;
+	TObjectPtr<UDEPRECATED_UPawnAction> Action;
 
 	EPawnActionEventType::Type EventType;
 
@@ -34,7 +34,7 @@ struct AIMODULE_API FPawnActionEvent
 	FPawnActionEvent() : Action(NULL), EventType(EPawnActionEventType::Invalid), Priority(EAIRequestPriority::MAX), Index(uint32(-1))
 	{}
 
-	FPawnActionEvent(UPawnAction& Action, EPawnActionEventType::Type EventType, uint32 Index);
+	FPawnActionEvent(UDEPRECATED_UPawnAction& Action, EPawnActionEventType::Type EventType, uint32 Index);
 
 	bool operator==(const FPawnActionEvent& Other) const { return (Action == Other.Action) && (EventType == Other.EventType) && (Priority == Other.Priority); }
 };
@@ -50,7 +50,7 @@ struct AIMODULE_API FPawnActionStack
 
 private:
 	UPROPERTY()
-	TObjectPtr<UPawnAction> TopAction;
+	TObjectPtr<UDEPRECATED_UPawnAction> TopAction;
 
 public:
 	void Pause();
@@ -58,13 +58,13 @@ public:
 
 	/** All it does is tie actions into a double-linked list making NewTopAction
 	 *	new stack's top */
-	void PushAction(UPawnAction& NewTopAction);
+	void PushAction(UDEPRECATED_UPawnAction& NewTopAction);
 
 	/** Looks through the double-linked action list looking for specified action
 	 *	and if found action will be popped along with all it's siblings */
-	void PopAction(UPawnAction& ActionToPop);
+	void PopAction(UDEPRECATED_UPawnAction& ActionToPop);
 	
-	FORCEINLINE UPawnAction* GetTop() const { return TopAction; }
+	FORCEINLINE UDEPRECATED_UPawnAction* GetTop() const { return TopAction; }
 
 	FORCEINLINE bool IsEmpty() const { return TopAction == NULL; }
 
@@ -74,8 +74,8 @@ public:
 	int32 GetStackSize() const;
 };
 
-UCLASS()
-class AIMODULE_API UPawnActionsComponent : public UActorComponent
+UCLASS(deprecated, meta = (DeprecationMessage = "PawnActions have been deprecated and are no longer being supported. It will get removed in following UE5 releases. Use GameplayTasks or AITasks instead."))
+class AIMODULE_API UDEPRECATED_UPawnActionsComponent : public UActorComponent
 {
 	GENERATED_UCLASS_BODY()
 
@@ -90,7 +90,7 @@ protected:
 	TArray<FPawnActionEvent> ActionEvents;
 
 	UPROPERTY(Transient)
-	TObjectPtr<UPawnAction> CurrentAction;
+	TObjectPtr<UDEPRECATED_UPawnAction> CurrentAction_DEPRECATED;
 
 	/** set when logic was locked by hi priority stack */
 	uint32 bLockedAILogic : 1;
@@ -108,9 +108,11 @@ public:
 	// blueprint interface
 	//----------------------------------------------------------------------//
 
-	UFUNCTION(BlueprintCallable, Category = "AI|PawnActions", meta = (DisplayName = "PerformAction", ScriptName = "PerformAction"))
-	static bool K2_PerformAction(APawn* Pawn, UPawnAction* Action, TEnumAsByte<EAIRequestPriority::Type> Priority = EAIRequestPriority::HardScript);
-	static bool PerformAction(APawn& Pawn, UPawnAction& Action, TEnumAsByte<EAIRequestPriority::Type> Priority = EAIRequestPriority::HardScript);
+	UFUNCTION(BlueprintCallable, Category = "AI|PawnActions", meta = (DisplayName = "PerformAction_DEPRECATED", ScriptName = "PerformAction", DeprecatedFunction, DeprecationMessage = "PawnActions have been deprecated and are no longer being supported. It will get removed in following UE5 releases. Use GameplayTasks or AITasks instead."))
+	static bool K2_PerformAction(APawn* Pawn, UDEPRECATED_UPawnAction* Action, TEnumAsByte<EAIRequestPriority::Type> Priority = EAIRequestPriority::HardScript);
+
+	UE_DEPRECATED(5.2, "PawnActions have been deprecated and are no longer being supported. It will get removed in following UE5 releases. Use GameplayTasks or AITasks instead.")
+	static bool PerformAction(APawn& Pawn, UDEPRECATED_UPawnAction& Action, TEnumAsByte<EAIRequestPriority::Type> Priority = EAIRequestPriority::HardScript);
 
 	//----------------------------------------------------------------------//
 	// 
@@ -122,23 +124,29 @@ public:
 	FORCEINLINE APawn* GetControlledPawn() { return ControlledPawn; }
 	FORCEINLINE const APawn* GetControlledPawn() const { return ControlledPawn; }
 	FORCEINLINE AController* GetController() { return ControlledPawn ? ControlledPawn->GetController() : NULL; }
-	FORCEINLINE UPawnAction* GetCurrentAction() { return CurrentAction; }
+	FORCEINLINE UDEPRECATED_UPawnAction* GetCurrentAction() { return CurrentAction_DEPRECATED; }
+
+	bool OnEvent(UDEPRECATED_UPawnAction& Action, EPawnActionEventType::Type Event);
 
-	bool OnEvent(UPawnAction& Action, EPawnActionEventType::Type Event);
+	UFUNCTION(BlueprintCallable, Category = PawnAction, meta = (DisplayName = "PushAction_DEPRECATED", ScriptName = "PushAction", DeprecatedFunction, DeprecationMessage = "PawnActions have been deprecated and are no longer being supported. It will get removed in following UE5 releases. Use GameplayTasks or AITasks instead."))
+	bool K2_PushAction(UDEPRECATED_UPawnAction* NewAction, EAIRequestPriority::Type Priority, UObject* Instigator = NULL);
 
-	UFUNCTION(BlueprintCallable, Category = PawnAction, meta = (DisplayName = "PushAction", ScriptName = "PushAction"))
-	bool K2_PushAction(UPawnAction* NewAction, EAIRequestPriority::Type Priority, UObject* Instigator = NULL);
-	bool PushAction(UPawnAction& NewAction, EAIRequestPriority::Type Priority, UObject* Instigator = NULL);	
+	UE_DEPRECATED(5.2, "PawnActions have been deprecated and are no longer being supported. It will get removed in following UE5 releases. Use GameplayTasks or AITasks instead.")
+	bool PushAction(UDEPRECATED_UPawnAction& NewAction, EAIRequestPriority::Type Priority, UObject* Instigator = NULL);	
 
 	/** Aborts given action instance */
-	UFUNCTION(BlueprintCallable, Category = PawnAction, meta = (DisplayName = "AbortAction", ScriptName = "AbortAction"))
-	EPawnActionAbortState::Type K2_AbortAction(UPawnAction* ActionToAbort);
-	EPawnActionAbortState::Type AbortAction(UPawnAction& ActionToAbort);
+	UFUNCTION(BlueprintCallable, Category = PawnAction, meta = (DisplayName = "AbortAction_DEPRECATED", ScriptName = "AbortAction", DeprecatedFunction, DeprecationMessage = "PawnActions have been deprecated and are no longer being supported. It will get removed in following UE5 releases. Use GameplayTasks or AITasks instead."))
+	EPawnActionAbortState::Type K2_AbortAction(UDEPRECATED_UPawnAction* ActionToAbort);
+
+	UE_DEPRECATED(5.2, "PawnActions have been deprecated and are no longer being supported. It will get removed in following UE5 releases. Use GameplayTasks or AITasks instead.")
+	EPawnActionAbortState::Type AbortAction(UDEPRECATED_UPawnAction& ActionToAbort);
 
 	/** Aborts given action instance */
-	UFUNCTION(BlueprintCallable, Category = PawnAction, meta = (DisplayName = "ForceAbortAction", ScriptName = "ForceAbortAction"))
-	EPawnActionAbortState::Type K2_ForceAbortAction(UPawnAction* ActionToAbort);
-	EPawnActionAbortState::Type ForceAbortAction(UPawnAction& ActionToAbort);
+	UFUNCTION(BlueprintCallable, Category = PawnAction, meta = (DisplayName = "ForceAbortAction_DEPRECATED", ScriptName = "ForceAbortAction", DeprecatedFunction, DeprecationMessage = "PawnActions have been deprecated and are no longer being supported. It will get removed in following UE5 releases. Use GameplayTasks or AITasks instead."))
+	EPawnActionAbortState::Type K2_ForceAbortAction(UDEPRECATED_UPawnAction* ActionToAbort);
+
+	UE_DEPRECATED(5.2, "PawnActions have been deprecated and are no longer being supported. It will get removed in following UE5 releases. Use GameplayTasks or AITasks instead.")
+	EPawnActionAbortState::Type ForceAbortAction(UDEPRECATED_UPawnAction& ActionToAbort);
 
 	/** removes all actions instigated with Priority by Instigator
 	 *	@param Priority if equal to EAIRequestPriority::MAX then all priority queues will be searched. 
@@ -148,8 +156,8 @@ public:
 	
 	void TickComponent(float DeltaTime, enum ELevelTick TickType, FActorComponentTickFunction *ThisTickFunction) override;
 
-	FORCEINLINE UPawnAction* GetActiveAction(EAIRequestPriority::Type Priority) const { return ActionStacks[Priority].GetTop(); }
-	bool HasActiveActionOfType(EAIRequestPriority::Type Priority, TSubclassOf<UPawnAction> PawnActionClass) const;
+	FORCEINLINE UDEPRECATED_UPawnAction* GetActiveAction(EAIRequestPriority::Type Priority) const { return ActionStacks[Priority].GetTop(); }
+	bool HasActiveActionOfType(EAIRequestPriority::Type Priority, TSubclassOf<UDEPRECATED_UPawnAction> PawnActionClass) const;
 
 #if ENABLE_VISUAL_LOG
 	void DescribeSelfToVisLog(struct FVisualLogEntry* Snapshot) const;
@@ -174,5 +182,5 @@ protected:
 
 private:
 	/** Removed all pending action events associated with PawnAction. Private to make sure it's called only in special cases */
-	void RemoveEventsForAction(UPawnAction& PawnAction);
+	void RemoveEventsForAction(UDEPRECATED_UPawnAction& PawnAction);
 };
diff --git a/Engine/Source/Runtime/AIModule/Classes/BehaviorTree/Tasks/BTTask_PawnActionBase.h b/Engine/Source/Runtime/AIModule/Classes/BehaviorTree/Tasks/BTTask_PawnActionBase.h
--- a/Engine/Source/Runtime/AIModule/Classes/BehaviorTree/Tasks/BTTask_PawnActionBase.h
+++ b/Engine/Source/Runtime/AIModule/Classes/BehaviorTree/Tasks/BTTask_PawnActionBase.h
@@ -7,7 +7,7 @@
 #include "BehaviorTree/BTTaskNode.h"
 #include "BTTask_PawnActionBase.generated.h"
 
-class UPawnAction;
+class UDEPRECATED_UPawnAction;
 
 enum class EPawnActionTaskResult : uint8
 {
@@ -36,17 +36,17 @@ class AIMODULE_API UBTTask_PawnActionBase : public UBTTaskNode
 protected:
 
 	/** starts executing pawn action */
-	EBTNodeResult::Type PushAction(UBehaviorTreeComponent& OwnerComp, UPawnAction& Action);
+	EBTNodeResult::Type PushAction(UBehaviorTreeComponent& OwnerComp, UDEPRECATED_UPawnAction& Action);
 
 	/** action observer, updates state of task */
-	virtual void OnActionEvent(UPawnAction& Action, EPawnActionEventType::Type Event);
+	virtual void OnActionEvent(UDEPRECATED_UPawnAction& Action, EPawnActionEventType::Type Event);
 
 	/** called when action is removed from stack (FinishedAborting) by some external event
 	 *  default behavior: finish task as failed */
-	virtual void OnActionLost(UPawnAction& Action);
+	virtual void OnActionLost(UDEPRECATED_UPawnAction& Action);
 
 public:
 
 	/** helper functions, should be used when behavior tree task deals with pawn actions, but can't derive from this class */
-	static EPawnActionTaskResult ActionEventHandler(UBTTaskNode* TaskNode, UPawnAction& Action, EPawnActionEventType::Type Event);
+	static EPawnActionTaskResult ActionEventHandler(UBTTaskNode* TaskNode, UDEPRECATED_UPawnAction& Action, EPawnActionEventType::Type Event);
 };
diff --git a/Engine/Source/Runtime/AIModule/Classes/BehaviorTree/Tasks/BTTask_PushPawnAction.h b/Engine/Source/Runtime/AIModule/Classes/BehaviorTree/Tasks/BTTask_PushPawnAction.h
--- a/Engine/Source/Runtime/AIModule/Classes/BehaviorTree/Tasks/BTTask_PushPawnAction.h
+++ b/Engine/Source/Runtime/AIModule/Classes/BehaviorTree/Tasks/BTTask_PushPawnAction.h
@@ -7,7 +7,7 @@
 #include "BehaviorTree/Tasks/BTTask_PawnActionBase.h"
 #include "BTTask_PushPawnAction.generated.h"
 
-class UPawnAction;
+class UDEPRECATED_UPawnAction;
 
 /**
  * Action task node.
@@ -19,8 +19,8 @@ class AIMODULE_API UBTTask_PushPawnAction : public UBTTask_PawnActionBase
 	GENERATED_UCLASS_BODY()
 
 protected:
-	UPROPERTY(EditAnywhere, Instanced, Category = Action)
-	TObjectPtr<UPawnAction> Action;
+	UPROPERTY(Instanced)
+	TObjectPtr<UDEPRECATED_UPawnAction> Action_DEPRECATED;
 
 	virtual EBTNodeResult::Type ExecuteTask(UBehaviorTreeComponent& OwnerComp, uint8* NodeMemory) override;
 	virtual FString GetStaticDescription() const override;
diff --git a/Engine/Source/Runtime/AIModule/Private/AIController.cpp b/Engine/Source/Runtime/AIModule/Private/AIController.cpp
--- a/Engine/Source/Runtime/AIModule/Private/AIController.cpp
+++ b/Engine/Source/Runtime/AIModule/Private/AIController.cpp
@@ -47,7 +47,7 @@ AAIController::AAIController(const FObjectInitializer& ObjectInitializer)
 		PathFollowingComponent->OnRequestFinished.AddUObject(this, &AAIController::OnMoveCompleted);
 	}
 
-	ActionsComp = CreateOptionalDefaultSubobject<UPawnActionsComponent>(TEXT("ActionsComp"));
+	ActionsComp_DEPRECATED = CreateOptionalDefaultSubobject<UDEPRECATED_UPawnActionsComponent>(TEXT("ActionsComp"));
 
 	bSkipExtraLOSChecks = true;
 	bWantsPlayerState = false;
@@ -1082,9 +1082,11 @@ bool AAIController::SuggestTossVelocity(FVector& OutTossVelocity, FVector Start,
 
 	return UGameplayStatics::SuggestProjectileVelocity(this, OutTossVelocity, Start, End, TossSpeed, bPreferHighArc, CollisionRadius, GravityOverride, TraceOption);
 }
-bool AAIController::PerformAction(UPawnAction& Action, EAIRequestPriority::Type Priority, UObject* const InInstigator /*= NULL*/)
+bool AAIController::PerformAction(UDEPRECATED_UPawnAction& Action, EAIRequestPriority::Type Priority, UObject* const InInstigator /*= NULL*/)
 {
-	return ActionsComp != NULL && ActionsComp->PushAction(Action, Priority, InInstigator);
+PRAGMA_DISABLE_DEPRECATION_WARNINGS
+	return ActionsComp_DEPRECATED != NULL && ActionsComp_DEPRECATED->PushAction(Action, Priority, InInstigator);
+PRAGMA_ENABLE_DEPRECATION_WARNINGS
 }
 
 FString AAIController::GetDebugIcon() const
diff --git a/Engine/Source/Runtime/AIModule/Private/Actions/PawnAction.cpp b/Engine/Source/Runtime/AIModule/Private/Actions/PawnAction.cpp
--- a/Engine/Source/Runtime/AIModule/Private/Actions/PawnAction.cpp
+++ b/Engine/Source/Runtime/AIModule/Private/Actions/PawnAction.cpp
@@ -10,6 +10,8 @@
 
 #include UE_INLINE_GENERATED_CPP_BY_NAME(PawnAction)
 
+PRAGMA_DISABLE_DEPRECATION_WARNINGS
+
 DEFINE_LOG_CATEGORY(LogPawnAction);
 
 namespace
@@ -22,7 +24,7 @@ namespace
 	}
 }
 
-UPawnAction::UPawnAction(const FObjectInitializer& ObjectInitializer)
+UDEPRECATED_UPawnAction::UDEPRECATED_UPawnAction(const FObjectInitializer& ObjectInitializer)
 	: Super(ObjectInitializer)
 	, RequiredResources(FAIResourcesSet::AllResources)
 	, AbortState(EPawnActionAbortState::NeverStarted)
@@ -34,16 +36,16 @@ UPawnAction::UPawnAction(const FObjectInitializer& ObjectInitializer)
 	IndexOnStack = INDEX_NONE;
 }
 
-UWorld* UPawnAction::GetWorld() const
+UWorld* UDEPRECATED_UPawnAction::GetWorld() const
 {
-	return OwnerComponent ? OwnerComponent->GetWorld() : Cast<UWorld>(GetOuter());
+	return OwnerComponent_DEPRECATED ? OwnerComponent_DEPRECATED->GetWorld() : Cast<UWorld>(GetOuter());
 }
 
-void UPawnAction::Tick(float DeltaTime)
+void UDEPRECATED_UPawnAction::Tick(float DeltaTime)
 {
 }
 
-EPawnActionAbortState::Type UPawnAction::Abort(EAIForceParam::Type ShouldForce)
+EPawnActionAbortState::Type UDEPRECATED_UPawnAction::Abort(EAIForceParam::Type ShouldForce)
 {
 	// if already aborting, and this request is not Forced, just skip it
 	if (AbortState != EPawnActionAbortState::NotBeingAborted && ShouldForce != EAIForceParam::Force)
@@ -66,14 +68,14 @@ EPawnActionAbortState::Type UPawnAction::Abort(EAIForceParam::Type ShouldForce)
 
 	SetAbortState(EPawnActionAbortState::MarkPendingAbort);
 
-	if (ChildAction != NULL)
+	if (ChildAction_DEPRECATED != NULL)
 	{
-		ChildResult = ChildAction->Abort(ShouldForce);
+		ChildResult = ChildAction_DEPRECATED->Abort(ShouldForce);
 
 		if (ChildResult == EPawnActionAbortState::NotBeingAborted)
 		{
-			UE_VLOG(GetPawn(), LogPawnAction, Error, TEXT("%s> ChildAction %s failed to carry out proper abortion! Things might get ugly..")
-				, *GetName(), *ChildAction->GetName());
+			UE_VLOG(GetPawn(), LogPawnAction, Error, TEXT("%s> ChildAction_DEPRECATED %s failed to carry out proper abortion! Things might get ugly..")
+				, *GetName(), *ChildAction_DEPRECATED->GetName());
 
 			// fake proper result and hope for the best!
 			ChildResult = EPawnActionAbortState::AbortDone;
@@ -120,17 +122,17 @@ EPawnActionAbortState::Type UPawnAction::Abort(EAIForceParam::Type ShouldForce)
 	return Result;
 }
 
-APawn* UPawnAction::GetPawn() const
+APawn* UDEPRECATED_UPawnAction::GetPawn() const
 {
-	return OwnerComponent ? OwnerComponent->GetControlledPawn() : NULL;
+	return OwnerComponent_DEPRECATED ? OwnerComponent_DEPRECATED->GetControlledPawn() : NULL;
 }
 
-AController* UPawnAction::GetController() const
+AController* UDEPRECATED_UPawnAction::GetController() const
 {
-	return OwnerComponent ? OwnerComponent->GetController() : NULL; 
+	return OwnerComponent_DEPRECATED ? OwnerComponent_DEPRECATED->GetController() : NULL;
 }
 
-void UPawnAction::SetAbortState(EPawnActionAbortState::Type NewAbortState)
+void UDEPRECATED_UPawnAction::SetAbortState(EPawnActionAbortState::Type NewAbortState)
 {
 	// allowing only progression
 	if (NewAbortState <= AbortState)
@@ -145,28 +147,28 @@ void UPawnAction::SetAbortState(EPawnActionAbortState::Type NewAbortState)
 	}
 }
 
-void UPawnAction::SendEvent(EPawnActionEventType::Type Event)
+void UDEPRECATED_UPawnAction::SendEvent(EPawnActionEventType::Type Event)
 {
-	if (IsValid(OwnerComponent))
+	if (IsValid(OwnerComponent_DEPRECATED))
 	{
 		// this will get communicated to parent action if needed, latently 
-		OwnerComponent->OnEvent(*this, Event);
+		OwnerComponent_DEPRECATED->OnEvent(*this, Event);
 	}
 
 	ActionObserver.ExecuteIfBound(*this, Event);
 }
 
-void UPawnAction::StopWaitingForMessages()
+void UDEPRECATED_UPawnAction::StopWaitingForMessages()
 {
 	MessageHandlers.Reset();
 }
 
-void UPawnAction::SetFinishResult(EPawnActionResult::Type Result)
+void UDEPRECATED_UPawnAction::SetFinishResult(EPawnActionResult::Type Result)
 {
 	// once return value had been set it's no longer possible to back to InProgress
 	if (Result <= EPawnActionResult::InProgress)
 	{
-		UE_VLOG(GetPawn(), LogPawnAction, Warning, TEXT("%s> UPawnAction::SetFinishResult setting FinishResult as EPawnActionResult::InProgress or EPawnActionResult::NotStarted - should not be happening"), *GetName());
+		UE_VLOG(GetPawn(), LogPawnAction, Warning, TEXT("%s> UDEPRECATED_UPawnAction::SetFinishResult setting FinishResult as EPawnActionResult::InProgress or EPawnActionResult::NotStarted - should not be happening"), *GetName());
 		return;
 	}
 
@@ -176,14 +178,14 @@ void UPawnAction::SetFinishResult(EPawnActionResult::Type Result)
 	}
 }
 
-void UPawnAction::SetOwnerComponent(UPawnActionsComponent* Component)
+void UDEPRECATED_UPawnAction::SetOwnerComponent(UDEPRECATED_UPawnActionsComponent* Component)
 {
-	if (OwnerComponent != NULL && OwnerComponent != Component)
+	if (OwnerComponent_DEPRECATED != NULL && OwnerComponent_DEPRECATED != Component)
 	{
-		UE_VLOG(GetPawn(), LogPawnAction, Warning, TEXT("%s> UPawnAction::SetOwnerComponent called to change already set valid owner component"), *GetName());
+		UE_VLOG(GetPawn(), LogPawnAction, Warning, TEXT("%s> UDEPRECATED_UPawnAction::SetOwnerComponent called to change already set valid owner component"), *GetName());
 	}
 
-	OwnerComponent = Component;
+	OwnerComponent_DEPRECATED = Component;
 	if (Component != NULL)
 	{
 		AAIController* AIController = Cast<AAIController>(Component->GetController());
@@ -194,7 +196,7 @@ void UPawnAction::SetOwnerComponent(UPawnActionsComponent* Component)
 	}
 }
 
-void UPawnAction::SetInstigator(UObject* const InInstigator)
+void UDEPRECATED_UPawnAction::SetInstigator(UObject* const InInstigator)
 { 
 	if (Instigator && Instigator != InInstigator)
 	{
@@ -204,7 +206,7 @@ void UPawnAction::SetInstigator(UObject* const InInstigator)
 	Instigator = InInstigator; 
 }
 
-void UPawnAction::Finish(TEnumAsByte<EPawnActionResult::Type> WithResult)
+void UDEPRECATED_UPawnAction::Finish(TEnumAsByte<EPawnActionResult::Type> WithResult)
 {
 	UE_VLOG(GetPawn(), LogPawnAction, Log, TEXT("%s> finishing with result %s")
 		, *GetName(), *GetActionResultName(WithResult));
@@ -216,7 +218,7 @@ void UPawnAction::Finish(TEnumAsByte<EPawnActionResult::Type> WithResult)
 	SendEvent(EPawnActionEventType::FinishedExecution);
 }
 
-bool UPawnAction::Activate()
+bool UDEPRECATED_UPawnAction::Activate()
 {
 	bool bResult = false; 
 
@@ -254,7 +256,7 @@ bool UPawnAction::Activate()
 	return bResult;
 }
 
-void UPawnAction::OnPopped()
+void UDEPRECATED_UPawnAction::OnPopped()
 {
 	// not calling OnFinish if action haven't actually started
 	if (!bFailedToStart || bAlwaysNotifyOnFinished)
@@ -263,7 +265,7 @@ void UPawnAction::OnPopped()
 	}
 }
 
-bool UPawnAction::Start()
+bool UDEPRECATED_UPawnAction::Start()
 {
 	AbortState = EPawnActionAbortState::NotBeingAborted;
 	FinishResult = EPawnActionResult::InProgress;
@@ -272,10 +274,10 @@ bool UPawnAction::Start()
 }
 
 
-bool UPawnAction::Pause(const UPawnAction* PausedBy)
+bool UDEPRECATED_UPawnAction::Pause(const UDEPRECATED_UPawnAction* PausedBy)
 {
 	// parent should be paused anyway
-	ensure(ParentAction == NULL || ParentAction->IsPaused() == true);
+	ensure(ParentAction_DEPRECATED == NULL || ParentAction_DEPRECATED->IsPaused() == true);
 
 	// don't pause twice, this should be guaranteed by the PawnActionsComponent
 	ensure(bPaused == false);
@@ -290,18 +292,18 @@ bool UPawnAction::Pause(const UPawnAction* PausedBy)
 	
 	bPaused = true;
 
-	if (ChildAction)
+	if (ChildAction_DEPRECATED)
 	{
-		ChildAction->Pause(PausedBy);
+		ChildAction_DEPRECATED->Pause(PausedBy);
 	}
 
 	return bPaused;
 }
 
-bool UPawnAction::Resume()
+bool UDEPRECATED_UPawnAction::Resume()
 {
 	// parent should be paused anyway
-	ensure(ParentAction == NULL || ParentAction->IsPaused() == true);
+	ensure(ParentAction_DEPRECATED == NULL || ParentAction_DEPRECATED->IsPaused() == true);
 
 	// do not unpause twice
 	if (bPaused == false)
@@ -309,12 +311,12 @@ bool UPawnAction::Resume()
 		return false;
 	}
 	
-	ensure(ChildAction == NULL);
+	ensure(ChildAction_DEPRECATED == NULL);
 
-	if (ChildAction)
+	if (ChildAction_DEPRECATED)
 	{
-		UE_VLOG(GetPawn(), LogPawnAction, Log, TEXT("%s> Resuming child, %s"), *GetName(), *ChildAction->GetName());
-		ChildAction->Resume();
+		UE_VLOG(GetPawn(), LogPawnAction, Log, TEXT("%s> Resuming child, %s"), *GetName(), *ChildAction_DEPRECATED->GetName());
+		ChildAction_DEPRECATED->Resume();
 	}
 	else
 	{
@@ -325,49 +327,49 @@ bool UPawnAction::Resume()
 	return !bPaused;
 }
 
-void UPawnAction::OnFinished(EPawnActionResult::Type WithResult)
+void UDEPRECATED_UPawnAction::OnFinished(EPawnActionResult::Type WithResult)
 {
 }
 
-void UPawnAction::OnChildFinished(UPawnAction& Action, EPawnActionResult::Type WithResult)
+void UDEPRECATED_UPawnAction::OnChildFinished(UDEPRECATED_UPawnAction& Action, EPawnActionResult::Type WithResult)
 {
 	UE_VLOG(GetPawn(), LogPawnAction, Log, TEXT("%s> Child \'%s\' finished with result %s")
 		, *GetName(), *Action.GetName(), *GetActionResultName(WithResult));
 
-	ensure(Action.ParentAction == this);
-	ensure(ChildAction == &Action);
-	Action.ParentAction = NULL;
-	ChildAction = NULL;
+	ensure(Action.ParentAction_DEPRECATED == this);
+	ensure(ChildAction_DEPRECATED == &Action);
+	Action.ParentAction_DEPRECATED = NULL;
+	ChildAction_DEPRECATED = NULL;
 }
 
-void UPawnAction::OnPushed()
+void UDEPRECATED_UPawnAction::OnPushed()
 {
 	IndexOnStack = 0;
-	UPawnAction* PrevAction = ParentAction;
+	UDEPRECATED_UPawnAction* PrevAction = ParentAction_DEPRECATED;
 	while (PrevAction)
 	{
 		++IndexOnStack;
-		PrevAction = PrevAction->ParentAction;
+		PrevAction = PrevAction->ParentAction_DEPRECATED;
 	}
 
 	UE_VLOG(GetPawn(), LogPawnAction, Log, TEXT("%s> Pushed with priority %s, IndexOnStack: %d, instigator %s")
 		, *GetName(), *GetPriorityName(), IndexOnStack, *GetNameSafe(Instigator));
 }
 
-bool UPawnAction::PushChildAction(UPawnAction& Action)
+bool UDEPRECATED_UPawnAction::PushChildAction(UDEPRECATED_UPawnAction& Action)
 {
 	bool bResult = false;
 	
-	if (OwnerComponent != NULL)
+	if (OwnerComponent_DEPRECATED != NULL)
 	{
-		UE_CVLOG( ChildAction != NULL
-			, GetPawn(), LogPawnAction, Log, TEXT("%s> Pushing child action %s while already having ChildAction set to %s")
-			, *GetName(), *Action.GetName(), *ChildAction->GetName());
+		UE_CVLOG( ChildAction_DEPRECATED != NULL
+			, GetPawn(), LogPawnAction, Log, TEXT("%s> Pushing child action %s while already having ChildAction_DEPRECATED set to %s")
+			, *GetName(), *Action.GetName(), *ChildAction_DEPRECATED->GetName());
 		
 		// copy runtime data
 		// note that priority and instigator will get assigned as part of PushAction.
 
-		bResult = OwnerComponent->PushAction(Action, GetPriority(), Instigator);
+		bResult = OwnerComponent_DEPRECATED->PushAction(Action, GetPriority(), Instigator);
 
 		// adding a check to make sure important data has been set 
 		ensure(Action.GetPriority() == GetPriority() && Action.GetInstigator() == GetInstigator());
@@ -380,21 +382,21 @@ bool UPawnAction::PushChildAction(UPawnAction& Action)
 // messaging
 //----------------------------------------------------------------------//
 
-void UPawnAction::WaitForMessage(FName MessageType, FAIRequestID InRequestID)
+void UDEPRECATED_UPawnAction::WaitForMessage(FName MessageType, FAIRequestID InRequestID)
 {
-	MessageHandlers.Add(FAIMessageObserver::Create(BrainComp, MessageType, InRequestID.GetID(), FOnAIMessage::CreateUObject(this, &UPawnAction::HandleAIMessage)));
+	MessageHandlers.Add(FAIMessageObserver::Create(BrainComp, MessageType, InRequestID.GetID(), FOnAIMessage::CreateUObject(this, &UDEPRECATED_UPawnAction::HandleAIMessage)));
 }
 
 //----------------------------------------------------------------------//
 // blueprint interface
 //----------------------------------------------------------------------//
 
-UPawnAction* UPawnAction::CreateActionInstance(UObject* WorldContextObject, TSubclassOf<UPawnAction> ActionClass)
+UDEPRECATED_UPawnAction* UDEPRECATED_UPawnAction::CreateActionInstance(UObject* WorldContextObject, TSubclassOf<UDEPRECATED_UPawnAction> ActionClass)
 {
 	UWorld* World = GEngine->GetWorldFromContextObject(WorldContextObject, EGetWorldErrorMode::LogAndReturnNull);
 	if (World && ActionClass)
 	{
-		return NewObject<UPawnAction>(World, ActionClass);
+		return NewObject<UDEPRECATED_UPawnAction>(World, ActionClass);
 	}
 	return NULL;
 }
@@ -403,7 +405,7 @@ UPawnAction* UPawnAction::CreateActionInstance(UObject* WorldContextObject, TSub
 // debug
 //----------------------------------------------------------------------//
 
-FString UPawnAction::GetStateDescription() const
+FString UDEPRECATED_UPawnAction::GetStateDescription() const
 {
 	static const UEnum* AbortStateEnum = StaticEnum<EPawnActionAbortState::Type>(); 
 		
@@ -414,15 +416,16 @@ FString UPawnAction::GetStateDescription() const
 	return IsPaused() ? TEXT("Paused") : TEXT("Active");
 }
 
-FString UPawnAction::GetPriorityName() const
+FString UDEPRECATED_UPawnAction::GetPriorityName() const
 {
 	static const UEnum* Enum = StaticEnum<EAIRequestPriority::Type>();
 	check(Enum);
 	return Enum->GetNameStringByValue(GetPriority());
 }
 
-FString UPawnAction::GetDisplayName() const
+FString UDEPRECATED_UPawnAction::GetDisplayName() const
 {
 	return GetClass()->GetName();
 }
 
+PRAGMA_ENABLE_DEPRECATION_WARNINGS
diff --git a/Engine/Source/Runtime/AIModule/Private/Actions/PawnAction_BlueprintBase.cpp b/Engine/Source/Runtime/AIModule/Private/Actions/PawnAction_BlueprintBase.cpp
--- a/Engine/Source/Runtime/AIModule/Private/Actions/PawnAction_BlueprintBase.cpp
+++ b/Engine/Source/Runtime/AIModule/Private/Actions/PawnAction_BlueprintBase.cpp
@@ -5,20 +5,22 @@
 
 #include UE_INLINE_GENERATED_CPP_BY_NAME(PawnAction_BlueprintBase)
 
-UPawnAction_BlueprintBase::UPawnAction_BlueprintBase(const FObjectInitializer& ObjectInitializer)
+PRAGMA_DISABLE_DEPRECATION_WARNINGS
+
+UDEPRECATED_UPawnAction_BlueprintBase::UDEPRECATED_UPawnAction_BlueprintBase(const FObjectInitializer& ObjectInitializer)
 	: Super(ObjectInitializer)
 {
-	UClass* StopAtClass = UPawnAction_BlueprintBase::StaticClass();
+	UClass* StopAtClass = UDEPRECATED_UPawnAction_BlueprintBase::StaticClass();
 	bWantsTick = BlueprintNodeHelpers::HasBlueprintFunction(TEXT("ActionTick"), *this, *StopAtClass);
 }
 
-void UPawnAction_BlueprintBase::Tick(float DeltaTime)
+void UDEPRECATED_UPawnAction_BlueprintBase::Tick(float DeltaTime)
 {
 	// no need to call super implementation
 	ActionTick(GetPawn(), DeltaTime);
 }
 
-bool UPawnAction_BlueprintBase::Start()
+bool UDEPRECATED_UPawnAction_BlueprintBase::Start()
 {
 	const bool bHasBeenEverStarted = HasBeenStarted();
 	const bool bSuperResult = Super::Start();
@@ -31,7 +33,7 @@ bool UPawnAction_BlueprintBase::Start()
 	return bSuperResult;
 }
 
-bool UPawnAction_BlueprintBase::Pause(const UPawnAction* PausedBy)
+bool UDEPRECATED_UPawnAction_BlueprintBase::Pause(const UDEPRECATED_UPawnAction* PausedBy)
 {
 	const bool bResult = Super::Pause(PausedBy);
 	if (bResult)
@@ -41,7 +43,7 @@ bool UPawnAction_BlueprintBase::Pause(const UPawnAction* PausedBy)
 	return bResult;
 }
 
-bool UPawnAction_BlueprintBase::Resume()
+bool UDEPRECATED_UPawnAction_BlueprintBase::Resume()
 {
 	const bool bResult = Super::Resume();
 	if (bResult)
@@ -52,8 +54,9 @@ bool UPawnAction_BlueprintBase::Resume()
 	return bResult;
 }
 
-void UPawnAction_BlueprintBase::OnFinished(EPawnActionResult::Type WithResult)
+void UDEPRECATED_UPawnAction_BlueprintBase::OnFinished(EPawnActionResult::Type WithResult)
 {
 	ActionFinished(GetPawn(), WithResult);
 }
 
+PRAGMA_ENABLE_DEPRECATION_WARNINGS
diff --git a/Engine/Source/Runtime/AIModule/Private/Actions/PawnAction_Move.cpp b/Engine/Source/Runtime/AIModule/Private/Actions/PawnAction_Move.cpp
--- a/Engine/Source/Runtime/AIModule/Private/Actions/PawnAction_Move.cpp
+++ b/Engine/Source/Runtime/AIModule/Private/Actions/PawnAction_Move.cpp
@@ -9,7 +9,9 @@
 
 #include UE_INLINE_GENERATED_CPP_BY_NAME(PawnAction_Move)
 
-UPawnAction_Move::UPawnAction_Move(const FObjectInitializer& ObjectInitializer)
+PRAGMA_DISABLE_DEPRECATION_WARNINGS
+
+UDEPRECATED_UPawnAction_Move::UDEPRECATED_UPawnAction_Move(const FObjectInitializer& ObjectInitializer)
 	: Super(ObjectInitializer)
 	, GoalLocation(FAISystem::InvalidLocation)
 	, AcceptableRadius(30.f)
@@ -26,7 +28,7 @@ UPawnAction_Move::UPawnAction_Move(const FObjectInitializer& ObjectInitializer)
 	bAlwaysNotifyOnFinished = true;
 }
 
-void UPawnAction_Move::BeginDestroy()
+void UDEPRECATED_UPawnAction_Move::BeginDestroy()
 {
 	ClearTimers();
 	ClearPath();
@@ -34,14 +36,14 @@ void UPawnAction_Move::BeginDestroy()
 	Super::BeginDestroy();
 }
 
-UPawnAction_Move* UPawnAction_Move::CreateAction(UWorld& World, AActor* GoalActor, EPawnActionMoveMode::Type Mode)
+UDEPRECATED_UPawnAction_Move* UDEPRECATED_UPawnAction_Move::CreateAction(UWorld& World, AActor* GoalActor, EPawnActionMoveMode::Type Mode)
 {
 	if (GoalActor == NULL)
 	{
 		return NULL;
 	}
 
-	UPawnAction_Move* Action = UPawnAction::CreateActionInstance<UPawnAction_Move>(World);
+	UDEPRECATED_UPawnAction_Move* Action = UDEPRECATED_UPawnAction::CreateActionInstance<UDEPRECATED_UPawnAction_Move>(World);
 	if (Action)
 	{
 		Action->GoalActor = GoalActor;
@@ -51,14 +53,14 @@ UPawnAction_Move* UPawnAction_Move::CreateAction(UWorld& World, AActor* GoalActo
 	return Action;
 }
 
-UPawnAction_Move* UPawnAction_Move::CreateAction(UWorld& World, const FVector& GoalLocation, EPawnActionMoveMode::Type Mode)
+UDEPRECATED_UPawnAction_Move* UDEPRECATED_UPawnAction_Move::CreateAction(UWorld& World, const FVector& GoalLocation, EPawnActionMoveMode::Type Mode)
 {
 	if (FAISystem::IsValidLocation(GoalLocation) == false)
 	{
 		return NULL;
 	}
 
-	UPawnAction_Move* Action = UPawnAction::CreateActionInstance<UPawnAction_Move>(World);
+	UDEPRECATED_UPawnAction_Move* Action = UDEPRECATED_UPawnAction::CreateActionInstance<UDEPRECATED_UPawnAction_Move>(World);
 	if (Action)
 	{
 		Action->GoalLocation = GoalLocation;
@@ -68,7 +70,7 @@ UPawnAction_Move* UPawnAction_Move::CreateAction(UWorld& World, const FVector& G
 	return Action;
 }
 
-bool UPawnAction_Move::Start()
+bool UDEPRECATED_UPawnAction_Move::Start()
 {
 	bool bResult = Super::Start();
 	if (bResult)
@@ -79,7 +81,7 @@ bool UPawnAction_Move::Start()
 	return bResult;
 }
 
-EPathFollowingRequestResult::Type UPawnAction_Move::RequestMove(AAIController& Controller)
+EPathFollowingRequestResult::Type UDEPRECATED_UPawnAction_Move::RequestMove(AAIController& Controller)
 {
 	EPathFollowingRequestResult::Type RequestResult = EPathFollowingRequestResult::Failed;
 	
@@ -115,13 +117,13 @@ EPathFollowingRequestResult::Type UPawnAction_Move::RequestMove(AAIController& C
 	}
 	else
 	{
-		UE_VLOG(&Controller, LogPawnAction, Warning, TEXT("UPawnAction_Move::Start: no valid move goal set"));
+		UE_VLOG(&Controller, LogPawnAction, Warning, TEXT("UDEPRECATED_UPawnAction_Move::Start: no valid move goal set"));
 	}
 
 	return RequestResult;
 }
 
-bool UPawnAction_Move::PerformMoveAction()
+bool UDEPRECATED_UPawnAction_Move::PerformMoveAction()
 {
 	AAIController* MyController = Cast<AAIController>(GetController());
 	if (MyController == NULL)
@@ -132,7 +134,7 @@ bool UPawnAction_Move::PerformMoveAction()
 	if (bUsePathfinding && MyController->ShouldPostponePathUpdates())
 	{
 		UE_VLOG(MyController, LogPawnAction, Log, TEXT("Can't path right now, waiting..."));
-		MyController->GetWorldTimerManager().SetTimer(TimerHandle_DeferredPerformMoveAction, this, &UPawnAction_Move::DeferredPerformMoveAction, 0.1f);
+		MyController->GetWorldTimerManager().SetTimer(TimerHandle_DeferredPerformMoveAction, this, &UDEPRECATED_UPawnAction_Move::DeferredPerformMoveAction, 0.1f);
 		return true;
 	}
 
@@ -159,7 +161,7 @@ bool UPawnAction_Move::PerformMoveAction()
 	return bResult;
 }
 
-void UPawnAction_Move::DeferredPerformMoveAction()
+void UDEPRECATED_UPawnAction_Move::DeferredPerformMoveAction()
 {
 	const bool bResult = PerformMoveAction();
 	if (!bResult)
@@ -168,7 +170,7 @@ void UPawnAction_Move::DeferredPerformMoveAction()
 	}
 }
 
-bool UPawnAction_Move::Pause(const UPawnAction* PausedBy)
+bool UDEPRECATED_UPawnAction_Move::Pause(const UDEPRECATED_UPawnAction* PausedBy)
 {
 	bool bResult = Super::Pause(PausedBy);
 	if (bResult)
@@ -182,7 +184,7 @@ bool UPawnAction_Move::Pause(const UPawnAction* PausedBy)
 	return bResult;
 }
 
-bool UPawnAction_Move::Resume()
+bool UDEPRECATED_UPawnAction_Move::Resume()
 {
 	if (GoalActor != NULL && GoalActor->IsPendingKillPending())
 	{
@@ -213,7 +215,7 @@ bool UPawnAction_Move::Resume()
 	return bResult;
 }
 
-EPawnActionAbortState::Type UPawnAction_Move::PerformAbort(EAIForceParam::Type ShouldForce)
+EPawnActionAbortState::Type UDEPRECATED_UPawnAction_Move::PerformAbort(EAIForceParam::Type ShouldForce)
 {
 	ClearTimers();
 	ClearPath();
@@ -228,7 +230,7 @@ EPawnActionAbortState::Type UPawnAction_Move::PerformAbort(EAIForceParam::Type S
 	return Super::PerformAbort(ShouldForce);
 }
 
-void UPawnAction_Move::HandleAIMessage(UBrainComponent*, const FAIMessage& Message)
+void UDEPRECATED_UPawnAction_Move::HandleAIMessage(UBrainComponent*, const FAIMessage& Message)
 {
 	if (Message.MessageName == UBrainComponent::AIMessage_MoveFinished && Message.HasFlag(FPathFollowingResultFlags::NewRequest))
 	{
@@ -242,7 +244,7 @@ void UPawnAction_Move::HandleAIMessage(UBrainComponent*, const FAIMessage& Messa
 	Finish(bFail ? EPawnActionResult::Failed : EPawnActionResult::Success);
 }
 
-void UPawnAction_Move::OnFinished(EPawnActionResult::Type WithResult)
+void UDEPRECATED_UPawnAction_Move::OnFinished(EPawnActionResult::Type WithResult)
 {
 	ClearTimers();
 	ClearPath();
@@ -250,7 +252,7 @@ void UPawnAction_Move::OnFinished(EPawnActionResult::Type WithResult)
 	Super::OnFinished(WithResult);
 }
 
-void UPawnAction_Move::ClearPath()
+void UDEPRECATED_UPawnAction_Move::ClearPath()
 {
 	ClearPendingRepath();
 	if (Path.IsValid())
@@ -260,20 +262,20 @@ void UPawnAction_Move::ClearPath()
 	}
 }
 
-void UPawnAction_Move::SetPath(FNavPathSharedRef InPath)
+void UDEPRECATED_UPawnAction_Move::SetPath(FNavPathSharedRef InPath)
 {
 	if (InPath != Path)
 	{
 		ClearPath();
 		Path = InPath;
-		PathObserverDelegateHandle = Path->AddObserver(FNavigationPath::FPathObserverDelegate::FDelegate::CreateUObject(this, &UPawnAction_Move::OnPathUpdated));
+		PathObserverDelegateHandle = Path->AddObserver(FNavigationPath::FPathObserverDelegate::FDelegate::CreateUObject(this, &UDEPRECATED_UPawnAction_Move::OnPathUpdated));
 
 		// skip auto updates, it will be handled manually to include controller's ShouldPostponePathUpdates()
 		Path->EnableRecalculationOnInvalidation(false);
 	}
 }
 
-void UPawnAction_Move::OnPathUpdated(FNavigationPath* UpdatedPath, ENavPathEvent::Type Event)
+void UDEPRECATED_UPawnAction_Move::OnPathUpdated(FNavigationPath* UpdatedPath, ENavPathEvent::Type Event)
 {
 	const AController* MyOwner = GetController();
 	if (MyOwner == NULL)
@@ -321,7 +323,7 @@ void UPawnAction_Move::OnPathUpdated(FNavigationPath* UpdatedPath, ENavPathEvent
 	}
 }
 
-void UPawnAction_Move::TryToRepath()
+void UDEPRECATED_UPawnAction_Move::TryToRepath()
 {
 	if (Path.IsValid())
 	{
@@ -336,12 +338,12 @@ void UPawnAction_Move::TryToRepath()
 		}
 		else if (GetWorld())
 		{
-			GetWorld()->GetTimerManager().SetTimer(TimerHandle_TryToRepath, this, &UPawnAction_Move::TryToRepath, 0.25f);
+			GetWorld()->GetTimerManager().SetTimer(TimerHandle_TryToRepath, this, &UDEPRECATED_UPawnAction_Move::TryToRepath, 0.25f);
 		}
 	}
 }
 
-void UPawnAction_Move::ClearPendingRepath()
+void UDEPRECATED_UPawnAction_Move::ClearPendingRepath()
 {
 	if (TimerHandle_TryToRepath.IsValid())
 	{
@@ -354,7 +356,7 @@ void UPawnAction_Move::ClearPendingRepath()
 	}
 }
 
-void UPawnAction_Move::ClearTimers()
+void UDEPRECATED_UPawnAction_Move::ClearTimers()
 {
 	UWorld* World = GetWorld();
 	if (World)
@@ -367,7 +369,7 @@ void UPawnAction_Move::ClearTimers()
 	}
 }
 
-bool UPawnAction_Move::CheckAlreadyAtGoal(AAIController& Controller, const FVector& TestLocation, float Radius)
+bool UDEPRECATED_UPawnAction_Move::CheckAlreadyAtGoal(AAIController& Controller, const FVector& TestLocation, float Radius)
 {
 	const bool bAlreadyAtGoal = Controller.GetPathFollowingComponent()->HasReached(TestLocation, EPathFollowingReachMode::OverlapAgentAndGoal, Radius);
 	if (bAlreadyAtGoal)
@@ -379,7 +381,7 @@ bool UPawnAction_Move::CheckAlreadyAtGoal(AAIController& Controller, const FVect
 	return bAlreadyAtGoal;
 }
 
-bool UPawnAction_Move::CheckAlreadyAtGoal(AAIController& Controller, const AActor& TestGoal, float Radius)
+bool UDEPRECATED_UPawnAction_Move::CheckAlreadyAtGoal(AAIController& Controller, const AActor& TestGoal, float Radius)
 {
 	const bool bAlreadyAtGoal = Controller.GetPathFollowingComponent()->HasReached(TestGoal, EPathFollowingReachMode::OverlapAgentAndGoal, Radius);
 	if (bAlreadyAtGoal)
@@ -391,8 +393,9 @@ bool UPawnAction_Move::CheckAlreadyAtGoal(AAIController& Controller, const AActo
 	return bAlreadyAtGoal;
 }
 
-bool UPawnAction_Move::IsPartialPathAllowed() const
+bool UDEPRECATED_UPawnAction_Move::IsPartialPathAllowed() const
 {
 	return bAllowPartialPath;
 }
 
+PRAGMA_ENABLE_DEPRECATION_WARNINGS
diff --git a/Engine/Source/Runtime/AIModule/Private/Actions/PawnAction_Repeat.cpp b/Engine/Source/Runtime/AIModule/Private/Actions/PawnAction_Repeat.cpp
--- a/Engine/Source/Runtime/AIModule/Private/Actions/PawnAction_Repeat.cpp
+++ b/Engine/Source/Runtime/AIModule/Private/Actions/PawnAction_Repeat.cpp
@@ -7,24 +7,26 @@
 
 #include UE_INLINE_GENERATED_CPP_BY_NAME(PawnAction_Repeat)
 
-UPawnAction_Repeat::UPawnAction_Repeat(const FObjectInitializer& ObjectInitializer)
+PRAGMA_DISABLE_DEPRECATION_WARNINGS
+
+UDEPRECATED_UPawnAction_Repeat::UDEPRECATED_UPawnAction_Repeat(const FObjectInitializer& ObjectInitializer)
 	: Super(ObjectInitializer)
 	, SubActionTriggeringPolicy(EPawnSubActionTriggeringPolicy::CopyBeforeTriggering)
 {
 	ChildFailureHandlingMode = EPawnActionFailHandling::IgnoreFailure;
 }
 
-UPawnAction_Repeat* UPawnAction_Repeat::CreateAction(UWorld& World, UPawnAction* ActionToRepeat, int32 NumberOfRepeats, EPawnSubActionTriggeringPolicy::Type InSubActionTriggeringPolicy)
+UDEPRECATED_UPawnAction_Repeat* UDEPRECATED_UPawnAction_Repeat::CreateAction(UWorld& World, UDEPRECATED_UPawnAction* ActionToRepeat, int32 NumberOfRepeats, EPawnSubActionTriggeringPolicy::Type InSubActionTriggeringPolicy)
 {
-	if (ActionToRepeat == NULL || !(NumberOfRepeats > 0 || NumberOfRepeats == UPawnAction_Repeat::LoopForever))
+	if (ActionToRepeat == NULL || !(NumberOfRepeats > 0 || NumberOfRepeats == UDEPRECATED_UPawnAction_Repeat::LoopForever))
 	{
 		return NULL;
 	}
 
-	UPawnAction_Repeat* Action = UPawnAction::CreateActionInstance<UPawnAction_Repeat>(World);
+	UDEPRECATED_UPawnAction_Repeat* Action = UDEPRECATED_UPawnAction::CreateActionInstance<UDEPRECATED_UPawnAction_Repeat>(World);
 	if (Action)
 	{
-		Action->ActionToRepeat = ActionToRepeat;
+		Action->ActionToRepeat_DEPRECATED = ActionToRepeat;
 		Action->RepeatsLeft = NumberOfRepeats;
 		Action->SubActionTriggeringPolicy = InSubActionTriggeringPolicy;
 
@@ -34,21 +36,21 @@ UPawnAction_Repeat* UPawnAction_Repeat::CreateAction(UWorld& World, UPawnAction*
 	return Action;
 }
 
-bool UPawnAction_Repeat::Start()
+bool UDEPRECATED_UPawnAction_Repeat::Start()
 {
 	bool bResult = Super::Start();
 	
 	if (bResult)
 	{
 		UE_VLOG(GetPawn(), LogPawnAction, Log, TEXT("Starting repeating action: %s. Requested repeats: %d")
-			, *GetNameSafe(ActionToRepeat), RepeatsLeft);
+			, *GetNameSafe(ActionToRepeat_DEPRECATED), RepeatsLeft);
 		bResult = PushSubAction();
 	}
 
 	return bResult;
 }
 
-bool UPawnAction_Repeat::Resume()
+bool UDEPRECATED_UPawnAction_Repeat::Resume()
 {
 	bool bResult = Super::Resume();
 
@@ -60,9 +62,9 @@ bool UPawnAction_Repeat::Resume()
 	return bResult;
 }
 
-void UPawnAction_Repeat::OnChildFinished(UPawnAction& Action, EPawnActionResult::Type WithResult)
+void UDEPRECATED_UPawnAction_Repeat::OnChildFinished(UDEPRECATED_UPawnAction& Action, EPawnActionResult::Type WithResult)
 {
-	if (RecentActionCopy == &Action)
+	if (RecentActionCopy_DEPRECATED == &Action)
 	{
 		if (WithResult == EPawnActionResult::Success || (WithResult == EPawnActionResult::Failed && ChildFailureHandlingMode == EPawnActionFailHandling::IgnoreFailure))
 		{
@@ -77,9 +79,9 @@ void UPawnAction_Repeat::OnChildFinished(UPawnAction& Action, EPawnActionResult:
 	Super::OnChildFinished(Action, WithResult);
 }
 
-bool UPawnAction_Repeat::PushSubAction()
+bool UDEPRECATED_UPawnAction_Repeat::PushSubAction()
 {
-	if (ActionToRepeat == NULL)
+	if (ActionToRepeat_DEPRECATED == NULL)
 	{
 		Finish(EPawnActionResult::Failed);
 		return false;
@@ -95,15 +97,16 @@ bool UPawnAction_Repeat::PushSubAction()
 		--RepeatsLeft;
 	}
 
-	UPawnAction* ActionCopy = SubActionTriggeringPolicy == EPawnSubActionTriggeringPolicy::CopyBeforeTriggering 
-		? Cast<UPawnAction>(StaticDuplicateObject(ActionToRepeat, this))
-		: ToRawPtr(ActionToRepeat);
+	UDEPRECATED_UPawnAction* ActionCopy = SubActionTriggeringPolicy == EPawnSubActionTriggeringPolicy::CopyBeforeTriggering 
+		? Cast<UDEPRECATED_UPawnAction>(StaticDuplicateObject(ActionToRepeat_DEPRECATED, this))
+		: ToRawPtr(ActionToRepeat_DEPRECATED);
 
 	UE_VLOG(GetPawn(), LogPawnAction, Log, TEXT("%s> pushing repeted action %s %s, repeats left: %d")
 		, *GetName(), SubActionTriggeringPolicy == EPawnSubActionTriggeringPolicy::CopyBeforeTriggering ? TEXT("copy") : TEXT("instance")
 		, *GetNameSafe(ActionCopy), RepeatsLeft);
 	check(ActionCopy);
-	RecentActionCopy = ActionCopy;
+	RecentActionCopy_DEPRECATED = ActionCopy;
 	return PushChildAction(*ActionCopy); 
 }
 
+PRAGMA_ENABLE_DEPRECATION_WARNINGS
diff --git a/Engine/Source/Runtime/AIModule/Private/Actions/PawnAction_Sequence.cpp b/Engine/Source/Runtime/AIModule/Private/Actions/PawnAction_Sequence.cpp
--- a/Engine/Source/Runtime/AIModule/Private/Actions/PawnAction_Sequence.cpp
+++ b/Engine/Source/Runtime/AIModule/Private/Actions/PawnAction_Sequence.cpp
@@ -7,13 +7,15 @@
 
 #include UE_INLINE_GENERATED_CPP_BY_NAME(PawnAction_Sequence)
 
-UPawnAction_Sequence::UPawnAction_Sequence(const FObjectInitializer& ObjectInitializer)
+PRAGMA_DISABLE_DEPRECATION_WARNINGS
+
+UDEPRECATED_UPawnAction_Sequence::UDEPRECATED_UPawnAction_Sequence(const FObjectInitializer& ObjectInitializer)
 	: Super(ObjectInitializer)
 	, SubActionTriggeringPolicy(EPawnSubActionTriggeringPolicy::CopyBeforeTriggering)
 {
 }
 
-UPawnAction_Sequence* UPawnAction_Sequence::CreateAction(UWorld& World, TArray<UPawnAction*>& ActionSequence, EPawnSubActionTriggeringPolicy::Type InSubActionTriggeringPolicy)
+UDEPRECATED_UPawnAction_Sequence* UDEPRECATED_UPawnAction_Sequence::CreateAction(UWorld& World, TArray<UDEPRECATED_UPawnAction*>& ActionSequence, EPawnSubActionTriggeringPolicy::Type InSubActionTriggeringPolicy)
 {
 	ActionSequence.Remove(NULL);
 	if (ActionSequence.Num() <= 0)
@@ -21,12 +23,12 @@ UPawnAction_Sequence* UPawnAction_Sequence::CreateAction(UWorld& World, TArray<U
 		return NULL;
 	}
 
-	UPawnAction_Sequence* Action = UPawnAction::CreateActionInstance<UPawnAction_Sequence>(World);
+	UDEPRECATED_UPawnAction_Sequence* Action = UDEPRECATED_UPawnAction::CreateActionInstance<UDEPRECATED_UPawnAction_Sequence>(World);
 	if (Action)
 	{
-		Action->ActionSequence = ActionSequence;
+		Action->ActionSequence_DEPRECATED = ActionSequence;
 
-		for (const UPawnAction* SubAction : ActionSequence)
+		for (const UDEPRECATED_UPawnAction* SubAction : ActionSequence)
 		{
 			if (SubAction && SubAction->ShouldPauseMovement())
 			{
@@ -41,14 +43,14 @@ UPawnAction_Sequence* UPawnAction_Sequence::CreateAction(UWorld& World, TArray<U
 	return Action;
 }
 
-bool UPawnAction_Sequence::Start()
+bool UDEPRECATED_UPawnAction_Sequence::Start()
 {
 	bool bResult = Super::Start();
 
 	if (bResult)
 	{
 		UE_VLOG(GetPawn(), LogPawnAction, Log, TEXT("Starting sequence. Items:"), *GetName());
-		for (auto Action : ActionSequence)
+		for (auto Action : ActionSequence_DEPRECATED)
 		{
 			UE_VLOG(GetPawn(), LogPawnAction, Log, TEXT("    %s"), *GetNameSafe(Action));
 		}
@@ -59,7 +61,7 @@ bool UPawnAction_Sequence::Start()
 	return bResult;
 }
 
-bool UPawnAction_Sequence::Resume()
+bool UDEPRECATED_UPawnAction_Sequence::Resume()
 {
 	bool bResult = Super::Resume();
 
@@ -71,9 +73,9 @@ bool UPawnAction_Sequence::Resume()
 	return bResult;
 }
 
-void UPawnAction_Sequence::OnChildFinished(UPawnAction& Action, EPawnActionResult::Type WithResult)
+void UDEPRECATED_UPawnAction_Sequence::OnChildFinished(UDEPRECATED_UPawnAction& Action, EPawnActionResult::Type WithResult)
 {
-	if (RecentActionCopy == &Action)
+	if (RecentActionCopy_DEPRECATED == &Action)
 	{
 		if (WithResult == EPawnActionResult::Success || (WithResult == EPawnActionResult::Failed && ChildFailureHandlingMode == EPawnActionFailHandling::IgnoreFailure))
 		{
@@ -91,23 +93,24 @@ void UPawnAction_Sequence::OnChildFinished(UPawnAction& Action, EPawnActionResul
 	Super::OnChildFinished(Action, WithResult);
 }
 
-bool UPawnAction_Sequence::PushNextActionCopy()
+bool UDEPRECATED_UPawnAction_Sequence::PushNextActionCopy()
 {
-	if (CurrentActionIndex >= uint32(ActionSequence.Num()))
+	if (CurrentActionIndex >= uint32(ActionSequence_DEPRECATED.Num()))
 	{
 		Finish(EPawnActionResult::Success);
 		return true;
 	}
 
-	UPawnAction* ActionCopy = SubActionTriggeringPolicy == EPawnSubActionTriggeringPolicy::CopyBeforeTriggering
-		? Cast<UPawnAction>(StaticDuplicateObject(ActionSequence[CurrentActionIndex], this))
-		: ToRawPtr(ActionSequence[CurrentActionIndex]);
+	UDEPRECATED_UPawnAction* ActionCopy = SubActionTriggeringPolicy == EPawnSubActionTriggeringPolicy::CopyBeforeTriggering
+		? Cast<UDEPRECATED_UPawnAction>(StaticDuplicateObject(ActionSequence_DEPRECATED[CurrentActionIndex], this))
+		: ToRawPtr(ActionSequence_DEPRECATED[CurrentActionIndex]);
 
 	UE_VLOG(GetPawn(), LogPawnAction, Log, TEXT("%s> pushing action %s")
 		, *GetName(), *GetNameSafe(ActionCopy));
 	++CurrentActionIndex;	
 	check(ActionCopy);
-	RecentActionCopy = ActionCopy;
+	RecentActionCopy_DEPRECATED = ActionCopy;
 	return PushChildAction(*ActionCopy);
 }
 
+PRAGMA_ENABLE_DEPRECATION_WARNINGS
diff --git a/Engine/Source/Runtime/AIModule/Private/Actions/PawnAction_Wait.cpp b/Engine/Source/Runtime/AIModule/Private/Actions/PawnAction_Wait.cpp
--- a/Engine/Source/Runtime/AIModule/Private/Actions/PawnAction_Wait.cpp
+++ b/Engine/Source/Runtime/AIModule/Private/Actions/PawnAction_Wait.cpp
@@ -6,7 +6,9 @@
 
 #include UE_INLINE_GENERATED_CPP_BY_NAME(PawnAction_Wait)
 
-UPawnAction_Wait::UPawnAction_Wait(const FObjectInitializer& ObjectInitializer)
+PRAGMA_DISABLE_DEPRECATION_WARNINGS
+
+UDEPRECATED_UPawnAction_Wait::UDEPRECATED_UPawnAction_Wait(const FObjectInitializer& ObjectInitializer)
 	: Super(ObjectInitializer)
 	, TimeToWait(0.f)
 	, FinishTimeStamp(0.f)
@@ -14,9 +16,9 @@ UPawnAction_Wait::UPawnAction_Wait(const FObjectInitializer& ObjectInitializer)
 
 }
 
-UPawnAction_Wait* UPawnAction_Wait::CreateAction(UWorld& World, float InTimeToWait)
+UDEPRECATED_UPawnAction_Wait* UDEPRECATED_UPawnAction_Wait::CreateAction(UWorld& World, float InTimeToWait)
 {
-	UPawnAction_Wait* Action = UPawnAction::CreateActionInstance<UPawnAction_Wait>(World);
+	UDEPRECATED_UPawnAction_Wait* Action = UDEPRECATED_UPawnAction::CreateActionInstance<UDEPRECATED_UPawnAction_Wait>(World);
 	
 	if (Action != NULL)
 	{
@@ -26,13 +28,13 @@ UPawnAction_Wait* UPawnAction_Wait::CreateAction(UWorld& World, float InTimeToWa
 	return Action;
 }
 
-bool UPawnAction_Wait::Start()
+bool UDEPRECATED_UPawnAction_Wait::Start()
 {
 	if (Super::Start())
 	{
 		if (TimeToWait >= 0)
 		{
-			GetWorld()->GetTimerManager().SetTimer(TimerHandle, this, &UPawnAction_Wait::TimerDone, TimeToWait);
+			GetWorld()->GetTimerManager().SetTimer(TimerHandle, this, &UDEPRECATED_UPawnAction_Wait::TimerDone, TimeToWait);
 		}
 		// else hang in there for ever!
 
@@ -42,26 +44,27 @@ bool UPawnAction_Wait::Start()
 	return false;
 }
 
-bool UPawnAction_Wait::Pause(const UPawnAction* PausedBy)
+bool UDEPRECATED_UPawnAction_Wait::Pause(const UDEPRECATED_UPawnAction* PausedBy)
 {
 	GetWorld()->GetTimerManager().PauseTimer(TimerHandle);
 	return true;
 }
 
-bool UPawnAction_Wait::Resume()
+bool UDEPRECATED_UPawnAction_Wait::Resume()
 {
 	GetWorld()->GetTimerManager().UnPauseTimer(TimerHandle);
 	return true;
 }
 
-EPawnActionAbortState::Type UPawnAction_Wait::PerformAbort(EAIForceParam::Type ShouldForce)
+EPawnActionAbortState::Type UDEPRECATED_UPawnAction_Wait::PerformAbort(EAIForceParam::Type ShouldForce)
 {
 	GetWorld()->GetTimerManager().ClearTimer(TimerHandle);
 	return EPawnActionAbortState::AbortDone;
 }
 
-void UPawnAction_Wait::TimerDone()
+void UDEPRECATED_UPawnAction_Wait::TimerDone()
 {
 	Finish(EPawnActionResult::Success);
 }
 
+PRAGMA_ENABLE_DEPRECATION_WARNINGS
diff --git a/Engine/Source/Runtime/AIModule/Private/Actions/PawnActionsComponent.cpp b/Engine/Source/Runtime/AIModule/Private/Actions/PawnActionsComponent.cpp
--- a/Engine/Source/Runtime/AIModule/Private/Actions/PawnActionsComponent.cpp
+++ b/Engine/Source/Runtime/AIModule/Private/Actions/PawnActionsComponent.cpp
@@ -11,6 +11,8 @@
 
 #include UE_INLINE_GENERATED_CPP_BY_NAME(PawnActionsComponent)
 
+PRAGMA_DISABLE_DEPRECATION_WARNINGS
+
 //----------------------------------------------------------------------//
 // helpers
 //----------------------------------------------------------------------//
@@ -31,7 +33,7 @@ namespace
 		return Enum->GetNameStringByValue(Value);
 	}
 
-	FString GetActionSignature(UPawnAction* Action)
+	FString GetActionSignature(UDEPRECATED_UPawnAction* Action)
 	{
 		if (Action == NULL)
 		{
@@ -60,7 +62,7 @@ namespace
 	};
 }
 
-FPawnActionEvent::FPawnActionEvent(UPawnAction& InAction, EPawnActionEventType::Type InEventType, uint32 InIndex)
+FPawnActionEvent::FPawnActionEvent(UDEPRECATED_UPawnAction& InAction, EPawnActionEventType::Type InEventType, uint32 InIndex)
 	: Action(&InAction), EventType(InEventType), Index(InIndex)
 {
 	Priority = InAction.GetPriority();
@@ -86,7 +88,7 @@ void FPawnActionStack::Resume()
 	}
 }
 
-void FPawnActionStack::PushAction(UPawnAction& NewTopAction)
+void FPawnActionStack::PushAction(UDEPRECATED_UPawnAction& NewTopAction)
 {
 	if (TopAction != NULL)
 	{
@@ -94,34 +96,34 @@ void FPawnActionStack::PushAction(UPawnAction& NewTopAction)
 		{
 			TopAction->Pause(&NewTopAction);
 		}
-		ensure(TopAction->ChildAction == NULL);
-		TopAction->ChildAction = &NewTopAction;
-		NewTopAction.ParentAction = TopAction;
+		ensure(TopAction->ChildAction_DEPRECATED == NULL);
+		TopAction->ChildAction_DEPRECATED = &NewTopAction;
+		NewTopAction.ParentAction_DEPRECATED = TopAction;
 	}
 
 	TopAction = &NewTopAction;
 	NewTopAction.OnPushed();
 }
 
-void FPawnActionStack::PopAction(UPawnAction& ActionToPop)
+void FPawnActionStack::PopAction(UDEPRECATED_UPawnAction& ActionToPop)
 {
 	// first check if it's there
-	UPawnAction* CutPoint = TopAction;
+	UDEPRECATED_UPawnAction* CutPoint = TopAction;
 	while (CutPoint != NULL && CutPoint != &ActionToPop)
 	{
-		CutPoint = CutPoint->ParentAction;
+		CutPoint = CutPoint->ParentAction_DEPRECATED;
 	}
 
 	if (CutPoint == &ActionToPop)
 	{
-		UPawnAction* ActionBeingRemoved = TopAction;
+		UDEPRECATED_UPawnAction* ActionBeingRemoved = TopAction;
 		// note StopAction can be null
-		UPawnAction* StopAction = ActionToPop.ParentAction;
+		UDEPRECATED_UPawnAction* StopAction = ActionToPop.ParentAction_DEPRECATED;
 
 		while (ActionBeingRemoved != StopAction && ActionBeingRemoved != nullptr)
 		{
 			checkSlow(ActionBeingRemoved);
-			UPawnAction* NextAction = ActionBeingRemoved->ParentAction;
+			UDEPRECATED_UPawnAction* NextAction = ActionBeingRemoved->ParentAction_DEPRECATED;
 
 			if (ActionBeingRemoved->IsBeingAborted() == false && ActionBeingRemoved->IsFinished() == false)
 			{
@@ -129,9 +131,9 @@ void FPawnActionStack::PopAction(UPawnAction& ActionToPop)
 				ActionBeingRemoved->Abort(EAIForceParam::Force);
 			}
 			ActionBeingRemoved->OnPopped();
-			if (ActionBeingRemoved->ParentAction)
+			if (ActionBeingRemoved->ParentAction_DEPRECATED)
 			{
-				ActionBeingRemoved->ParentAction->OnChildFinished(*ActionBeingRemoved, ActionBeingRemoved->FinishResult);
+				ActionBeingRemoved->ParentAction_DEPRECATED->OnChildFinished(*ActionBeingRemoved, ActionBeingRemoved->FinishResult);
 			}
 			ActionBeingRemoved = NextAction;
 		}
@@ -143,7 +145,7 @@ void FPawnActionStack::PopAction(UPawnAction& ActionToPop)
 int32 FPawnActionStack::GetStackSize() const
 {
 	int32 Size = 0;
-	const UPawnAction* TempAction = TopAction;
+	const UDEPRECATED_UPawnAction* TempAction = TopAction;
 	while (TempAction != nullptr)
 	{
 		TempAction = TempAction->GetParentAction();
@@ -153,10 +155,10 @@ int32 FPawnActionStack::GetStackSize() const
 }
 
 //----------------------------------------------------------------------//
-// UPawnActionsComponent
+// UDEPRECATED_UPawnActionsComponent
 //----------------------------------------------------------------------//
 
-UPawnActionsComponent::UPawnActionsComponent(const FObjectInitializer& ObjectInitializer)
+UDEPRECATED_UPawnActionsComponent::UDEPRECATED_UPawnActionsComponent(const FObjectInitializer& ObjectInitializer)
 	: Super(ObjectInitializer)
 {
 	PrimaryComponentTick.bCanEverTick = true;
@@ -170,18 +172,18 @@ UPawnActionsComponent::UPawnActionsComponent(const FObjectInitializer& ObjectIni
 	ActionStacks.AddZeroed(EAIRequestPriority::MAX);
 }
 
-void UPawnActionsComponent::OnUnregister()
+void UDEPRECATED_UPawnActionsComponent::OnUnregister()
 {
 	if ((ControlledPawn != nullptr) && !ControlledPawn->IsPendingKillPending())
 	{
 		// call for every regular priority 
 		for (int32 PriorityIndex = 0; PriorityIndex < EAIRequestPriority::MAX; ++PriorityIndex)
 		{
-			UPawnAction* Action = ActionStacks[PriorityIndex].GetTop();
+			UDEPRECATED_UPawnAction* Action = ActionStacks[PriorityIndex].GetTop();
 			while (Action)
 			{
 				Action->Abort(EAIForceParam::Force);
-				Action = Action->ParentAction;
+				Action = Action->ParentAction_DEPRECATED;
 			}
 		}
 	}
@@ -189,7 +191,7 @@ void UPawnActionsComponent::OnUnregister()
 	Super::OnUnregister();
 }
 
-void UPawnActionsComponent::TickComponent(float DeltaTime, enum ELevelTick TickType, FActorComponentTickFunction *ThisTickFunction)
+void UDEPRECATED_UPawnActionsComponent::TickComponent(float DeltaTime, enum ELevelTick TickType, FActorComponentTickFunction *ThisTickFunction)
 {
 	Super::TickComponent(DeltaTime, TickType, ThisTickFunction);
 
@@ -244,26 +246,26 @@ void UPawnActionsComponent::TickComponent(float DeltaTime, enum ELevelTick TickT
 		UpdateCurrentAction();
 	}
 
-	if (CurrentAction)
+	if (CurrentAction_DEPRECATED)
 	{
-		CurrentAction->TickAction(DeltaTime);
+		CurrentAction_DEPRECATED->TickAction(DeltaTime);
 	}
 
-	// it's possible we got new events with CurrentAction's tick
-	if (ActionEvents.Num() == 0 && (CurrentAction == NULL || CurrentAction->WantsTick() == false))
+	// it's possible we got new events with CurrentAction_DEPRECATED's tick
+	if (ActionEvents.Num() == 0 && (CurrentAction_DEPRECATED == NULL || CurrentAction_DEPRECATED->WantsTick() == false))
 	{
 		SetComponentTickEnabled(false);
 	}
 }
 
-bool UPawnActionsComponent::HasActiveActionOfType(EAIRequestPriority::Type Priority, TSubclassOf<UPawnAction> PawnActionClass) const
+bool UDEPRECATED_UPawnActionsComponent::HasActiveActionOfType(EAIRequestPriority::Type Priority, TSubclassOf<UDEPRECATED_UPawnAction> PawnActionClass) const
 {
-	TArray<UPawnAction*> ActionsToTest;
+	TArray<UDEPRECATED_UPawnAction*> ActionsToTest;
 	ActionsToTest.Add(GetActiveAction(Priority));
 
 	while (ActionsToTest.Num() > 0)
 	{
-		UPawnAction* ActiveActionIter = ActionsToTest[0];
+		UDEPRECATED_UPawnAction* ActiveActionIter = ActionsToTest[0];
 
 		if (ActiveActionIter)
 		{
@@ -273,13 +275,13 @@ bool UPawnActionsComponent::HasActiveActionOfType(EAIRequestPriority::Type Prior
 			}	
 			else
 			{
-				UPawnAction_Sequence* PawnActionSequence = Cast<UPawnAction_Sequence>(ActiveActionIter);
+				UDEPRECATED_UPawnAction_Sequence* PawnActionSequence = Cast<UDEPRECATED_UPawnAction_Sequence>(ActiveActionIter);
 
 				if (PawnActionSequence)
 				{
-					for (int32 PawnActionSequenceCount = 0; PawnActionSequenceCount < PawnActionSequence->ActionSequence.Num(); ++PawnActionSequenceCount)
+					for (int32 PawnActionSequenceCount = 0; PawnActionSequenceCount < PawnActionSequence->ActionSequence_DEPRECATED.Num(); ++PawnActionSequenceCount)
 					{
-						ActionsToTest.Add(PawnActionSequence->ActionSequence[PawnActionSequenceCount]);
+						ActionsToTest.Add(PawnActionSequence->ActionSequence_DEPRECATED[PawnActionSequenceCount]);
 					}
 				}
 			}
@@ -292,13 +294,13 @@ bool UPawnActionsComponent::HasActiveActionOfType(EAIRequestPriority::Type Prior
 	return false;
 }
 
-void UPawnActionsComponent::UpdateCurrentAction()
+void UDEPRECATED_UPawnActionsComponent::UpdateCurrentAction()
 {
-	UE_VLOG(ControlledPawn, LogPawnAction, Log, TEXT("Picking new current actions. Old CurrentAction %s")
-		, *GetActionSignature(CurrentAction));
+	UE_VLOG(ControlledPawn, LogPawnAction, Log, TEXT("Picking new current actions. Old CurrentAction_DEPRECATED %s")
+		, *GetActionSignature(CurrentAction_DEPRECATED));
 
 	// find the highest priority action available
-	UPawnAction* NewCurrentAction = NULL;
+	UDEPRECATED_UPawnAction* NewCurrentAction = NULL;
 	int32 Priority = EAIRequestPriority::MAX - 1;
 	do 
 	{
@@ -307,45 +309,45 @@ void UPawnActionsComponent::UpdateCurrentAction()
 	} while (NewCurrentAction == NULL && --Priority >= 0);
 
 	// if it's a new Action then enable it
-	if (CurrentAction != NewCurrentAction)
+	if (CurrentAction_DEPRECATED != NewCurrentAction)
 	{
 		UE_VLOG(ControlledPawn, LogPawnAction, Log, TEXT("New action: %s")
 			, *GetActionSignature(NewCurrentAction));
 
-		if (CurrentAction != NULL && CurrentAction->IsActive())
+		if (CurrentAction_DEPRECATED != NULL && CurrentAction_DEPRECATED->IsActive())
 		{
-			CurrentAction->Pause(NewCurrentAction);
+			CurrentAction_DEPRECATED->Pause(NewCurrentAction);
 		}
-		CurrentAction = NewCurrentAction;
+		CurrentAction_DEPRECATED = NewCurrentAction;
 		bool bNewActionStartedSuccessfully = true;
-		if (CurrentAction != NULL)
+		if (CurrentAction_DEPRECATED != NULL)
 		{
-			bNewActionStartedSuccessfully = CurrentAction->Activate();
+			bNewActionStartedSuccessfully = CurrentAction_DEPRECATED->Activate();
 		}
 
 		if (bNewActionStartedSuccessfully == false)
 		{
-			UE_VLOG(ControlledPawn, LogPawnAction, Warning, TEXT("CurrentAction %s failed to activate. Removing and re-running action selection")
+			UE_VLOG(ControlledPawn, LogPawnAction, Warning, TEXT("CurrentAction_DEPRECATED %s failed to activate. Removing and re-running action selection")
 				, *GetActionSignature(NewCurrentAction));
 
-			CurrentAction = NULL;			
+			CurrentAction_DEPRECATED = NULL;			
 		}
 		// @HACK temporary solution to have actions and old BT tasks work together
-		else if (CurrentAction == NULL || CurrentAction->GetPriority() != EAIRequestPriority::Logic)
+		else if (CurrentAction_DEPRECATED == NULL || CurrentAction_DEPRECATED->GetPriority() != EAIRequestPriority::Logic)
 		{
 			UpdateAILogicLock();
 		}
 	}
 	else
 	{
-		if (CurrentAction == NULL)
+		if (CurrentAction_DEPRECATED == NULL)
 		{
 			UpdateAILogicLock();
 		}
-		else if (CurrentAction->IsFinished())
+		else if (CurrentAction_DEPRECATED->IsFinished())
 		{
 			UE_VLOG(ControlledPawn, LogPawnAction, Warning, TEXT("Re-running same action"));
-			CurrentAction->Activate();
+			CurrentAction_DEPRECATED->Activate();
 		}
 		else
 		{ 
@@ -354,14 +356,14 @@ void UPawnActionsComponent::UpdateCurrentAction()
 	}
 }
 
-void UPawnActionsComponent::UpdateAILogicLock()
+void UDEPRECATED_UPawnActionsComponent::UpdateAILogicLock()
 {
 	if (ControlledPawn && ControlledPawn->GetController())
 	{
 		UBrainComponent* BrainComp = ControlledPawn->GetController()->FindComponentByClass<UBrainComponent>();
 		if (BrainComp)
 		{
-			if (CurrentAction != NULL && CurrentAction->GetPriority() > EAIRequestPriority::Logic)
+			if (CurrentAction_DEPRECATED != NULL && CurrentAction_DEPRECATED->GetPriority() > EAIRequestPriority::Logic)
 			{
 				UE_VLOG(ControlledPawn, LogPawnAction, Log, TEXT("Locking AI logic"));
 				BrainComp->LockResource(EAIRequestPriority::HardScript);
@@ -388,7 +390,7 @@ void UPawnActionsComponent::UpdateAILogicLock()
 	}
 }
 
-EPawnActionAbortState::Type UPawnActionsComponent::K2_AbortAction(UPawnAction* ActionToAbort)
+EPawnActionAbortState::Type UDEPRECATED_UPawnActionsComponent::K2_AbortAction(UDEPRECATED_UPawnAction* ActionToAbort)
 {
 	if (ActionToAbort != NULL)
 	{
@@ -397,7 +399,7 @@ EPawnActionAbortState::Type UPawnActionsComponent::K2_AbortAction(UPawnAction* A
 	return EPawnActionAbortState::NeverStarted;
 }
 
-EPawnActionAbortState::Type UPawnActionsComponent::AbortAction(UPawnAction& ActionToAbort)
+EPawnActionAbortState::Type UDEPRECATED_UPawnActionsComponent::AbortAction(UDEPRECATED_UPawnAction& ActionToAbort)
 {
 	const EPawnActionAbortState::Type AbortState = ActionToAbort.Abort(EAIForceParam::DoNotForce);
 	if (AbortState == EPawnActionAbortState::NeverStarted)
@@ -411,7 +413,7 @@ EPawnActionAbortState::Type UPawnActionsComponent::AbortAction(UPawnAction& Acti
 	return AbortState;
 }
 
-void UPawnActionsComponent::RemoveEventsForAction(UPawnAction& PawnAction)
+void UDEPRECATED_UPawnActionsComponent::RemoveEventsForAction(UDEPRECATED_UPawnAction& PawnAction)
 {
 	for (int32 ActionIndex = ActionEvents.Num() - 1; ActionIndex >= 0; --ActionIndex)
 	{
@@ -422,7 +424,7 @@ void UPawnActionsComponent::RemoveEventsForAction(UPawnAction& PawnAction)
 	}
 }
 
-EPawnActionAbortState::Type UPawnActionsComponent::K2_ForceAbortAction(UPawnAction* ActionToAbort)
+EPawnActionAbortState::Type UDEPRECATED_UPawnActionsComponent::K2_ForceAbortAction(UDEPRECATED_UPawnAction* ActionToAbort)
 {
 	if (ActionToAbort)
 	{
@@ -431,12 +433,12 @@ EPawnActionAbortState::Type UPawnActionsComponent::K2_ForceAbortAction(UPawnActi
 	return EPawnActionAbortState::NeverStarted;
 }
 
-EPawnActionAbortState::Type UPawnActionsComponent::ForceAbortAction(UPawnAction& ActionToAbort)
+EPawnActionAbortState::Type UDEPRECATED_UPawnActionsComponent::ForceAbortAction(UDEPRECATED_UPawnAction& ActionToAbort)
 {
 	return ActionToAbort.Abort(EAIForceParam::Force);
 }
 
-uint32 UPawnActionsComponent::AbortActionsInstigatedBy(UObject* const Instigator, EAIRequestPriority::Type Priority)
+uint32 UDEPRECATED_UPawnActionsComponent::AbortActionsInstigatedBy(UObject* const Instigator, EAIRequestPriority::Type Priority)
 {
 	uint32 AbortedActionsCount = 0;
 
@@ -450,7 +452,7 @@ uint32 UPawnActionsComponent::AbortActionsInstigatedBy(UObject* const Instigator
 	}
 	else
 	{
-		UPawnAction* Action = ActionStacks[Priority].GetTop();
+		UDEPRECATED_UPawnAction* Action = ActionStacks[Priority].GetTop();
 		while (Action)
 		{
 			if (Action->GetInstigator() == Instigator)
@@ -458,7 +460,7 @@ uint32 UPawnActionsComponent::AbortActionsInstigatedBy(UObject* const Instigator
 				OnEvent(*Action, EPawnActionEventType::InstantAbort);
 				++AbortedActionsCount;
 			}
-			Action = Action->ParentAction;
+			Action = Action->ParentAction_DEPRECATED;
 		}
 
 		for (int32 ActionIndex = ActionEvents.Num() - 1; ActionIndex >= 0; --ActionIndex)
@@ -477,7 +479,7 @@ uint32 UPawnActionsComponent::AbortActionsInstigatedBy(UObject* const Instigator
 	return AbortedActionsCount;
 }
 
-bool UPawnActionsComponent::K2_PushAction(UPawnAction* NewAction, EAIRequestPriority::Type Priority, UObject* Instigator)
+bool UDEPRECATED_UPawnActionsComponent::K2_PushAction(UDEPRECATED_UPawnAction* NewAction, EAIRequestPriority::Type Priority, UObject* Instigator)
 {
 	if (NewAction)
 	{
@@ -486,7 +488,7 @@ bool UPawnActionsComponent::K2_PushAction(UPawnAction* NewAction, EAIRequestPrio
 	return false;
 }
 
-bool UPawnActionsComponent::PushAction(UPawnAction& NewAction, EAIRequestPriority::Type Priority, UObject* Instigator)
+bool UDEPRECATED_UPawnActionsComponent::PushAction(UDEPRECATED_UPawnAction& NewAction, EAIRequestPriority::Type Priority, UObject* Instigator)
 {
 	if (NewAction.HasBeenStarted() == false || NewAction.IsFinished() == true)
 	{
@@ -499,7 +501,7 @@ bool UPawnActionsComponent::PushAction(UPawnAction& NewAction, EAIRequestPriorit
 	return false;
 }
 
-bool UPawnActionsComponent::OnEvent(UPawnAction& Action, EPawnActionEventType::Type Event)
+bool UDEPRECATED_UPawnActionsComponent::OnEvent(UDEPRECATED_UPawnAction& Action, EPawnActionEventType::Type Event)
 {
 	bool bResult = false;
 	const FPawnActionEvent ActionEvent(Action, Event, ActionEventIndex++);
@@ -531,7 +533,7 @@ bool UPawnActionsComponent::OnEvent(UPawnAction& Action, EPawnActionEventType::T
 	return bResult;
 }
 
-void UPawnActionsComponent::SetControlledPawn(APawn* NewPawn)
+void UDEPRECATED_UPawnActionsComponent::SetControlledPawn(APawn* NewPawn)
 {
 	if (ControlledPawn != NULL && ControlledPawn != NewPawn)
 	{
@@ -543,7 +545,7 @@ void UPawnActionsComponent::SetControlledPawn(APawn* NewPawn)
 	}
 }
 
-APawn* UPawnActionsComponent::CacheControlledPawn()
+APawn* UDEPRECATED_UPawnActionsComponent::CacheControlledPawn()
 {
 	if (ControlledPawn == NULL)
 	{
@@ -569,7 +571,7 @@ APawn* UPawnActionsComponent::CacheControlledPawn()
 //----------------------------------------------------------------------//
 // blueprint interface
 //----------------------------------------------------------------------//
-bool UPawnActionsComponent::K2_PerformAction(APawn* Pawn, UPawnAction* Action, TEnumAsByte<EAIRequestPriority::Type> Priority)
+bool UDEPRECATED_UPawnActionsComponent::K2_PerformAction(APawn* Pawn, UDEPRECATED_UPawnAction* Action, TEnumAsByte<EAIRequestPriority::Type> Priority)
 {
 	if (Pawn && Action)
 	{
@@ -578,7 +580,7 @@ bool UPawnActionsComponent::K2_PerformAction(APawn* Pawn, UPawnAction* Action, T
 	return false;
 }
 
-bool UPawnActionsComponent::PerformAction(APawn& Pawn, UPawnAction& Action, TEnumAsByte<EAIRequestPriority::Type> Priority)
+bool UDEPRECATED_UPawnActionsComponent::PerformAction(APawn& Pawn, UDEPRECATED_UPawnAction& Action, TEnumAsByte<EAIRequestPriority::Type> Priority)
 {
 	bool bSuccess = false;
 
@@ -586,7 +588,7 @@ bool UPawnActionsComponent::PerformAction(APawn& Pawn, UPawnAction& Action, TEnu
 
 	if (Pawn.GetController())
 	{
-		UPawnActionsComponent* ActionComp = Pawn.GetController()->FindComponentByClass<UPawnActionsComponent>();
+		UDEPRECATED_UPawnActionsComponent* ActionComp = Pawn.GetController()->FindComponentByClass<UDEPRECATED_UPawnActionsComponent>();
 		if (ActionComp)
 		{
 			ActionComp->PushAction(Action, Priority);
@@ -601,7 +603,7 @@ bool UPawnActionsComponent::PerformAction(APawn& Pawn, UPawnAction& Action, TEnu
 // debug
 //----------------------------------------------------------------------//
 #if ENABLE_VISUAL_LOG
-void UPawnActionsComponent::DescribeSelfToVisLog(FVisualLogEntry* Snapshot) const
+void UDEPRECATED_UPawnActionsComponent::DescribeSelfToVisLog(FVisualLogEntry* Snapshot) const
 {
 	static const FString Category = TEXT("PawnActions");
 
@@ -612,7 +614,7 @@ void UPawnActionsComponent::DescribeSelfToVisLog(FVisualLogEntry* Snapshot) cons
 
 	for (int32 PriorityIndex = 0; PriorityIndex < ActionStacks.Num(); ++PriorityIndex)
 	{
-		const UPawnAction* Action = ActionStacks[PriorityIndex].GetTop();
+		const UDEPRECATED_UPawnAction* Action = ActionStacks[PriorityIndex].GetTop();
 		if (Action == NULL)
 		{
 			continue;
@@ -639,8 +641,9 @@ void UPawnActionsComponent::DescribeSelfToVisLog(FVisualLogEntry* Snapshot) cons
 }
 #endif // ENABLE_VISUAL_LOG
 
-FString UPawnActionsComponent::DescribeEventType(EPawnActionEventType::Type EventType)
+FString UDEPRECATED_UPawnActionsComponent::DescribeEventType(EPawnActionEventType::Type EventType)
 {
 	return GetEventName(EventType);
 }
 
+PRAGMA_ENABLE_DEPRECATION_WARNINGS
diff --git a/Engine/Source/Runtime/AIModule/Private/BehaviorTree/Tasks/BTTask_PawnActionBase.cpp b/Engine/Source/Runtime/AIModule/Private/BehaviorTree/Tasks/BTTask_PawnActionBase.cpp
--- a/Engine/Source/Runtime/AIModule/Private/BehaviorTree/Tasks/BTTask_PawnActionBase.cpp
+++ b/Engine/Source/Runtime/AIModule/Private/BehaviorTree/Tasks/BTTask_PawnActionBase.cpp
@@ -7,6 +7,8 @@
 
 #include UE_INLINE_GENERATED_CPP_BY_NAME(BTTask_PawnActionBase)
 
+PRAGMA_DISABLE_DEPRECATION_WARNINGS
+
 UBTTask_PawnActionBase::UBTTask_PawnActionBase(const FObjectInitializer& ObjectInitializer) : Super(ObjectInitializer)
 {
 	NodeName = "Action Base";
@@ -20,7 +22,7 @@ EBTNodeResult::Type UBTTask_PawnActionBase::AbortTask(UBehaviorTreeComponent& Ow
 		? EBTNodeResult::InProgress : EBTNodeResult::Aborted;
 }
 
-EBTNodeResult::Type UBTTask_PawnActionBase::PushAction(UBehaviorTreeComponent& OwnerComp, UPawnAction& Action)
+EBTNodeResult::Type UBTTask_PawnActionBase::PushAction(UBehaviorTreeComponent& OwnerComp, UDEPRECATED_UPawnAction& Action)
 {
 	EBTNodeResult::Type NodeResult = EBTNodeResult::Failed;
 
@@ -45,7 +47,7 @@ EBTNodeResult::Type UBTTask_PawnActionBase::PushAction(UBehaviorTreeComponent& O
 	return NodeResult;
 }
 
-void UBTTask_PawnActionBase::OnActionEvent(UPawnAction& Action, EPawnActionEventType::Type Event)
+void UBTTask_PawnActionBase::OnActionEvent(UDEPRECATED_UPawnAction& Action, EPawnActionEventType::Type Event)
 {
 	EPawnActionTaskResult Result = ActionEventHandler(this, Action, Event);
 	if (Result == EPawnActionTaskResult::ActionLost)
@@ -54,7 +56,7 @@ void UBTTask_PawnActionBase::OnActionEvent(UPawnAction& Action, EPawnActionEvent
 	}
 }
 
-EPawnActionTaskResult UBTTask_PawnActionBase::ActionEventHandler(UBTTaskNode* TaskNode, UPawnAction& Action, EPawnActionEventType::Type Event)
+EPawnActionTaskResult UBTTask_PawnActionBase::ActionEventHandler(UBTTaskNode* TaskNode, UDEPRECATED_UPawnAction& Action, EPawnActionEventType::Type Event)
 {
 	EPawnActionTaskResult Result = EPawnActionTaskResult::Unknown;
 
@@ -64,7 +66,7 @@ EPawnActionTaskResult UBTTask_PawnActionBase::ActionEventHandler(UBTTaskNode* Ta
 	{
 		UE_VLOG(Action.GetController(), LogBehaviorTree, Warning, TEXT("%s: Unhandled PawnAction event: %s, can't find owning component!"),
 			*UBehaviorTreeTypes::DescribeNodeHelper(TaskNode),
-			*UPawnActionsComponent::DescribeEventType(Event));
+			*UDEPRECATED_UPawnActionsComponent::DescribeEventType(Event));
 
 		return Result;
 	}
@@ -98,14 +100,14 @@ EPawnActionTaskResult UBTTask_PawnActionBase::ActionEventHandler(UBTTaskNode* Ta
 	{
 		UE_VLOG(Action.GetController(), LogBehaviorTree, Verbose, TEXT("%s: Unhandled PawnAction event: %s TaskStatus: %s"),
 			*UBehaviorTreeTypes::DescribeNodeHelper(TaskNode),
-			*UPawnActionsComponent::DescribeEventType(Event),
+			*UDEPRECATED_UPawnActionsComponent::DescribeEventType(Event),
 			*UBehaviorTreeTypes::DescribeTaskStatus(TaskStatus));
 	}
 
 	return Result;
 }
 
-void UBTTask_PawnActionBase::OnActionLost(UPawnAction& Action)
+void UBTTask_PawnActionBase::OnActionLost(UDEPRECATED_UPawnAction& Action)
 {
 	AAIController* AIOwner = Cast<AAIController>(Action.GetController());
 	UBehaviorTreeComponent* OwnerComp = AIOwner ? Cast<UBehaviorTreeComponent>(AIOwner->BrainComponent) : nullptr;
@@ -116,3 +118,4 @@ void UBTTask_PawnActionBase::OnActionLost(UPawnAction& Action)
 	}
 }
 
+PRAGMA_ENABLE_DEPRECATION_WARNINGS
diff --git a/Engine/Source/Runtime/AIModule/Private/BehaviorTree/Tasks/BTTask_PushPawnAction.cpp b/Engine/Source/Runtime/AIModule/Private/BehaviorTree/Tasks/BTTask_PushPawnAction.cpp
--- a/Engine/Source/Runtime/AIModule/Private/BehaviorTree/Tasks/BTTask_PushPawnAction.cpp
+++ b/Engine/Source/Runtime/AIModule/Private/BehaviorTree/Tasks/BTTask_PushPawnAction.cpp
@@ -5,6 +5,8 @@
 
 #include UE_INLINE_GENERATED_CPP_BY_NAME(BTTask_PushPawnAction)
 
+PRAGMA_DISABLE_DEPRECATION_WARNINGS
+
 UBTTask_PushPawnAction::UBTTask_PushPawnAction(const FObjectInitializer& ObjectInitializer) : Super(ObjectInitializer)
 {
 	NodeName = "Push PawnAction";
@@ -12,7 +14,7 @@ UBTTask_PushPawnAction::UBTTask_PushPawnAction(const FObjectInitializer& ObjectI
 
 EBTNodeResult::Type UBTTask_PushPawnAction::ExecuteTask(UBehaviorTreeComponent& OwnerComp, uint8* NodeMemory)
 {
-	UPawnAction* ActionCopy = Action ? DuplicateObject<UPawnAction>(Action, &OwnerComp) : nullptr;
+	UDEPRECATED_UPawnAction* ActionCopy = Action_DEPRECATED ? DuplicateObject<UDEPRECATED_UPawnAction>(Action_DEPRECATED, &OwnerComp) : nullptr;
 	if (ActionCopy == nullptr)
 	{
 		return EBTNodeResult::Failed;
@@ -23,7 +25,7 @@ EBTNodeResult::Type UBTTask_PushPawnAction::ExecuteTask(UBehaviorTreeComponent&
 
 FString UBTTask_PushPawnAction::GetStaticDescription() const
 {
-	//return FString::Printf(TEXT("Push Action: %s"), Action ? *Action->GetDisplayName() : TEXT("None"));
-	return FString::Printf(TEXT("Push Action: %s"), *GetNameSafe(Action));
+	return FString::Printf(TEXT("Push Action: %s"), *GetNameSafe(Action_DEPRECATED));
 }
 
+PRAGMA_ENABLE_DEPRECATION_WARNINGS
