[BehaviorTree] Blackboard based decorator can allow none values which will be used, instead of the first key, when the referred key no longer exists.
To use that behavior by default you need to set
[/Script/AIModule.AISystem]
bBlackboardKeyDecoratorAllowsNoneAsValue=true
in the configuration file.

diff --git a/Engine/Source/Editor/BehaviorTreeEditor/Private/DetailCustomizations/BlackboardSelectorDetails.cpp b/Engine/Source/Editor/BehaviorTreeEditor/Private/DetailCustomizations/BlackboardSelectorDetails.cpp
--- a/Engine/Source/Editor/BehaviorTreeEditor/Private/DetailCustomizations/BlackboardSelectorDetails.cpp
+++ b/Engine/Source/Editor/BehaviorTreeEditor/Private/DetailCustomizations/BlackboardSelectorDetails.cpp
@@ -238,8 +238,9 @@ void FBlackboardSelectorDetails::InitKeyFromProperty()
 			}
 			else
 			{
-				MyKeyClassProperty->SetValue((UObject*)NULL);
+				// Set ID first so callbacks can properly test against InvalidKey
 				MyKeyIDProperty->SetValue((int32)FBlackboard::InvalidKey);
+				MyKeyClassProperty->SetValue((UObject*)nullptr);				
 				MyKeyNameProperty->SetValue(TEXT("None"));
 			}
 		}
diff --git a/Engine/Source/Runtime/AIModule/Classes/AISystem.h b/Engine/Source/Runtime/AIModule/Classes/AISystem.h
--- a/Engine/Source/Runtime/AIModule/Classes/AISystem.h
+++ b/Engine/Source/Runtime/AIModule/Classes/AISystem.h
@@ -90,6 +90,10 @@ public:
 
 	UPROPERTY(globalconfig, EditDefaultsOnly, Category = "Behavior Tree")
 	bool bClearBBEntryOnBTEQSFail = true;
+	
+	/** If enabled, blackboard based decorators will set key to 'Invalid' on creation or when selected key no longer exists (instead of using the first key of the blackboard). */
+	UPROPERTY(globalconfig, EditDefaultsOnly, Category = "Behavior Tree")
+	bool bBlackboardKeyDecoratorAllowsNoneAsValue = false;
 
 	/** Which collision channel to use for sight checks by default */
 	UPROPERTY(globalconfig, EditDefaultsOnly, Category = "PerceptionSystem")
diff --git a/Engine/Source/Runtime/AIModule/Private/BehaviorTree/Decorators/BTDecorator_BlackboardBase.cpp b/Engine/Source/Runtime/AIModule/Private/BehaviorTree/Decorators/BTDecorator_BlackboardBase.cpp
--- a/Engine/Source/Runtime/AIModule/Private/BehaviorTree/Decorators/BTDecorator_BlackboardBase.cpp
+++ b/Engine/Source/Runtime/AIModule/Private/BehaviorTree/Decorators/BTDecorator_BlackboardBase.cpp
@@ -2,6 +2,7 @@
 
 #include "BehaviorTree/Decorators/BTDecorator_BlackboardBase.h"
 #include "BehaviorTree/BlackboardComponent.h"
+#include "AISystem.h"
 
 #include UE_INLINE_GENERATED_CPP_BY_NAME(BTDecorator_BlackboardBase)
 
@@ -12,6 +13,8 @@ UBTDecorator_BlackboardBase::UBTDecorator_BlackboardBase(const FObjectInitialize
 	INIT_DECORATOR_NODE_NOTIFY_FLAGS();
 
 	// empty KeySelector = allow everything
+
+	BlackboardKey.AllowNoneAsValue(GET_AI_CONFIG_VAR(bBlackboardKeyDecoratorAllowsNoneAsValue));
 }
 
 void UBTDecorator_BlackboardBase::InitializeFromAsset(UBehaviorTree& Asset)
@@ -73,7 +76,3 @@ FName UBTDecorator_BlackboardBase::GetNodeIconName() const
 
 #endif	// WITH_EDITOR
 
-//----------------------------------------------------------------------//
-// DEPRECATED
-//----------------------------------------------------------------------//
-
