Fixed deprecation of PawnActions

Horde issue #250664

diff --git a/Engine/Source/Runtime/AIModule/Classes/AIController.h b/Engine/Source/Runtime/AIModule/Classes/AIController.h
--- a/Engine/Source/Runtime/AIModule/Classes/AIController.h
+++ b/Engine/Source/Runtime/AIModule/Classes/AIController.h
@@ -31,8 +31,8 @@ class UCanvas;
 class UGameplayTaskResource;
 class UGameplayTasksComponent;
 class UPathFollowingComponent;
-class UDEPRECATED_UPawnAction;
-class UDEPRECATED_UPawnActionsComponent;
+class UDEPRECATED_PawnAction;
+class UDEPRECATED_PawnActionsComponent;
 
 namespace EPathFollowingRequestResult {	enum Type : int; }
 namespace EPathFollowingResult { enum Type : int; }
@@ -146,7 +146,7 @@ public:
 
 private:
 	UPROPERTY(Category = AI, BlueprintGetter = GetDeprecatedActionsComponent)
-	TObjectPtr<UDEPRECATED_UPawnActionsComponent> ActionsComp_DEPRECATED;
+	TObjectPtr<UDEPRECATED_PawnActionsComponent> ActionsComp_DEPRECATED;
 
 protected:
 	/** blackboard */
@@ -456,13 +456,13 @@ public:
 	// DEPRECATED
 	//----------------------------------------------------------------------//
 	UE_DEPRECATED(5.2, "PawnActions have been deprecated and are no longer being supported. It will get removed in following UE5 releases. Use GameplayTasks or AITasks instead.")
-	UDEPRECATED_UPawnActionsComponent* GetActionsComp() const { return ActionsComp_DEPRECATED; }
+	UDEPRECATED_PawnActionsComponent* GetActionsComp() const { return ActionsComp_DEPRECATED; }
 	
 	UE_DEPRECATED(5.2, "PawnActions have been deprecated and are no longer being supported. It will get removed in following UE5 releases. Use GameplayTasks or AITasks instead.")
-	bool PerformAction(UDEPRECATED_UPawnAction& Action, EAIRequestPriority::Type Priority, UObject* const Instigator = NULL);
+	bool PerformAction(UDEPRECATED_PawnAction& Action, EAIRequestPriority::Type Priority, UObject* const Instigator = NULL);
 
 	UFUNCTION(BlueprintGetter, meta = (DeprecatedFunction, DeprecationMessage = "PawnActions have been deprecated and are no longer being supported. It will get removed in following UE5 releases. Use GameplayTasks or AITasks instead."))
-	UDEPRECATED_UPawnActionsComponent* GetDeprecatedActionsComponent() const { return ActionsComp_DEPRECATED; }
+	UDEPRECATED_PawnActionsComponent* GetDeprecatedActionsComponent() const { return ActionsComp_DEPRECATED; }
 };
 
 //----------------------------------------------------------------------//
diff --git a/Engine/Source/Runtime/AIModule/Classes/Actions/PawnAction.h b/Engine/Source/Runtime/AIModule/Classes/Actions/PawnAction.h
--- a/Engine/Source/Runtime/AIModule/Classes/Actions/PawnAction.h
+++ b/Engine/Source/Runtime/AIModule/Classes/Actions/PawnAction.h
@@ -12,8 +12,8 @@
 
 class AController;
 class APawn;
-class UDEPRECATED_UPawnAction;
-class UDEPRECATED_UPawnActionsComponent;
+class UDEPRECATED_PawnAction;
+class UDEPRECATED_PawnActionsComponent;
 struct FPawnActionStack;
 
 UENUM()
@@ -27,7 +27,7 @@ namespace EPawnSubActionTriggeringPolicy
 }
 
 AIMODULE_API DECLARE_LOG_CATEGORY_EXTERN(LogPawnAction, Warning, All);
-DECLARE_DELEGATE_TwoParams(FPawnActionEventDelegate, UDEPRECATED_UPawnAction&, EPawnActionEventType::Type);
+DECLARE_DELEGATE_TwoParams(FPawnActionEventDelegate, UDEPRECATED_PawnAction&, EPawnActionEventType::Type);
 
 UENUM()
 namespace EPawnActionFailHandling
@@ -44,24 +44,24 @@ namespace EPawnActionFailHandling
  *	* Actions are created paused
  */
 UCLASS(abstract, EditInlineNew, deprecated, meta = (DeprecationMessage = "PawnActions have been deprecated and are no longer being supported. It will get removed in following UE5 releases. Use GameplayTasks or AITasks instead."))
-class AIMODULE_API UDEPRECATED_UPawnAction : public UObject
+class AIMODULE_API UDEPRECATED_PawnAction : public UObject
 {
 	GENERATED_UCLASS_BODY()
 
-	friend UDEPRECATED_UPawnActionsComponent;
+	friend UDEPRECATED_PawnActionsComponent;
 	friend FPawnActionStack;
 
 private:
 	/** Current child node executing on top of this Action */
 	UPROPERTY(Transient)
-	TObjectPtr<UDEPRECATED_UPawnAction> ChildAction_DEPRECATED;
+	TObjectPtr<UDEPRECATED_PawnAction> ChildAction_DEPRECATED;
 
 	UPROPERTY(Transient)
-	TObjectPtr<UDEPRECATED_UPawnAction> ParentAction_DEPRECATED;
+	TObjectPtr<UDEPRECATED_PawnAction> ParentAction_DEPRECATED;
 
 	/** Extra reference to the component this action is being governed by */
 	UPROPERTY(Transient)
-	TObjectPtr<UDEPRECATED_UPawnActionsComponent> OwnerComponent_DEPRECATED;
+	TObjectPtr<UDEPRECATED_PawnActionsComponent> OwnerComponent_DEPRECATED;
 	
 	/** indicates an object that caused this action. Used for mass removal of actions 
 	 *	by specific object */
@@ -137,9 +137,9 @@ public:
 	virtual UWorld* GetWorld() const override;
 	// End UObject
 
-	FORCEINLINE const UDEPRECATED_UPawnAction* GetParentAction() const { return ParentAction_DEPRECATED; }
-	FORCEINLINE const UDEPRECATED_UPawnAction* GetChildAction() const { return ChildAction_DEPRECATED; }
-	FORCEINLINE UDEPRECATED_UPawnAction* GetChildAction() { return ChildAction_DEPRECATED; }
+	FORCEINLINE const UDEPRECATED_PawnAction* GetParentAction() const { return ParentAction_DEPRECATED; }
+	FORCEINLINE const UDEPRECATED_PawnAction* GetChildAction() const { return ChildAction_DEPRECATED; }
+	FORCEINLINE UDEPRECATED_PawnAction* GetChildAction() { return ChildAction_DEPRECATED; }
 	FORCEINLINE bool IsPaused() const { return !!bPaused; }
 	FORCEINLINE bool IsActive() const { return FinishResult == EPawnActionResult::InProgress && IsPaused() == false && AbortState == EPawnActionAbortState::NotBeingAborted; }
 	FORCEINLINE bool IsBeingAborted() const { return AbortState != EPawnActionAbortState::NotBeingAborted; }
@@ -169,12 +169,12 @@ protected:
 	 *	@NOTE do not make this virtual! Contains some essential logic. */
 	EPawnActionAbortState::Type Abort(EAIForceParam::Type ShouldForce = EAIForceParam::DoNotForce);
 	
-	FORCEINLINE UDEPRECATED_UPawnActionsComponent* GetOwnerComponent() { return OwnerComponent_DEPRECATED; }
+	FORCEINLINE UDEPRECATED_PawnActionsComponent* GetOwnerComponent() { return OwnerComponent_DEPRECATED; }
 public:
 	FORCEINLINE EAIRequestPriority::Type GetPriority() const { return ExecutionPriority; }
 	FORCEINLINE EPawnActionResult::Type GetResult() const { return FinishResult; }
 	FORCEINLINE EPawnActionAbortState::Type GetAbortState() const { return AbortState; }
-	FORCEINLINE UDEPRECATED_UPawnActionsComponent* GetOwnerComponent() const { return OwnerComponent_DEPRECATED; }
+	FORCEINLINE UDEPRECATED_PawnActionsComponent* GetOwnerComponent() const { return OwnerComponent_DEPRECATED; }
 	FORCEINLINE UObject* GetInstigator() const { return Instigator; }
 	APawn* GetPawn() const;
 	AController* GetController() const;
@@ -182,7 +182,7 @@ public:
 	template<class TActionClass>
 	static TActionClass* CreateActionInstance(UWorld& World)
 	{
-		TSubclassOf<UDEPRECATED_UPawnAction> ActionClass = TActionClass::StaticClass();
+		TSubclassOf<UDEPRECATED_PawnAction> ActionClass = TActionClass::StaticClass();
 		return NewObject<TActionClass>(&World, ActionClass);
 	}
 
@@ -204,7 +204,7 @@ public:
 
 	UE_DEPRECATED(5.2, "PawnActions have been deprecated and are no longer being supported. It will get removed in following UE5 releases. Use GameplayTasks or AITasks instead.")
 	UFUNCTION(BlueprintCallable, Category = "AI|PawnActions", meta = (WorldContext="WorldContextObject", DeprecatedFunction, DeprecationMessage = "PawnActions have been deprecated and are no longer being supported. It will get removed in following UE5 releases. Use GameplayTasks or AITasks instead."))
-	static UDEPRECATED_UPawnAction* CreateActionInstance(UObject* WorldContextObject, TSubclassOf<UDEPRECATED_UPawnAction> ActionClass);
+	static UDEPRECATED_PawnAction* CreateActionInstance(UObject* WorldContextObject, TSubclassOf<UDEPRECATED_PawnAction> ActionClass);
 
 	//----------------------------------------------------------------------//
 	// debug
@@ -226,7 +226,7 @@ protected:
 
 	void StopWaitingForMessages();
 
-	void SetOwnerComponent(UDEPRECATED_UPawnActionsComponent* Component);
+	void SetOwnerComponent(UDEPRECATED_PawnActionsComponent* Component);
 
 	void SetInstigator(UObject* const InInstigator);
 
@@ -237,7 +237,7 @@ protected:
 	 *	@NOTE if action fails to start no finishing or aborting mechanics will be triggered */
 	virtual bool Start();
 	/** called to pause action when higher priority or child action kicks in */
-	virtual bool Pause(const UDEPRECATED_UPawnAction* PausedBy);
+	virtual bool Pause(const UDEPRECATED_PawnAction* PausedBy);
 	/** called to resume action after being paused */
 	virtual bool Resume();
 	/** called when this action is being removed from action stacks */
@@ -246,10 +246,10 @@ protected:
 	 *	@NOTE gets called _AFTER_ child's OnFinished to give child action chance 
 	 *		to prepare "finishing data" for parent to read. 
 	 *	@NOTE clears parent-child binding */
-	virtual void OnChildFinished(UDEPRECATED_UPawnAction& Action, EPawnActionResult::Type WithResult);
+	virtual void OnChildFinished(UDEPRECATED_PawnAction& Action, EPawnActionResult::Type WithResult);
 
 	/** apart from doing regular push request copies additional values from Parent, like Priority and Instigator */
-	bool PushChildAction(UDEPRECATED_UPawnAction& Action);
+	bool PushChildAction(UDEPRECATED_PawnAction& Action);
 	
 	/** performs actual work on aborting Action. Should be called exclusively by Abort function
 	 *	@return only valid return values here are LatendAbortInProgress and AbortDone */
@@ -273,7 +273,7 @@ private:
 //----------------------------------------------------------------------//
 // Blueprint inlines
 //----------------------------------------------------------------------//
-FORCEINLINE TEnumAsByte<EAIRequestPriority::Type> UDEPRECATED_UPawnAction::GetActionPriority()
+FORCEINLINE TEnumAsByte<EAIRequestPriority::Type> UDEPRECATED_PawnAction::GetActionPriority()
 {
 	return ExecutionPriority;
 }
diff --git a/Engine/Source/Runtime/AIModule/Classes/Actions/PawnAction_BlueprintBase.h b/Engine/Source/Runtime/AIModule/Classes/Actions/PawnAction_BlueprintBase.h
--- a/Engine/Source/Runtime/AIModule/Classes/Actions/PawnAction_BlueprintBase.h
+++ b/Engine/Source/Runtime/AIModule/Classes/Actions/PawnAction_BlueprintBase.h
@@ -10,7 +10,7 @@
 class APawn;
 
 UCLASS(abstract, Blueprintable)
-class AIMODULE_API UDEPRECATED_UPawnAction_BlueprintBase : public UDEPRECATED_UPawnAction
+class AIMODULE_API UDEPRECATED_PawnAction_BlueprintBase : public UDEPRECATED_PawnAction
 {
 	GENERATED_UCLASS_BODY()
 
@@ -34,7 +34,7 @@ public:
 protected:
 	virtual void Tick(float DeltaTime) override;
 	virtual bool Start() override;
-	virtual bool Pause(const UDEPRECATED_UPawnAction* PausedBy) override;
+	virtual bool Pause(const UDEPRECATED_PawnAction* PausedBy) override;
 	virtual bool Resume() override;
 	virtual void OnFinished(EPawnActionResult::Type WithResult) override;
 };
diff --git a/Engine/Source/Runtime/AIModule/Classes/Actions/PawnAction_Move.h b/Engine/Source/Runtime/AIModule/Classes/Actions/PawnAction_Move.h
--- a/Engine/Source/Runtime/AIModule/Classes/Actions/PawnAction_Move.h
+++ b/Engine/Source/Runtime/AIModule/Classes/Actions/PawnAction_Move.h
@@ -24,7 +24,7 @@ namespace EPawnActionMoveMode
 }
 
 UCLASS()
-class AIMODULE_API UDEPRECATED_UPawnAction_Move : public UDEPRECATED_UPawnAction
+class AIMODULE_API UDEPRECATED_PawnAction_Move : public UDEPRECATED_PawnAction
 {
 	GENERATED_UCLASS_BODY()
 protected:
@@ -71,8 +71,8 @@ protected:
 public:
 	virtual void BeginDestroy() override;
 
-	static UDEPRECATED_UPawnAction_Move* CreateAction(UWorld& World, AActor* GoalActor, EPawnActionMoveMode::Type Mode);
-	static UDEPRECATED_UPawnAction_Move* CreateAction(UWorld& World, const FVector& GoalLocation, EPawnActionMoveMode::Type Mode);
+	static UDEPRECATED_PawnAction_Move* CreateAction(UWorld& World, AActor* GoalActor, EPawnActionMoveMode::Type Mode);
+	static UDEPRECATED_PawnAction_Move* CreateAction(UWorld& World, const FVector& GoalLocation, EPawnActionMoveMode::Type Mode);
 
 	static bool CheckAlreadyAtGoal(AAIController& Controller, const FVector& TestLocation, float Radius);
 	static bool CheckAlreadyAtGoal(AAIController& Controller, const AActor& TestGoal, float Radius);
@@ -108,7 +108,7 @@ protected:
 
 	void ClearPath();
 	virtual bool Start() override;
-	virtual bool Pause(const UDEPRECATED_UPawnAction* PausedBy) override;
+	virtual bool Pause(const UDEPRECATED_PawnAction* PausedBy) override;
 	virtual bool Resume() override;
 	virtual void OnFinished(EPawnActionResult::Type WithResult) override;
 	virtual EPawnActionAbortState::Type PerformAbort(EAIForceParam::Type ShouldForce) override;
diff --git a/Engine/Source/Runtime/AIModule/Classes/Actions/PawnAction_Repeat.h b/Engine/Source/Runtime/AIModule/Classes/Actions/PawnAction_Repeat.h
--- a/Engine/Source/Runtime/AIModule/Classes/Actions/PawnAction_Repeat.h
+++ b/Engine/Source/Runtime/AIModule/Classes/Actions/PawnAction_Repeat.h
@@ -8,7 +8,7 @@
 #include "PawnAction_Repeat.generated.h"
 
 UCLASS()
-class AIMODULE_API UDEPRECATED_UPawnAction_Repeat : public UDEPRECATED_UPawnAction
+class AIMODULE_API UDEPRECATED_PawnAction_Repeat : public UDEPRECATED_PawnAction
 {
 	GENERATED_UCLASS_BODY()
 
@@ -19,10 +19,10 @@ class AIMODULE_API UDEPRECATED_UPawnAction_Repeat : public UDEPRECATED_UPawnActi
 
 	/** Action to repeat. This instance won't really be run, it's a source for copying actions to be actually performed */
 	UPROPERTY()
-	TObjectPtr<UDEPRECATED_UPawnAction> ActionToRepeat_DEPRECATED;
+	TObjectPtr<UDEPRECATED_PawnAction> ActionToRepeat_DEPRECATED;
 
 	UPROPERTY(Transient)
-	TObjectPtr<UDEPRECATED_UPawnAction> RecentActionCopy_DEPRECATED;
+	TObjectPtr<UDEPRECATED_PawnAction> RecentActionCopy_DEPRECATED;
 
 	UPROPERTY(Category = PawnAction, EditAnywhere, BlueprintReadOnly)
 	TEnumAsByte<EPawnActionFailHandling::Type> ChildFailureHandlingMode;
@@ -31,13 +31,13 @@ class AIMODULE_API UDEPRECATED_UPawnAction_Repeat : public UDEPRECATED_UPawnActi
 
 	EPawnSubActionTriggeringPolicy::Type SubActionTriggeringPolicy;
 
-	/** @param NumberOfRepeats number of times to repeat action. UDEPRECATED_UPawnAction_Repeat::LoopForever loops forever */
-	static UDEPRECATED_UPawnAction_Repeat* CreateAction(UWorld& World, UDEPRECATED_UPawnAction* ActionToRepeat, int32 NumberOfRepeats, EPawnSubActionTriggeringPolicy::Type InSubActionTriggeringPolicy = EPawnSubActionTriggeringPolicy::CopyBeforeTriggering);
+	/** @param NumberOfRepeats number of times to repeat action. UDEPRECATED_PawnAction_Repeat::LoopForever loops forever */
+	static UDEPRECATED_PawnAction_Repeat* CreateAction(UWorld& World, UDEPRECATED_PawnAction* ActionToRepeat, int32 NumberOfRepeats, EPawnSubActionTriggeringPolicy::Type InSubActionTriggeringPolicy = EPawnSubActionTriggeringPolicy::CopyBeforeTriggering);
 
 protected:
 	virtual bool Start() override;
 	virtual bool Resume() override;
-	virtual void OnChildFinished(UDEPRECATED_UPawnAction& Action, EPawnActionResult::Type WithResult) override;
+	virtual void OnChildFinished(UDEPRECATED_PawnAction& Action, EPawnActionResult::Type WithResult) override;
 
 	bool PushSubAction();
 };
diff --git a/Engine/Source/Runtime/AIModule/Classes/Actions/PawnAction_Sequence.h b/Engine/Source/Runtime/AIModule/Classes/Actions/PawnAction_Sequence.h
--- a/Engine/Source/Runtime/AIModule/Classes/Actions/PawnAction_Sequence.h
+++ b/Engine/Source/Runtime/AIModule/Classes/Actions/PawnAction_Sequence.h
@@ -8,29 +8,29 @@
 #include "PawnAction_Sequence.generated.h"
 
 UCLASS()
-class AIMODULE_API UDEPRECATED_UPawnAction_Sequence : public UDEPRECATED_UPawnAction
+class AIMODULE_API UDEPRECATED_PawnAction_Sequence : public UDEPRECATED_PawnAction
 {
 	GENERATED_UCLASS_BODY()
 
 	UPROPERTY()
-	TArray<TObjectPtr<UDEPRECATED_UPawnAction>> ActionSequence_DEPRECATED;
+	TArray<TObjectPtr<UDEPRECATED_PawnAction>> ActionSequence_DEPRECATED;
 
 	UPROPERTY(Category = PawnAction, EditAnywhere, BlueprintReadOnly)
 	TEnumAsByte<EPawnActionFailHandling::Type> ChildFailureHandlingMode;
 
 	UPROPERTY(Transient)
-	TObjectPtr<UDEPRECATED_UPawnAction> RecentActionCopy_DEPRECATED;
+	TObjectPtr<UDEPRECATED_PawnAction> RecentActionCopy_DEPRECATED;
 
 	uint32 CurrentActionIndex;
 
 	EPawnSubActionTriggeringPolicy::Type SubActionTriggeringPolicy;
 
-	static UDEPRECATED_UPawnAction_Sequence* CreateAction(UWorld& World, TArray<UDEPRECATED_UPawnAction*>& ActionSequence, EPawnSubActionTriggeringPolicy::Type InSubActionTriggeringPolicy = EPawnSubActionTriggeringPolicy::CopyBeforeTriggering);
+	static UDEPRECATED_PawnAction_Sequence* CreateAction(UWorld& World, TArray<UDEPRECATED_PawnAction*>& ActionSequence, EPawnSubActionTriggeringPolicy::Type InSubActionTriggeringPolicy = EPawnSubActionTriggeringPolicy::CopyBeforeTriggering);
 
 protected:
 	virtual bool Start() override;
 	virtual bool Resume() override;
-	virtual void OnChildFinished(UDEPRECATED_UPawnAction& Action, EPawnActionResult::Type WithResult) override;
+	virtual void OnChildFinished(UDEPRECATED_PawnAction& Action, EPawnActionResult::Type WithResult) override;
 
 	bool PushNextActionCopy();
 };
diff --git a/Engine/Source/Runtime/AIModule/Classes/Actions/PawnAction_Wait.h b/Engine/Source/Runtime/AIModule/Classes/Actions/PawnAction_Wait.h
--- a/Engine/Source/Runtime/AIModule/Classes/Actions/PawnAction_Wait.h
+++ b/Engine/Source/Runtime/AIModule/Classes/Actions/PawnAction_Wait.h
@@ -10,7 +10,7 @@
 
 /** uses system timers rather then ticking */
 UCLASS()
-class AIMODULE_API UDEPRECATED_UPawnAction_Wait : public UDEPRECATED_UPawnAction
+class AIMODULE_API UDEPRECATED_PawnAction_Wait : public UDEPRECATED_PawnAction
 {
 	GENERATED_UCLASS_BODY()
 		
@@ -22,10 +22,10 @@ class AIMODULE_API UDEPRECATED_UPawnAction_Wait : public UDEPRECATED_UPawnAction
 	FTimerHandle TimerHandle;
 
 	/** InTimeToWait < 0 (or just FAISystem::InfiniteInterval) will result in waiting forever */
-	static UDEPRECATED_UPawnAction_Wait* CreateAction(UWorld& World, float InTimeToWait = FAISystem::InfiniteInterval);
+	static UDEPRECATED_PawnAction_Wait* CreateAction(UWorld& World, float InTimeToWait = FAISystem::InfiniteInterval);
 
 	virtual bool Start() override;
-	virtual bool Pause(const UDEPRECATED_UPawnAction* PausedBy) override;
+	virtual bool Pause(const UDEPRECATED_PawnAction* PausedBy) override;
 	virtual bool Resume() override;
 	virtual EPawnActionAbortState::Type PerformAbort(EAIForceParam::Type ShouldForce) override;
 
diff --git a/Engine/Source/Runtime/AIModule/Classes/Actions/PawnActionsComponent.h b/Engine/Source/Runtime/AIModule/Classes/Actions/PawnActionsComponent.h
--- a/Engine/Source/Runtime/AIModule/Classes/Actions/PawnActionsComponent.h
+++ b/Engine/Source/Runtime/AIModule/Classes/Actions/PawnActionsComponent.h
@@ -22,7 +22,7 @@ struct AIMODULE_API FPawnActionEvent
 	static const int32 FakeActionIndex = INDEX_NONE;
 
 	UPROPERTY()
-	TObjectPtr<UDEPRECATED_UPawnAction> Action;
+	TObjectPtr<UDEPRECATED_PawnAction> Action;
 
 	EPawnActionEventType::Type EventType;
 
@@ -34,7 +34,7 @@ struct AIMODULE_API FPawnActionEvent
 	FPawnActionEvent() : Action(NULL), EventType(EPawnActionEventType::Invalid), Priority(EAIRequestPriority::MAX), Index(uint32(-1))
 	{}
 
-	FPawnActionEvent(UDEPRECATED_UPawnAction& Action, EPawnActionEventType::Type EventType, uint32 Index);
+	FPawnActionEvent(UDEPRECATED_PawnAction& Action, EPawnActionEventType::Type EventType, uint32 Index);
 
 	bool operator==(const FPawnActionEvent& Other) const { return (Action == Other.Action) && (EventType == Other.EventType) && (Priority == Other.Priority); }
 };
@@ -50,7 +50,7 @@ struct AIMODULE_API FPawnActionStack
 
 private:
 	UPROPERTY()
-	TObjectPtr<UDEPRECATED_UPawnAction> TopAction;
+	TObjectPtr<UDEPRECATED_PawnAction> TopAction;
 
 public:
 	void Pause();
@@ -58,13 +58,13 @@ public:
 
 	/** All it does is tie actions into a double-linked list making NewTopAction
 	 *	new stack's top */
-	void PushAction(UDEPRECATED_UPawnAction& NewTopAction);
+	void PushAction(UDEPRECATED_PawnAction& NewTopAction);
 
 	/** Looks through the double-linked action list looking for specified action
 	 *	and if found action will be popped along with all it's siblings */
-	void PopAction(UDEPRECATED_UPawnAction& ActionToPop);
+	void PopAction(UDEPRECATED_PawnAction& ActionToPop);
 	
-	FORCEINLINE UDEPRECATED_UPawnAction* GetTop() const { return TopAction; }
+	FORCEINLINE UDEPRECATED_PawnAction* GetTop() const { return TopAction; }
 
 	FORCEINLINE bool IsEmpty() const { return TopAction == NULL; }
 
@@ -75,7 +75,7 @@ public:
 };
 
 UCLASS(deprecated, meta = (DeprecationMessage = "PawnActions have been deprecated and are no longer being supported. It will get removed in following UE5 releases. Use GameplayTasks or AITasks instead."))
-class AIMODULE_API UDEPRECATED_UPawnActionsComponent : public UActorComponent
+class AIMODULE_API UDEPRECATED_PawnActionsComponent : public UActorComponent
 {
 	GENERATED_UCLASS_BODY()
 
@@ -90,7 +90,7 @@ protected:
 	TArray<FPawnActionEvent> ActionEvents;
 
 	UPROPERTY(Transient)
-	TObjectPtr<UDEPRECATED_UPawnAction> CurrentAction_DEPRECATED;
+	TObjectPtr<UDEPRECATED_PawnAction> CurrentAction_DEPRECATED;
 
 	/** set when logic was locked by hi priority stack */
 	uint32 bLockedAILogic : 1;
@@ -109,10 +109,10 @@ public:
 	//----------------------------------------------------------------------//
 
 	UFUNCTION(BlueprintCallable, Category = "AI|PawnActions", meta = (DisplayName = "PerformAction_DEPRECATED", ScriptName = "PerformAction", DeprecatedFunction, DeprecationMessage = "PawnActions have been deprecated and are no longer being supported. It will get removed in following UE5 releases. Use GameplayTasks or AITasks instead."))
-	static bool K2_PerformAction(APawn* Pawn, UDEPRECATED_UPawnAction* Action, TEnumAsByte<EAIRequestPriority::Type> Priority = EAIRequestPriority::HardScript);
+	static bool K2_PerformAction(APawn* Pawn, UDEPRECATED_PawnAction* Action, TEnumAsByte<EAIRequestPriority::Type> Priority = EAIRequestPriority::HardScript);
 
 	UE_DEPRECATED(5.2, "PawnActions have been deprecated and are no longer being supported. It will get removed in following UE5 releases. Use GameplayTasks or AITasks instead.")
-	static bool PerformAction(APawn& Pawn, UDEPRECATED_UPawnAction& Action, TEnumAsByte<EAIRequestPriority::Type> Priority = EAIRequestPriority::HardScript);
+	static bool PerformAction(APawn& Pawn, UDEPRECATED_PawnAction& Action, TEnumAsByte<EAIRequestPriority::Type> Priority = EAIRequestPriority::HardScript);
 
 	//----------------------------------------------------------------------//
 	// 
@@ -124,29 +124,29 @@ public:
 	FORCEINLINE APawn* GetControlledPawn() { return ControlledPawn; }
 	FORCEINLINE const APawn* GetControlledPawn() const { return ControlledPawn; }
 	FORCEINLINE AController* GetController() { return ControlledPawn ? ControlledPawn->GetController() : NULL; }
-	FORCEINLINE UDEPRECATED_UPawnAction* GetCurrentAction() { return CurrentAction_DEPRECATED; }
+	FORCEINLINE UDEPRECATED_PawnAction* GetCurrentAction() { return CurrentAction_DEPRECATED; }
 
-	bool OnEvent(UDEPRECATED_UPawnAction& Action, EPawnActionEventType::Type Event);
+	bool OnEvent(UDEPRECATED_PawnAction& Action, EPawnActionEventType::Type Event);
 
 	UFUNCTION(BlueprintCallable, Category = PawnAction, meta = (DisplayName = "PushAction_DEPRECATED", ScriptName = "PushAction", DeprecatedFunction, DeprecationMessage = "PawnActions have been deprecated and are no longer being supported. It will get removed in following UE5 releases. Use GameplayTasks or AITasks instead."))
-	bool K2_PushAction(UDEPRECATED_UPawnAction* NewAction, EAIRequestPriority::Type Priority, UObject* Instigator = NULL);
+	bool K2_PushAction(UDEPRECATED_PawnAction* NewAction, EAIRequestPriority::Type Priority, UObject* Instigator = NULL);
 
 	UE_DEPRECATED(5.2, "PawnActions have been deprecated and are no longer being supported. It will get removed in following UE5 releases. Use GameplayTasks or AITasks instead.")
-	bool PushAction(UDEPRECATED_UPawnAction& NewAction, EAIRequestPriority::Type Priority, UObject* Instigator = NULL);	
+	bool PushAction(UDEPRECATED_PawnAction& NewAction, EAIRequestPriority::Type Priority, UObject* Instigator = NULL);	
 
 	/** Aborts given action instance */
 	UFUNCTION(BlueprintCallable, Category = PawnAction, meta = (DisplayName = "AbortAction_DEPRECATED", ScriptName = "AbortAction", DeprecatedFunction, DeprecationMessage = "PawnActions have been deprecated and are no longer being supported. It will get removed in following UE5 releases. Use GameplayTasks or AITasks instead."))
-	EPawnActionAbortState::Type K2_AbortAction(UDEPRECATED_UPawnAction* ActionToAbort);
+	EPawnActionAbortState::Type K2_AbortAction(UDEPRECATED_PawnAction* ActionToAbort);
 
 	UE_DEPRECATED(5.2, "PawnActions have been deprecated and are no longer being supported. It will get removed in following UE5 releases. Use GameplayTasks or AITasks instead.")
-	EPawnActionAbortState::Type AbortAction(UDEPRECATED_UPawnAction& ActionToAbort);
+	EPawnActionAbortState::Type AbortAction(UDEPRECATED_PawnAction& ActionToAbort);
 
 	/** Aborts given action instance */
 	UFUNCTION(BlueprintCallable, Category = PawnAction, meta = (DisplayName = "ForceAbortAction_DEPRECATED", ScriptName = "ForceAbortAction", DeprecatedFunction, DeprecationMessage = "PawnActions have been deprecated and are no longer being supported. It will get removed in following UE5 releases. Use GameplayTasks or AITasks instead."))
-	EPawnActionAbortState::Type K2_ForceAbortAction(UDEPRECATED_UPawnAction* ActionToAbort);
+	EPawnActionAbortState::Type K2_ForceAbortAction(UDEPRECATED_PawnAction* ActionToAbort);
 
 	UE_DEPRECATED(5.2, "PawnActions have been deprecated and are no longer being supported. It will get removed in following UE5 releases. Use GameplayTasks or AITasks instead.")
-	EPawnActionAbortState::Type ForceAbortAction(UDEPRECATED_UPawnAction& ActionToAbort);
+	EPawnActionAbortState::Type ForceAbortAction(UDEPRECATED_PawnAction& ActionToAbort);
 
 	/** removes all actions instigated with Priority by Instigator
 	 *	@param Priority if equal to EAIRequestPriority::MAX then all priority queues will be searched. 
@@ -156,8 +156,8 @@ public:
 	
 	void TickComponent(float DeltaTime, enum ELevelTick TickType, FActorComponentTickFunction *ThisTickFunction) override;
 
-	FORCEINLINE UDEPRECATED_UPawnAction* GetActiveAction(EAIRequestPriority::Type Priority) const { return ActionStacks[Priority].GetTop(); }
-	bool HasActiveActionOfType(EAIRequestPriority::Type Priority, TSubclassOf<UDEPRECATED_UPawnAction> PawnActionClass) const;
+	FORCEINLINE UDEPRECATED_PawnAction* GetActiveAction(EAIRequestPriority::Type Priority) const { return ActionStacks[Priority].GetTop(); }
+	bool HasActiveActionOfType(EAIRequestPriority::Type Priority, TSubclassOf<UDEPRECATED_PawnAction> PawnActionClass) const;
 
 #if ENABLE_VISUAL_LOG
 	void DescribeSelfToVisLog(struct FVisualLogEntry* Snapshot) const;
@@ -182,5 +182,5 @@ protected:
 
 private:
 	/** Removed all pending action events associated with PawnAction. Private to make sure it's called only in special cases */
-	void RemoveEventsForAction(UDEPRECATED_UPawnAction& PawnAction);
+	void RemoveEventsForAction(UDEPRECATED_PawnAction& PawnAction);
 };
diff --git a/Engine/Source/Runtime/AIModule/Classes/BehaviorTree/Tasks/BTTask_PawnActionBase.h b/Engine/Source/Runtime/AIModule/Classes/BehaviorTree/Tasks/BTTask_PawnActionBase.h
--- a/Engine/Source/Runtime/AIModule/Classes/BehaviorTree/Tasks/BTTask_PawnActionBase.h
+++ b/Engine/Source/Runtime/AIModule/Classes/BehaviorTree/Tasks/BTTask_PawnActionBase.h
@@ -7,7 +7,7 @@
 #include "BehaviorTree/BTTaskNode.h"
 #include "BTTask_PawnActionBase.generated.h"
 
-class UDEPRECATED_UPawnAction;
+class UDEPRECATED_PawnAction;
 
 enum class EPawnActionTaskResult : uint8
 {
@@ -36,17 +36,17 @@ class AIMODULE_API UBTTask_PawnActionBase : public UBTTaskNode
 protected:
 
 	/** starts executing pawn action */
-	EBTNodeResult::Type PushAction(UBehaviorTreeComponent& OwnerComp, UDEPRECATED_UPawnAction& Action);
+	EBTNodeResult::Type PushAction(UBehaviorTreeComponent& OwnerComp, UDEPRECATED_PawnAction& Action);
 
 	/** action observer, updates state of task */
-	virtual void OnActionEvent(UDEPRECATED_UPawnAction& Action, EPawnActionEventType::Type Event);
+	virtual void OnActionEvent(UDEPRECATED_PawnAction& Action, EPawnActionEventType::Type Event);
 
 	/** called when action is removed from stack (FinishedAborting) by some external event
 	 *  default behavior: finish task as failed */
-	virtual void OnActionLost(UDEPRECATED_UPawnAction& Action);
+	virtual void OnActionLost(UDEPRECATED_PawnAction& Action);
 
 public:
 
 	/** helper functions, should be used when behavior tree task deals with pawn actions, but can't derive from this class */
-	static EPawnActionTaskResult ActionEventHandler(UBTTaskNode* TaskNode, UDEPRECATED_UPawnAction& Action, EPawnActionEventType::Type Event);
+	static EPawnActionTaskResult ActionEventHandler(UBTTaskNode* TaskNode, UDEPRECATED_PawnAction& Action, EPawnActionEventType::Type Event);
 };
diff --git a/Engine/Source/Runtime/AIModule/Classes/BehaviorTree/Tasks/BTTask_PushPawnAction.h b/Engine/Source/Runtime/AIModule/Classes/BehaviorTree/Tasks/BTTask_PushPawnAction.h
--- a/Engine/Source/Runtime/AIModule/Classes/BehaviorTree/Tasks/BTTask_PushPawnAction.h
+++ b/Engine/Source/Runtime/AIModule/Classes/BehaviorTree/Tasks/BTTask_PushPawnAction.h
@@ -7,7 +7,7 @@
 #include "BehaviorTree/Tasks/BTTask_PawnActionBase.h"
 #include "BTTask_PushPawnAction.generated.h"
 
-class UDEPRECATED_UPawnAction;
+class UDEPRECATED_PawnAction;
 
 /**
  * Action task node.
@@ -20,7 +20,7 @@ class AIMODULE_API UBTTask_PushPawnAction : public UBTTask_PawnActionBase
 
 protected:
 	UPROPERTY(Instanced)
-	TObjectPtr<UDEPRECATED_UPawnAction> Action_DEPRECATED;
+	TObjectPtr<UDEPRECATED_PawnAction> Action_DEPRECATED;
 
 	virtual EBTNodeResult::Type ExecuteTask(UBehaviorTreeComponent& OwnerComp, uint8* NodeMemory) override;
 	virtual FString GetStaticDescription() const override;
diff --git a/Engine/Source/Runtime/AIModule/Private/AIController.cpp b/Engine/Source/Runtime/AIModule/Private/AIController.cpp
--- a/Engine/Source/Runtime/AIModule/Private/AIController.cpp
+++ b/Engine/Source/Runtime/AIModule/Private/AIController.cpp
@@ -47,7 +47,7 @@ AAIController::AAIController(const FObjectInitializer& ObjectInitializer)
 		PathFollowingComponent->OnRequestFinished.AddUObject(this, &AAIController::OnMoveCompleted);
 	}
 
-	ActionsComp_DEPRECATED = CreateOptionalDefaultSubobject<UDEPRECATED_UPawnActionsComponent>(TEXT("ActionsComp"));
+	ActionsComp_DEPRECATED = CreateOptionalDefaultSubobject<UDEPRECATED_PawnActionsComponent>(TEXT("ActionsComp"));
 
 	bSkipExtraLOSChecks = true;
 	bWantsPlayerState = false;
@@ -1082,7 +1082,7 @@ bool AAIController::SuggestTossVelocity(FVector& OutTossVelocity, FVector Start,
 
 	return UGameplayStatics::SuggestProjectileVelocity(this, OutTossVelocity, Start, End, TossSpeed, bPreferHighArc, CollisionRadius, GravityOverride, TraceOption);
 }
-bool AAIController::PerformAction(UDEPRECATED_UPawnAction& Action, EAIRequestPriority::Type Priority, UObject* const InInstigator /*= NULL*/)
+bool AAIController::PerformAction(UDEPRECATED_PawnAction& Action, EAIRequestPriority::Type Priority, UObject* const InInstigator /*= NULL*/)
 {
 PRAGMA_DISABLE_DEPRECATION_WARNINGS
 	return ActionsComp_DEPRECATED != NULL && ActionsComp_DEPRECATED->PushAction(Action, Priority, InInstigator);
diff --git a/Engine/Source/Runtime/AIModule/Private/Actions/PawnAction.cpp b/Engine/Source/Runtime/AIModule/Private/Actions/PawnAction.cpp
--- a/Engine/Source/Runtime/AIModule/Private/Actions/PawnAction.cpp
+++ b/Engine/Source/Runtime/AIModule/Private/Actions/PawnAction.cpp
@@ -24,7 +24,7 @@ namespace
 	}
 }
 
-UDEPRECATED_UPawnAction::UDEPRECATED_UPawnAction(const FObjectInitializer& ObjectInitializer)
+UDEPRECATED_PawnAction::UDEPRECATED_PawnAction(const FObjectInitializer& ObjectInitializer)
 	: Super(ObjectInitializer)
 	, RequiredResources(FAIResourcesSet::AllResources)
 	, AbortState(EPawnActionAbortState::NeverStarted)
@@ -36,16 +36,16 @@ UDEPRECATED_UPawnAction::UDEPRECATED_UPawnAction(const FObjectInitializer& Objec
 	IndexOnStack = INDEX_NONE;
 }
 
-UWorld* UDEPRECATED_UPawnAction::GetWorld() const
+UWorld* UDEPRECATED_PawnAction::GetWorld() const
 {
 	return OwnerComponent_DEPRECATED ? OwnerComponent_DEPRECATED->GetWorld() : Cast<UWorld>(GetOuter());
 }
 
-void UDEPRECATED_UPawnAction::Tick(float DeltaTime)
+void UDEPRECATED_PawnAction::Tick(float DeltaTime)
 {
 }
 
-EPawnActionAbortState::Type UDEPRECATED_UPawnAction::Abort(EAIForceParam::Type ShouldForce)
+EPawnActionAbortState::Type UDEPRECATED_PawnAction::Abort(EAIForceParam::Type ShouldForce)
 {
 	// if already aborting, and this request is not Forced, just skip it
 	if (AbortState != EPawnActionAbortState::NotBeingAborted && ShouldForce != EAIForceParam::Force)
@@ -122,17 +122,17 @@ EPawnActionAbortState::Type UDEPRECATED_UPawnAction::Abort(EAIForceParam::Type S
 	return Result;
 }
 
-APawn* UDEPRECATED_UPawnAction::GetPawn() const
+APawn* UDEPRECATED_PawnAction::GetPawn() const
 {
 	return OwnerComponent_DEPRECATED ? OwnerComponent_DEPRECATED->GetControlledPawn() : NULL;
 }
 
-AController* UDEPRECATED_UPawnAction::GetController() const
+AController* UDEPRECATED_PawnAction::GetController() const
 {
 	return OwnerComponent_DEPRECATED ? OwnerComponent_DEPRECATED->GetController() : NULL;
 }
 
-void UDEPRECATED_UPawnAction::SetAbortState(EPawnActionAbortState::Type NewAbortState)
+void UDEPRECATED_PawnAction::SetAbortState(EPawnActionAbortState::Type NewAbortState)
 {
 	// allowing only progression
 	if (NewAbortState <= AbortState)
@@ -147,7 +147,7 @@ void UDEPRECATED_UPawnAction::SetAbortState(EPawnActionAbortState::Type NewAbort
 	}
 }
 
-void UDEPRECATED_UPawnAction::SendEvent(EPawnActionEventType::Type Event)
+void UDEPRECATED_PawnAction::SendEvent(EPawnActionEventType::Type Event)
 {
 	if (IsValid(OwnerComponent_DEPRECATED))
 	{
@@ -158,17 +158,17 @@ void UDEPRECATED_UPawnAction::SendEvent(EPawnActionEventType::Type Event)
 	ActionObserver.ExecuteIfBound(*this, Event);
 }
 
-void UDEPRECATED_UPawnAction::StopWaitingForMessages()
+void UDEPRECATED_PawnAction::StopWaitingForMessages()
 {
 	MessageHandlers.Reset();
 }
 
-void UDEPRECATED_UPawnAction::SetFinishResult(EPawnActionResult::Type Result)
+void UDEPRECATED_PawnAction::SetFinishResult(EPawnActionResult::Type Result)
 {
 	// once return value had been set it's no longer possible to back to InProgress
 	if (Result <= EPawnActionResult::InProgress)
 	{
-		UE_VLOG(GetPawn(), LogPawnAction, Warning, TEXT("%s> UDEPRECATED_UPawnAction::SetFinishResult setting FinishResult as EPawnActionResult::InProgress or EPawnActionResult::NotStarted - should not be happening"), *GetName());
+		UE_VLOG(GetPawn(), LogPawnAction, Warning, TEXT("%s> UDEPRECATED_PawnAction::SetFinishResult setting FinishResult as EPawnActionResult::InProgress or EPawnActionResult::NotStarted - should not be happening"), *GetName());
 		return;
 	}
 
@@ -178,11 +178,11 @@ void UDEPRECATED_UPawnAction::SetFinishResult(EPawnActionResult::Type Result)
 	}
 }
 
-void UDEPRECATED_UPawnAction::SetOwnerComponent(UDEPRECATED_UPawnActionsComponent* Component)
+void UDEPRECATED_PawnAction::SetOwnerComponent(UDEPRECATED_PawnActionsComponent* Component)
 {
 	if (OwnerComponent_DEPRECATED != NULL && OwnerComponent_DEPRECATED != Component)
 	{
-		UE_VLOG(GetPawn(), LogPawnAction, Warning, TEXT("%s> UDEPRECATED_UPawnAction::SetOwnerComponent called to change already set valid owner component"), *GetName());
+		UE_VLOG(GetPawn(), LogPawnAction, Warning, TEXT("%s> UDEPRECATED_PawnAction::SetOwnerComponent called to change already set valid owner component"), *GetName());
 	}
 
 	OwnerComponent_DEPRECATED = Component;
@@ -196,7 +196,7 @@ void UDEPRECATED_UPawnAction::SetOwnerComponent(UDEPRECATED_UPawnActionsComponen
 	}
 }
 
-void UDEPRECATED_UPawnAction::SetInstigator(UObject* const InInstigator)
+void UDEPRECATED_PawnAction::SetInstigator(UObject* const InInstigator)
 { 
 	if (Instigator && Instigator != InInstigator)
 	{
@@ -206,7 +206,7 @@ void UDEPRECATED_UPawnAction::SetInstigator(UObject* const InInstigator)
 	Instigator = InInstigator; 
 }
 
-void UDEPRECATED_UPawnAction::Finish(TEnumAsByte<EPawnActionResult::Type> WithResult)
+void UDEPRECATED_PawnAction::Finish(TEnumAsByte<EPawnActionResult::Type> WithResult)
 {
 	UE_VLOG(GetPawn(), LogPawnAction, Log, TEXT("%s> finishing with result %s")
 		, *GetName(), *GetActionResultName(WithResult));
@@ -218,7 +218,7 @@ void UDEPRECATED_UPawnAction::Finish(TEnumAsByte<EPawnActionResult::Type> WithRe
 	SendEvent(EPawnActionEventType::FinishedExecution);
 }
 
-bool UDEPRECATED_UPawnAction::Activate()
+bool UDEPRECATED_PawnAction::Activate()
 {
 	bool bResult = false; 
 
@@ -256,7 +256,7 @@ bool UDEPRECATED_UPawnAction::Activate()
 	return bResult;
 }
 
-void UDEPRECATED_UPawnAction::OnPopped()
+void UDEPRECATED_PawnAction::OnPopped()
 {
 	// not calling OnFinish if action haven't actually started
 	if (!bFailedToStart || bAlwaysNotifyOnFinished)
@@ -265,7 +265,7 @@ void UDEPRECATED_UPawnAction::OnPopped()
 	}
 }
 
-bool UDEPRECATED_UPawnAction::Start()
+bool UDEPRECATED_PawnAction::Start()
 {
 	AbortState = EPawnActionAbortState::NotBeingAborted;
 	FinishResult = EPawnActionResult::InProgress;
@@ -274,7 +274,7 @@ bool UDEPRECATED_UPawnAction::Start()
 }
 
 
-bool UDEPRECATED_UPawnAction::Pause(const UDEPRECATED_UPawnAction* PausedBy)
+bool UDEPRECATED_PawnAction::Pause(const UDEPRECATED_PawnAction* PausedBy)
 {
 	// parent should be paused anyway
 	ensure(ParentAction_DEPRECATED == NULL || ParentAction_DEPRECATED->IsPaused() == true);
@@ -300,7 +300,7 @@ bool UDEPRECATED_UPawnAction::Pause(const UDEPRECATED_UPawnAction* PausedBy)
 	return bPaused;
 }
 
-bool UDEPRECATED_UPawnAction::Resume()
+bool UDEPRECATED_PawnAction::Resume()
 {
 	// parent should be paused anyway
 	ensure(ParentAction_DEPRECATED == NULL || ParentAction_DEPRECATED->IsPaused() == true);
@@ -327,11 +327,11 @@ bool UDEPRECATED_UPawnAction::Resume()
 	return !bPaused;
 }
 
-void UDEPRECATED_UPawnAction::OnFinished(EPawnActionResult::Type WithResult)
+void UDEPRECATED_PawnAction::OnFinished(EPawnActionResult::Type WithResult)
 {
 }
 
-void UDEPRECATED_UPawnAction::OnChildFinished(UDEPRECATED_UPawnAction& Action, EPawnActionResult::Type WithResult)
+void UDEPRECATED_PawnAction::OnChildFinished(UDEPRECATED_PawnAction& Action, EPawnActionResult::Type WithResult)
 {
 	UE_VLOG(GetPawn(), LogPawnAction, Log, TEXT("%s> Child \'%s\' finished with result %s")
 		, *GetName(), *Action.GetName(), *GetActionResultName(WithResult));
@@ -342,10 +342,10 @@ void UDEPRECATED_UPawnAction::OnChildFinished(UDEPRECATED_UPawnAction& Action, E
 	ChildAction_DEPRECATED = NULL;
 }
 
-void UDEPRECATED_UPawnAction::OnPushed()
+void UDEPRECATED_PawnAction::OnPushed()
 {
 	IndexOnStack = 0;
-	UDEPRECATED_UPawnAction* PrevAction = ParentAction_DEPRECATED;
+	UDEPRECATED_PawnAction* PrevAction = ParentAction_DEPRECATED;
 	while (PrevAction)
 	{
 		++IndexOnStack;
@@ -356,7 +356,7 @@ void UDEPRECATED_UPawnAction::OnPushed()
 		, *GetName(), *GetPriorityName(), IndexOnStack, *GetNameSafe(Instigator));
 }
 
-bool UDEPRECATED_UPawnAction::PushChildAction(UDEPRECATED_UPawnAction& Action)
+bool UDEPRECATED_PawnAction::PushChildAction(UDEPRECATED_PawnAction& Action)
 {
 	bool bResult = false;
 	
@@ -382,21 +382,21 @@ bool UDEPRECATED_UPawnAction::PushChildAction(UDEPRECATED_UPawnAction& Action)
 // messaging
 //----------------------------------------------------------------------//
 
-void UDEPRECATED_UPawnAction::WaitForMessage(FName MessageType, FAIRequestID InRequestID)
+void UDEPRECATED_PawnAction::WaitForMessage(FName MessageType, FAIRequestID InRequestID)
 {
-	MessageHandlers.Add(FAIMessageObserver::Create(BrainComp, MessageType, InRequestID.GetID(), FOnAIMessage::CreateUObject(this, &UDEPRECATED_UPawnAction::HandleAIMessage)));
+	MessageHandlers.Add(FAIMessageObserver::Create(BrainComp, MessageType, InRequestID.GetID(), FOnAIMessage::CreateUObject(this, &UDEPRECATED_PawnAction::HandleAIMessage)));
 }
 
 //----------------------------------------------------------------------//
 // blueprint interface
 //----------------------------------------------------------------------//
 
-UDEPRECATED_UPawnAction* UDEPRECATED_UPawnAction::CreateActionInstance(UObject* WorldContextObject, TSubclassOf<UDEPRECATED_UPawnAction> ActionClass)
+UDEPRECATED_PawnAction* UDEPRECATED_PawnAction::CreateActionInstance(UObject* WorldContextObject, TSubclassOf<UDEPRECATED_PawnAction> ActionClass)
 {
 	UWorld* World = GEngine->GetWorldFromContextObject(WorldContextObject, EGetWorldErrorMode::LogAndReturnNull);
 	if (World && ActionClass)
 	{
-		return NewObject<UDEPRECATED_UPawnAction>(World, ActionClass);
+		return NewObject<UDEPRECATED_PawnAction>(World, ActionClass);
 	}
 	return NULL;
 }
@@ -405,7 +405,7 @@ UDEPRECATED_UPawnAction* UDEPRECATED_UPawnAction::CreateActionInstance(UObject*
 // debug
 //----------------------------------------------------------------------//
 
-FString UDEPRECATED_UPawnAction::GetStateDescription() const
+FString UDEPRECATED_PawnAction::GetStateDescription() const
 {
 	static const UEnum* AbortStateEnum = StaticEnum<EPawnActionAbortState::Type>(); 
 		
@@ -416,14 +416,14 @@ FString UDEPRECATED_UPawnAction::GetStateDescription() const
 	return IsPaused() ? TEXT("Paused") : TEXT("Active");
 }
 
-FString UDEPRECATED_UPawnAction::GetPriorityName() const
+FString UDEPRECATED_PawnAction::GetPriorityName() const
 {
 	static const UEnum* Enum = StaticEnum<EAIRequestPriority::Type>();
 	check(Enum);
 	return Enum->GetNameStringByValue(GetPriority());
 }
 
-FString UDEPRECATED_UPawnAction::GetDisplayName() const
+FString UDEPRECATED_PawnAction::GetDisplayName() const
 {
 	return GetClass()->GetName();
 }
diff --git a/Engine/Source/Runtime/AIModule/Private/Actions/PawnAction_BlueprintBase.cpp b/Engine/Source/Runtime/AIModule/Private/Actions/PawnAction_BlueprintBase.cpp
--- a/Engine/Source/Runtime/AIModule/Private/Actions/PawnAction_BlueprintBase.cpp
+++ b/Engine/Source/Runtime/AIModule/Private/Actions/PawnAction_BlueprintBase.cpp
@@ -7,20 +7,20 @@
 
 PRAGMA_DISABLE_DEPRECATION_WARNINGS
 
-UDEPRECATED_UPawnAction_BlueprintBase::UDEPRECATED_UPawnAction_BlueprintBase(const FObjectInitializer& ObjectInitializer)
+UDEPRECATED_PawnAction_BlueprintBase::UDEPRECATED_PawnAction_BlueprintBase(const FObjectInitializer& ObjectInitializer)
 	: Super(ObjectInitializer)
 {
-	UClass* StopAtClass = UDEPRECATED_UPawnAction_BlueprintBase::StaticClass();
+	UClass* StopAtClass = UDEPRECATED_PawnAction_BlueprintBase::StaticClass();
 	bWantsTick = BlueprintNodeHelpers::HasBlueprintFunction(TEXT("ActionTick"), *this, *StopAtClass);
 }
 
-void UDEPRECATED_UPawnAction_BlueprintBase::Tick(float DeltaTime)
+void UDEPRECATED_PawnAction_BlueprintBase::Tick(float DeltaTime)
 {
 	// no need to call super implementation
 	ActionTick(GetPawn(), DeltaTime);
 }
 
-bool UDEPRECATED_UPawnAction_BlueprintBase::Start()
+bool UDEPRECATED_PawnAction_BlueprintBase::Start()
 {
 	const bool bHasBeenEverStarted = HasBeenStarted();
 	const bool bSuperResult = Super::Start();
@@ -33,7 +33,7 @@ bool UDEPRECATED_UPawnAction_BlueprintBase::Start()
 	return bSuperResult;
 }
 
-bool UDEPRECATED_UPawnAction_BlueprintBase::Pause(const UDEPRECATED_UPawnAction* PausedBy)
+bool UDEPRECATED_PawnAction_BlueprintBase::Pause(const UDEPRECATED_PawnAction* PausedBy)
 {
 	const bool bResult = Super::Pause(PausedBy);
 	if (bResult)
@@ -43,7 +43,7 @@ bool UDEPRECATED_UPawnAction_BlueprintBase::Pause(const UDEPRECATED_UPawnAction*
 	return bResult;
 }
 
-bool UDEPRECATED_UPawnAction_BlueprintBase::Resume()
+bool UDEPRECATED_PawnAction_BlueprintBase::Resume()
 {
 	const bool bResult = Super::Resume();
 	if (bResult)
@@ -54,7 +54,7 @@ bool UDEPRECATED_UPawnAction_BlueprintBase::Resume()
 	return bResult;
 }
 
-void UDEPRECATED_UPawnAction_BlueprintBase::OnFinished(EPawnActionResult::Type WithResult)
+void UDEPRECATED_PawnAction_BlueprintBase::OnFinished(EPawnActionResult::Type WithResult)
 {
 	ActionFinished(GetPawn(), WithResult);
 }
diff --git a/Engine/Source/Runtime/AIModule/Private/Actions/PawnAction_Move.cpp b/Engine/Source/Runtime/AIModule/Private/Actions/PawnAction_Move.cpp
--- a/Engine/Source/Runtime/AIModule/Private/Actions/PawnAction_Move.cpp
+++ b/Engine/Source/Runtime/AIModule/Private/Actions/PawnAction_Move.cpp
@@ -11,7 +11,7 @@
 
 PRAGMA_DISABLE_DEPRECATION_WARNINGS
 
-UDEPRECATED_UPawnAction_Move::UDEPRECATED_UPawnAction_Move(const FObjectInitializer& ObjectInitializer)
+UDEPRECATED_PawnAction_Move::UDEPRECATED_PawnAction_Move(const FObjectInitializer& ObjectInitializer)
 	: Super(ObjectInitializer)
 	, GoalLocation(FAISystem::InvalidLocation)
 	, AcceptableRadius(30.f)
@@ -28,7 +28,7 @@ UDEPRECATED_UPawnAction_Move::UDEPRECATED_UPawnAction_Move(const FObjectInitiali
 	bAlwaysNotifyOnFinished = true;
 }
 
-void UDEPRECATED_UPawnAction_Move::BeginDestroy()
+void UDEPRECATED_PawnAction_Move::BeginDestroy()
 {
 	ClearTimers();
 	ClearPath();
@@ -36,14 +36,14 @@ void UDEPRECATED_UPawnAction_Move::BeginDestroy()
 	Super::BeginDestroy();
 }
 
-UDEPRECATED_UPawnAction_Move* UDEPRECATED_UPawnAction_Move::CreateAction(UWorld& World, AActor* GoalActor, EPawnActionMoveMode::Type Mode)
+UDEPRECATED_PawnAction_Move* UDEPRECATED_PawnAction_Move::CreateAction(UWorld& World, AActor* GoalActor, EPawnActionMoveMode::Type Mode)
 {
 	if (GoalActor == NULL)
 	{
 		return NULL;
 	}
 
-	UDEPRECATED_UPawnAction_Move* Action = UDEPRECATED_UPawnAction::CreateActionInstance<UDEPRECATED_UPawnAction_Move>(World);
+	UDEPRECATED_PawnAction_Move* Action = UDEPRECATED_PawnAction::CreateActionInstance<UDEPRECATED_PawnAction_Move>(World);
 	if (Action)
 	{
 		Action->GoalActor = GoalActor;
@@ -53,14 +53,14 @@ UDEPRECATED_UPawnAction_Move* UDEPRECATED_UPawnAction_Move::CreateAction(UWorld&
 	return Action;
 }
 
-UDEPRECATED_UPawnAction_Move* UDEPRECATED_UPawnAction_Move::CreateAction(UWorld& World, const FVector& GoalLocation, EPawnActionMoveMode::Type Mode)
+UDEPRECATED_PawnAction_Move* UDEPRECATED_PawnAction_Move::CreateAction(UWorld& World, const FVector& GoalLocation, EPawnActionMoveMode::Type Mode)
 {
 	if (FAISystem::IsValidLocation(GoalLocation) == false)
 	{
 		return NULL;
 	}
 
-	UDEPRECATED_UPawnAction_Move* Action = UDEPRECATED_UPawnAction::CreateActionInstance<UDEPRECATED_UPawnAction_Move>(World);
+	UDEPRECATED_PawnAction_Move* Action = UDEPRECATED_PawnAction::CreateActionInstance<UDEPRECATED_PawnAction_Move>(World);
 	if (Action)
 	{
 		Action->GoalLocation = GoalLocation;
@@ -70,7 +70,7 @@ UDEPRECATED_UPawnAction_Move* UDEPRECATED_UPawnAction_Move::CreateAction(UWorld&
 	return Action;
 }
 
-bool UDEPRECATED_UPawnAction_Move::Start()
+bool UDEPRECATED_PawnAction_Move::Start()
 {
 	bool bResult = Super::Start();
 	if (bResult)
@@ -81,7 +81,7 @@ bool UDEPRECATED_UPawnAction_Move::Start()
 	return bResult;
 }
 
-EPathFollowingRequestResult::Type UDEPRECATED_UPawnAction_Move::RequestMove(AAIController& Controller)
+EPathFollowingRequestResult::Type UDEPRECATED_PawnAction_Move::RequestMove(AAIController& Controller)
 {
 	EPathFollowingRequestResult::Type RequestResult = EPathFollowingRequestResult::Failed;
 	
@@ -117,13 +117,13 @@ EPathFollowingRequestResult::Type UDEPRECATED_UPawnAction_Move::RequestMove(AAIC
 	}
 	else
 	{
-		UE_VLOG(&Controller, LogPawnAction, Warning, TEXT("UDEPRECATED_UPawnAction_Move::Start: no valid move goal set"));
+		UE_VLOG(&Controller, LogPawnAction, Warning, TEXT("UDEPRECATED_PawnAction_Move::Start: no valid move goal set"));
 	}
 
 	return RequestResult;
 }
 
-bool UDEPRECATED_UPawnAction_Move::PerformMoveAction()
+bool UDEPRECATED_PawnAction_Move::PerformMoveAction()
 {
 	AAIController* MyController = Cast<AAIController>(GetController());
 	if (MyController == NULL)
@@ -134,7 +134,7 @@ bool UDEPRECATED_UPawnAction_Move::PerformMoveAction()
 	if (bUsePathfinding && MyController->ShouldPostponePathUpdates())
 	{
 		UE_VLOG(MyController, LogPawnAction, Log, TEXT("Can't path right now, waiting..."));
-		MyController->GetWorldTimerManager().SetTimer(TimerHandle_DeferredPerformMoveAction, this, &UDEPRECATED_UPawnAction_Move::DeferredPerformMoveAction, 0.1f);
+		MyController->GetWorldTimerManager().SetTimer(TimerHandle_DeferredPerformMoveAction, this, &UDEPRECATED_PawnAction_Move::DeferredPerformMoveAction, 0.1f);
 		return true;
 	}
 
@@ -161,7 +161,7 @@ bool UDEPRECATED_UPawnAction_Move::PerformMoveAction()
 	return bResult;
 }
 
-void UDEPRECATED_UPawnAction_Move::DeferredPerformMoveAction()
+void UDEPRECATED_PawnAction_Move::DeferredPerformMoveAction()
 {
 	const bool bResult = PerformMoveAction();
 	if (!bResult)
@@ -170,7 +170,7 @@ void UDEPRECATED_UPawnAction_Move::DeferredPerformMoveAction()
 	}
 }
 
-bool UDEPRECATED_UPawnAction_Move::Pause(const UDEPRECATED_UPawnAction* PausedBy)
+bool UDEPRECATED_PawnAction_Move::Pause(const UDEPRECATED_PawnAction* PausedBy)
 {
 	bool bResult = Super::Pause(PausedBy);
 	if (bResult)
@@ -184,7 +184,7 @@ bool UDEPRECATED_UPawnAction_Move::Pause(const UDEPRECATED_UPawnAction* PausedBy
 	return bResult;
 }
 
-bool UDEPRECATED_UPawnAction_Move::Resume()
+bool UDEPRECATED_PawnAction_Move::Resume()
 {
 	if (GoalActor != NULL && GoalActor->IsPendingKillPending())
 	{
@@ -215,7 +215,7 @@ bool UDEPRECATED_UPawnAction_Move::Resume()
 	return bResult;
 }
 
-EPawnActionAbortState::Type UDEPRECATED_UPawnAction_Move::PerformAbort(EAIForceParam::Type ShouldForce)
+EPawnActionAbortState::Type UDEPRECATED_PawnAction_Move::PerformAbort(EAIForceParam::Type ShouldForce)
 {
 	ClearTimers();
 	ClearPath();
@@ -230,7 +230,7 @@ EPawnActionAbortState::Type UDEPRECATED_UPawnAction_Move::PerformAbort(EAIForceP
 	return Super::PerformAbort(ShouldForce);
 }
 
-void UDEPRECATED_UPawnAction_Move::HandleAIMessage(UBrainComponent*, const FAIMessage& Message)
+void UDEPRECATED_PawnAction_Move::HandleAIMessage(UBrainComponent*, const FAIMessage& Message)
 {
 	if (Message.MessageName == UBrainComponent::AIMessage_MoveFinished && Message.HasFlag(FPathFollowingResultFlags::NewRequest))
 	{
@@ -244,7 +244,7 @@ void UDEPRECATED_UPawnAction_Move::HandleAIMessage(UBrainComponent*, const FAIMe
 	Finish(bFail ? EPawnActionResult::Failed : EPawnActionResult::Success);
 }
 
-void UDEPRECATED_UPawnAction_Move::OnFinished(EPawnActionResult::Type WithResult)
+void UDEPRECATED_PawnAction_Move::OnFinished(EPawnActionResult::Type WithResult)
 {
 	ClearTimers();
 	ClearPath();
@@ -252,7 +252,7 @@ void UDEPRECATED_UPawnAction_Move::OnFinished(EPawnActionResult::Type WithResult
 	Super::OnFinished(WithResult);
 }
 
-void UDEPRECATED_UPawnAction_Move::ClearPath()
+void UDEPRECATED_PawnAction_Move::ClearPath()
 {
 	ClearPendingRepath();
 	if (Path.IsValid())
@@ -262,20 +262,20 @@ void UDEPRECATED_UPawnAction_Move::ClearPath()
 	}
 }
 
-void UDEPRECATED_UPawnAction_Move::SetPath(FNavPathSharedRef InPath)
+void UDEPRECATED_PawnAction_Move::SetPath(FNavPathSharedRef InPath)
 {
 	if (InPath != Path)
 	{
 		ClearPath();
 		Path = InPath;
-		PathObserverDelegateHandle = Path->AddObserver(FNavigationPath::FPathObserverDelegate::FDelegate::CreateUObject(this, &UDEPRECATED_UPawnAction_Move::OnPathUpdated));
+		PathObserverDelegateHandle = Path->AddObserver(FNavigationPath::FPathObserverDelegate::FDelegate::CreateUObject(this, &UDEPRECATED_PawnAction_Move::OnPathUpdated));
 
 		// skip auto updates, it will be handled manually to include controller's ShouldPostponePathUpdates()
 		Path->EnableRecalculationOnInvalidation(false);
 	}
 }
 
-void UDEPRECATED_UPawnAction_Move::OnPathUpdated(FNavigationPath* UpdatedPath, ENavPathEvent::Type Event)
+void UDEPRECATED_PawnAction_Move::OnPathUpdated(FNavigationPath* UpdatedPath, ENavPathEvent::Type Event)
 {
 	const AController* MyOwner = GetController();
 	if (MyOwner == NULL)
@@ -323,7 +323,7 @@ void UDEPRECATED_UPawnAction_Move::OnPathUpdated(FNavigationPath* UpdatedPath, E
 	}
 }
 
-void UDEPRECATED_UPawnAction_Move::TryToRepath()
+void UDEPRECATED_PawnAction_Move::TryToRepath()
 {
 	if (Path.IsValid())
 	{
@@ -338,12 +338,12 @@ void UDEPRECATED_UPawnAction_Move::TryToRepath()
 		}
 		else if (GetWorld())
 		{
-			GetWorld()->GetTimerManager().SetTimer(TimerHandle_TryToRepath, this, &UDEPRECATED_UPawnAction_Move::TryToRepath, 0.25f);
+			GetWorld()->GetTimerManager().SetTimer(TimerHandle_TryToRepath, this, &UDEPRECATED_PawnAction_Move::TryToRepath, 0.25f);
 		}
 	}
 }
 
-void UDEPRECATED_UPawnAction_Move::ClearPendingRepath()
+void UDEPRECATED_PawnAction_Move::ClearPendingRepath()
 {
 	if (TimerHandle_TryToRepath.IsValid())
 	{
@@ -356,7 +356,7 @@ void UDEPRECATED_UPawnAction_Move::ClearPendingRepath()
 	}
 }
 
-void UDEPRECATED_UPawnAction_Move::ClearTimers()
+void UDEPRECATED_PawnAction_Move::ClearTimers()
 {
 	UWorld* World = GetWorld();
 	if (World)
@@ -369,7 +369,7 @@ void UDEPRECATED_UPawnAction_Move::ClearTimers()
 	}
 }
 
-bool UDEPRECATED_UPawnAction_Move::CheckAlreadyAtGoal(AAIController& Controller, const FVector& TestLocation, float Radius)
+bool UDEPRECATED_PawnAction_Move::CheckAlreadyAtGoal(AAIController& Controller, const FVector& TestLocation, float Radius)
 {
 	const bool bAlreadyAtGoal = Controller.GetPathFollowingComponent()->HasReached(TestLocation, EPathFollowingReachMode::OverlapAgentAndGoal, Radius);
 	if (bAlreadyAtGoal)
@@ -381,7 +381,7 @@ bool UDEPRECATED_UPawnAction_Move::CheckAlreadyAtGoal(AAIController& Controller,
 	return bAlreadyAtGoal;
 }
 
-bool UDEPRECATED_UPawnAction_Move::CheckAlreadyAtGoal(AAIController& Controller, const AActor& TestGoal, float Radius)
+bool UDEPRECATED_PawnAction_Move::CheckAlreadyAtGoal(AAIController& Controller, const AActor& TestGoal, float Radius)
 {
 	const bool bAlreadyAtGoal = Controller.GetPathFollowingComponent()->HasReached(TestGoal, EPathFollowingReachMode::OverlapAgentAndGoal, Radius);
 	if (bAlreadyAtGoal)
@@ -393,7 +393,7 @@ bool UDEPRECATED_UPawnAction_Move::CheckAlreadyAtGoal(AAIController& Controller,
 	return bAlreadyAtGoal;
 }
 
-bool UDEPRECATED_UPawnAction_Move::IsPartialPathAllowed() const
+bool UDEPRECATED_PawnAction_Move::IsPartialPathAllowed() const
 {
 	return bAllowPartialPath;
 }
diff --git a/Engine/Source/Runtime/AIModule/Private/Actions/PawnAction_Repeat.cpp b/Engine/Source/Runtime/AIModule/Private/Actions/PawnAction_Repeat.cpp
--- a/Engine/Source/Runtime/AIModule/Private/Actions/PawnAction_Repeat.cpp
+++ b/Engine/Source/Runtime/AIModule/Private/Actions/PawnAction_Repeat.cpp
@@ -9,21 +9,21 @@
 
 PRAGMA_DISABLE_DEPRECATION_WARNINGS
 
-UDEPRECATED_UPawnAction_Repeat::UDEPRECATED_UPawnAction_Repeat(const FObjectInitializer& ObjectInitializer)
+UDEPRECATED_PawnAction_Repeat::UDEPRECATED_PawnAction_Repeat(const FObjectInitializer& ObjectInitializer)
 	: Super(ObjectInitializer)
 	, SubActionTriggeringPolicy(EPawnSubActionTriggeringPolicy::CopyBeforeTriggering)
 {
 	ChildFailureHandlingMode = EPawnActionFailHandling::IgnoreFailure;
 }
 
-UDEPRECATED_UPawnAction_Repeat* UDEPRECATED_UPawnAction_Repeat::CreateAction(UWorld& World, UDEPRECATED_UPawnAction* ActionToRepeat, int32 NumberOfRepeats, EPawnSubActionTriggeringPolicy::Type InSubActionTriggeringPolicy)
+UDEPRECATED_PawnAction_Repeat* UDEPRECATED_PawnAction_Repeat::CreateAction(UWorld& World, UDEPRECATED_PawnAction* ActionToRepeat, int32 NumberOfRepeats, EPawnSubActionTriggeringPolicy::Type InSubActionTriggeringPolicy)
 {
-	if (ActionToRepeat == NULL || !(NumberOfRepeats > 0 || NumberOfRepeats == UDEPRECATED_UPawnAction_Repeat::LoopForever))
+	if (ActionToRepeat == NULL || !(NumberOfRepeats > 0 || NumberOfRepeats == UDEPRECATED_PawnAction_Repeat::LoopForever))
 	{
 		return NULL;
 	}
 
-	UDEPRECATED_UPawnAction_Repeat* Action = UDEPRECATED_UPawnAction::CreateActionInstance<UDEPRECATED_UPawnAction_Repeat>(World);
+	UDEPRECATED_PawnAction_Repeat* Action = UDEPRECATED_PawnAction::CreateActionInstance<UDEPRECATED_PawnAction_Repeat>(World);
 	if (Action)
 	{
 		Action->ActionToRepeat_DEPRECATED = ActionToRepeat;
@@ -36,7 +36,7 @@ UDEPRECATED_UPawnAction_Repeat* UDEPRECATED_UPawnAction_Repeat::CreateAction(UWo
 	return Action;
 }
 
-bool UDEPRECATED_UPawnAction_Repeat::Start()
+bool UDEPRECATED_PawnAction_Repeat::Start()
 {
 	bool bResult = Super::Start();
 	
@@ -50,7 +50,7 @@ bool UDEPRECATED_UPawnAction_Repeat::Start()
 	return bResult;
 }
 
-bool UDEPRECATED_UPawnAction_Repeat::Resume()
+bool UDEPRECATED_PawnAction_Repeat::Resume()
 {
 	bool bResult = Super::Resume();
 
@@ -62,7 +62,7 @@ bool UDEPRECATED_UPawnAction_Repeat::Resume()
 	return bResult;
 }
 
-void UDEPRECATED_UPawnAction_Repeat::OnChildFinished(UDEPRECATED_UPawnAction& Action, EPawnActionResult::Type WithResult)
+void UDEPRECATED_PawnAction_Repeat::OnChildFinished(UDEPRECATED_PawnAction& Action, EPawnActionResult::Type WithResult)
 {
 	if (RecentActionCopy_DEPRECATED == &Action)
 	{
@@ -79,7 +79,7 @@ void UDEPRECATED_UPawnAction_Repeat::OnChildFinished(UDEPRECATED_UPawnAction& Ac
 	Super::OnChildFinished(Action, WithResult);
 }
 
-bool UDEPRECATED_UPawnAction_Repeat::PushSubAction()
+bool UDEPRECATED_PawnAction_Repeat::PushSubAction()
 {
 	if (ActionToRepeat_DEPRECATED == NULL)
 	{
@@ -97,8 +97,8 @@ bool UDEPRECATED_UPawnAction_Repeat::PushSubAction()
 		--RepeatsLeft;
 	}
 
-	UDEPRECATED_UPawnAction* ActionCopy = SubActionTriggeringPolicy == EPawnSubActionTriggeringPolicy::CopyBeforeTriggering 
-		? Cast<UDEPRECATED_UPawnAction>(StaticDuplicateObject(ActionToRepeat_DEPRECATED, this))
+	UDEPRECATED_PawnAction* ActionCopy = SubActionTriggeringPolicy == EPawnSubActionTriggeringPolicy::CopyBeforeTriggering 
+		? Cast<UDEPRECATED_PawnAction>(StaticDuplicateObject(ActionToRepeat_DEPRECATED, this))
 		: ToRawPtr(ActionToRepeat_DEPRECATED);
 
 	UE_VLOG(GetPawn(), LogPawnAction, Log, TEXT("%s> pushing repeted action %s %s, repeats left: %d")
diff --git a/Engine/Source/Runtime/AIModule/Private/Actions/PawnAction_Sequence.cpp b/Engine/Source/Runtime/AIModule/Private/Actions/PawnAction_Sequence.cpp
--- a/Engine/Source/Runtime/AIModule/Private/Actions/PawnAction_Sequence.cpp
+++ b/Engine/Source/Runtime/AIModule/Private/Actions/PawnAction_Sequence.cpp
@@ -9,13 +9,13 @@
 
 PRAGMA_DISABLE_DEPRECATION_WARNINGS
 
-UDEPRECATED_UPawnAction_Sequence::UDEPRECATED_UPawnAction_Sequence(const FObjectInitializer& ObjectInitializer)
+UDEPRECATED_PawnAction_Sequence::UDEPRECATED_PawnAction_Sequence(const FObjectInitializer& ObjectInitializer)
 	: Super(ObjectInitializer)
 	, SubActionTriggeringPolicy(EPawnSubActionTriggeringPolicy::CopyBeforeTriggering)
 {
 }
 
-UDEPRECATED_UPawnAction_Sequence* UDEPRECATED_UPawnAction_Sequence::CreateAction(UWorld& World, TArray<UDEPRECATED_UPawnAction*>& ActionSequence, EPawnSubActionTriggeringPolicy::Type InSubActionTriggeringPolicy)
+UDEPRECATED_PawnAction_Sequence* UDEPRECATED_PawnAction_Sequence::CreateAction(UWorld& World, TArray<UDEPRECATED_PawnAction*>& ActionSequence, EPawnSubActionTriggeringPolicy::Type InSubActionTriggeringPolicy)
 {
 	ActionSequence.Remove(NULL);
 	if (ActionSequence.Num() <= 0)
@@ -23,12 +23,12 @@ UDEPRECATED_UPawnAction_Sequence* UDEPRECATED_UPawnAction_Sequence::CreateAction
 		return NULL;
 	}
 
-	UDEPRECATED_UPawnAction_Sequence* Action = UDEPRECATED_UPawnAction::CreateActionInstance<UDEPRECATED_UPawnAction_Sequence>(World);
+	UDEPRECATED_PawnAction_Sequence* Action = UDEPRECATED_PawnAction::CreateActionInstance<UDEPRECATED_PawnAction_Sequence>(World);
 	if (Action)
 	{
 		Action->ActionSequence_DEPRECATED = ActionSequence;
 
-		for (const UDEPRECATED_UPawnAction* SubAction : ActionSequence)
+		for (const UDEPRECATED_PawnAction* SubAction : ActionSequence)
 		{
 			if (SubAction && SubAction->ShouldPauseMovement())
 			{
@@ -43,7 +43,7 @@ UDEPRECATED_UPawnAction_Sequence* UDEPRECATED_UPawnAction_Sequence::CreateAction
 	return Action;
 }
 
-bool UDEPRECATED_UPawnAction_Sequence::Start()
+bool UDEPRECATED_PawnAction_Sequence::Start()
 {
 	bool bResult = Super::Start();
 
@@ -61,7 +61,7 @@ bool UDEPRECATED_UPawnAction_Sequence::Start()
 	return bResult;
 }
 
-bool UDEPRECATED_UPawnAction_Sequence::Resume()
+bool UDEPRECATED_PawnAction_Sequence::Resume()
 {
 	bool bResult = Super::Resume();
 
@@ -73,7 +73,7 @@ bool UDEPRECATED_UPawnAction_Sequence::Resume()
 	return bResult;
 }
 
-void UDEPRECATED_UPawnAction_Sequence::OnChildFinished(UDEPRECATED_UPawnAction& Action, EPawnActionResult::Type WithResult)
+void UDEPRECATED_PawnAction_Sequence::OnChildFinished(UDEPRECATED_PawnAction& Action, EPawnActionResult::Type WithResult)
 {
 	if (RecentActionCopy_DEPRECATED == &Action)
 	{
@@ -93,7 +93,7 @@ void UDEPRECATED_UPawnAction_Sequence::OnChildFinished(UDEPRECATED_UPawnAction&
 	Super::OnChildFinished(Action, WithResult);
 }
 
-bool UDEPRECATED_UPawnAction_Sequence::PushNextActionCopy()
+bool UDEPRECATED_PawnAction_Sequence::PushNextActionCopy()
 {
 	if (CurrentActionIndex >= uint32(ActionSequence_DEPRECATED.Num()))
 	{
@@ -101,8 +101,8 @@ bool UDEPRECATED_UPawnAction_Sequence::PushNextActionCopy()
 		return true;
 	}
 
-	UDEPRECATED_UPawnAction* ActionCopy = SubActionTriggeringPolicy == EPawnSubActionTriggeringPolicy::CopyBeforeTriggering
-		? Cast<UDEPRECATED_UPawnAction>(StaticDuplicateObject(ActionSequence_DEPRECATED[CurrentActionIndex], this))
+	UDEPRECATED_PawnAction* ActionCopy = SubActionTriggeringPolicy == EPawnSubActionTriggeringPolicy::CopyBeforeTriggering
+		? Cast<UDEPRECATED_PawnAction>(StaticDuplicateObject(ActionSequence_DEPRECATED[CurrentActionIndex], this))
 		: ToRawPtr(ActionSequence_DEPRECATED[CurrentActionIndex]);
 
 	UE_VLOG(GetPawn(), LogPawnAction, Log, TEXT("%s> pushing action %s")
diff --git a/Engine/Source/Runtime/AIModule/Private/Actions/PawnAction_Wait.cpp b/Engine/Source/Runtime/AIModule/Private/Actions/PawnAction_Wait.cpp
--- a/Engine/Source/Runtime/AIModule/Private/Actions/PawnAction_Wait.cpp
+++ b/Engine/Source/Runtime/AIModule/Private/Actions/PawnAction_Wait.cpp
@@ -8,7 +8,7 @@
 
 PRAGMA_DISABLE_DEPRECATION_WARNINGS
 
-UDEPRECATED_UPawnAction_Wait::UDEPRECATED_UPawnAction_Wait(const FObjectInitializer& ObjectInitializer)
+UDEPRECATED_PawnAction_Wait::UDEPRECATED_PawnAction_Wait(const FObjectInitializer& ObjectInitializer)
 	: Super(ObjectInitializer)
 	, TimeToWait(0.f)
 	, FinishTimeStamp(0.f)
@@ -16,9 +16,9 @@ UDEPRECATED_UPawnAction_Wait::UDEPRECATED_UPawnAction_Wait(const FObjectInitiali
 
 }
 
-UDEPRECATED_UPawnAction_Wait* UDEPRECATED_UPawnAction_Wait::CreateAction(UWorld& World, float InTimeToWait)
+UDEPRECATED_PawnAction_Wait* UDEPRECATED_PawnAction_Wait::CreateAction(UWorld& World, float InTimeToWait)
 {
-	UDEPRECATED_UPawnAction_Wait* Action = UDEPRECATED_UPawnAction::CreateActionInstance<UDEPRECATED_UPawnAction_Wait>(World);
+	UDEPRECATED_PawnAction_Wait* Action = UDEPRECATED_PawnAction::CreateActionInstance<UDEPRECATED_PawnAction_Wait>(World);
 	
 	if (Action != NULL)
 	{
@@ -28,13 +28,13 @@ UDEPRECATED_UPawnAction_Wait* UDEPRECATED_UPawnAction_Wait::CreateAction(UWorld&
 	return Action;
 }
 
-bool UDEPRECATED_UPawnAction_Wait::Start()
+bool UDEPRECATED_PawnAction_Wait::Start()
 {
 	if (Super::Start())
 	{
 		if (TimeToWait >= 0)
 		{
-			GetWorld()->GetTimerManager().SetTimer(TimerHandle, this, &UDEPRECATED_UPawnAction_Wait::TimerDone, TimeToWait);
+			GetWorld()->GetTimerManager().SetTimer(TimerHandle, this, &UDEPRECATED_PawnAction_Wait::TimerDone, TimeToWait);
 		}
 		// else hang in there for ever!
 
@@ -44,25 +44,25 @@ bool UDEPRECATED_UPawnAction_Wait::Start()
 	return false;
 }
 
-bool UDEPRECATED_UPawnAction_Wait::Pause(const UDEPRECATED_UPawnAction* PausedBy)
+bool UDEPRECATED_PawnAction_Wait::Pause(const UDEPRECATED_PawnAction* PausedBy)
 {
 	GetWorld()->GetTimerManager().PauseTimer(TimerHandle);
 	return true;
 }
 
-bool UDEPRECATED_UPawnAction_Wait::Resume()
+bool UDEPRECATED_PawnAction_Wait::Resume()
 {
 	GetWorld()->GetTimerManager().UnPauseTimer(TimerHandle);
 	return true;
 }
 
-EPawnActionAbortState::Type UDEPRECATED_UPawnAction_Wait::PerformAbort(EAIForceParam::Type ShouldForce)
+EPawnActionAbortState::Type UDEPRECATED_PawnAction_Wait::PerformAbort(EAIForceParam::Type ShouldForce)
 {
 	GetWorld()->GetTimerManager().ClearTimer(TimerHandle);
 	return EPawnActionAbortState::AbortDone;
 }
 
-void UDEPRECATED_UPawnAction_Wait::TimerDone()
+void UDEPRECATED_PawnAction_Wait::TimerDone()
 {
 	Finish(EPawnActionResult::Success);
 }
diff --git a/Engine/Source/Runtime/AIModule/Private/Actions/PawnActionsComponent.cpp b/Engine/Source/Runtime/AIModule/Private/Actions/PawnActionsComponent.cpp
--- a/Engine/Source/Runtime/AIModule/Private/Actions/PawnActionsComponent.cpp
+++ b/Engine/Source/Runtime/AIModule/Private/Actions/PawnActionsComponent.cpp
@@ -33,7 +33,7 @@ namespace
 		return Enum->GetNameStringByValue(Value);
 	}
 
-	FString GetActionSignature(UDEPRECATED_UPawnAction* Action)
+	FString GetActionSignature(UDEPRECATED_PawnAction* Action)
 	{
 		if (Action == NULL)
 		{
@@ -62,7 +62,7 @@ namespace
 	};
 }
 
-FPawnActionEvent::FPawnActionEvent(UDEPRECATED_UPawnAction& InAction, EPawnActionEventType::Type InEventType, uint32 InIndex)
+FPawnActionEvent::FPawnActionEvent(UDEPRECATED_PawnAction& InAction, EPawnActionEventType::Type InEventType, uint32 InIndex)
 	: Action(&InAction), EventType(InEventType), Index(InIndex)
 {
 	Priority = InAction.GetPriority();
@@ -88,7 +88,7 @@ void FPawnActionStack::Resume()
 	}
 }
 
-void FPawnActionStack::PushAction(UDEPRECATED_UPawnAction& NewTopAction)
+void FPawnActionStack::PushAction(UDEPRECATED_PawnAction& NewTopAction)
 {
 	if (TopAction != NULL)
 	{
@@ -105,10 +105,10 @@ void FPawnActionStack::PushAction(UDEPRECATED_UPawnAction& NewTopAction)
 	NewTopAction.OnPushed();
 }
 
-void FPawnActionStack::PopAction(UDEPRECATED_UPawnAction& ActionToPop)
+void FPawnActionStack::PopAction(UDEPRECATED_PawnAction& ActionToPop)
 {
 	// first check if it's there
-	UDEPRECATED_UPawnAction* CutPoint = TopAction;
+	UDEPRECATED_PawnAction* CutPoint = TopAction;
 	while (CutPoint != NULL && CutPoint != &ActionToPop)
 	{
 		CutPoint = CutPoint->ParentAction_DEPRECATED;
@@ -116,14 +116,14 @@ void FPawnActionStack::PopAction(UDEPRECATED_UPawnAction& ActionToPop)
 
 	if (CutPoint == &ActionToPop)
 	{
-		UDEPRECATED_UPawnAction* ActionBeingRemoved = TopAction;
+		UDEPRECATED_PawnAction* ActionBeingRemoved = TopAction;
 		// note StopAction can be null
-		UDEPRECATED_UPawnAction* StopAction = ActionToPop.ParentAction_DEPRECATED;
+		UDEPRECATED_PawnAction* StopAction = ActionToPop.ParentAction_DEPRECATED;
 
 		while (ActionBeingRemoved != StopAction && ActionBeingRemoved != nullptr)
 		{
 			checkSlow(ActionBeingRemoved);
-			UDEPRECATED_UPawnAction* NextAction = ActionBeingRemoved->ParentAction_DEPRECATED;
+			UDEPRECATED_PawnAction* NextAction = ActionBeingRemoved->ParentAction_DEPRECATED;
 
 			if (ActionBeingRemoved->IsBeingAborted() == false && ActionBeingRemoved->IsFinished() == false)
 			{
@@ -145,7 +145,7 @@ void FPawnActionStack::PopAction(UDEPRECATED_UPawnAction& ActionToPop)
 int32 FPawnActionStack::GetStackSize() const
 {
 	int32 Size = 0;
-	const UDEPRECATED_UPawnAction* TempAction = TopAction;
+	const UDEPRECATED_PawnAction* TempAction = TopAction;
 	while (TempAction != nullptr)
 	{
 		TempAction = TempAction->GetParentAction();
@@ -155,10 +155,10 @@ int32 FPawnActionStack::GetStackSize() const
 }
 
 //----------------------------------------------------------------------//
-// UDEPRECATED_UPawnActionsComponent
+// UDEPRECATED_PawnActionsComponent
 //----------------------------------------------------------------------//
 
-UDEPRECATED_UPawnActionsComponent::UDEPRECATED_UPawnActionsComponent(const FObjectInitializer& ObjectInitializer)
+UDEPRECATED_PawnActionsComponent::UDEPRECATED_PawnActionsComponent(const FObjectInitializer& ObjectInitializer)
 	: Super(ObjectInitializer)
 {
 	PrimaryComponentTick.bCanEverTick = true;
@@ -172,14 +172,14 @@ UDEPRECATED_UPawnActionsComponent::UDEPRECATED_UPawnActionsComponent(const FObje
 	ActionStacks.AddZeroed(EAIRequestPriority::MAX);
 }
 
-void UDEPRECATED_UPawnActionsComponent::OnUnregister()
+void UDEPRECATED_PawnActionsComponent::OnUnregister()
 {
 	if ((ControlledPawn != nullptr) && !ControlledPawn->IsPendingKillPending())
 	{
 		// call for every regular priority 
 		for (int32 PriorityIndex = 0; PriorityIndex < EAIRequestPriority::MAX; ++PriorityIndex)
 		{
-			UDEPRECATED_UPawnAction* Action = ActionStacks[PriorityIndex].GetTop();
+			UDEPRECATED_PawnAction* Action = ActionStacks[PriorityIndex].GetTop();
 			while (Action)
 			{
 				Action->Abort(EAIForceParam::Force);
@@ -191,7 +191,7 @@ void UDEPRECATED_UPawnActionsComponent::OnUnregister()
 	Super::OnUnregister();
 }
 
-void UDEPRECATED_UPawnActionsComponent::TickComponent(float DeltaTime, enum ELevelTick TickType, FActorComponentTickFunction *ThisTickFunction)
+void UDEPRECATED_PawnActionsComponent::TickComponent(float DeltaTime, enum ELevelTick TickType, FActorComponentTickFunction *ThisTickFunction)
 {
 	Super::TickComponent(DeltaTime, TickType, ThisTickFunction);
 
@@ -258,14 +258,14 @@ void UDEPRECATED_UPawnActionsComponent::TickComponent(float DeltaTime, enum ELev
 	}
 }
 
-bool UDEPRECATED_UPawnActionsComponent::HasActiveActionOfType(EAIRequestPriority::Type Priority, TSubclassOf<UDEPRECATED_UPawnAction> PawnActionClass) const
+bool UDEPRECATED_PawnActionsComponent::HasActiveActionOfType(EAIRequestPriority::Type Priority, TSubclassOf<UDEPRECATED_PawnAction> PawnActionClass) const
 {
-	TArray<UDEPRECATED_UPawnAction*> ActionsToTest;
+	TArray<UDEPRECATED_PawnAction*> ActionsToTest;
 	ActionsToTest.Add(GetActiveAction(Priority));
 
 	while (ActionsToTest.Num() > 0)
 	{
-		UDEPRECATED_UPawnAction* ActiveActionIter = ActionsToTest[0];
+		UDEPRECATED_PawnAction* ActiveActionIter = ActionsToTest[0];
 
 		if (ActiveActionIter)
 		{
@@ -275,7 +275,7 @@ bool UDEPRECATED_UPawnActionsComponent::HasActiveActionOfType(EAIRequestPriority
 			}	
 			else
 			{
-				UDEPRECATED_UPawnAction_Sequence* PawnActionSequence = Cast<UDEPRECATED_UPawnAction_Sequence>(ActiveActionIter);
+				UDEPRECATED_PawnAction_Sequence* PawnActionSequence = Cast<UDEPRECATED_PawnAction_Sequence>(ActiveActionIter);
 
 				if (PawnActionSequence)
 				{
@@ -294,13 +294,13 @@ bool UDEPRECATED_UPawnActionsComponent::HasActiveActionOfType(EAIRequestPriority
 	return false;
 }
 
-void UDEPRECATED_UPawnActionsComponent::UpdateCurrentAction()
+void UDEPRECATED_PawnActionsComponent::UpdateCurrentAction()
 {
 	UE_VLOG(ControlledPawn, LogPawnAction, Log, TEXT("Picking new current actions. Old CurrentAction_DEPRECATED %s")
 		, *GetActionSignature(CurrentAction_DEPRECATED));
 
 	// find the highest priority action available
-	UDEPRECATED_UPawnAction* NewCurrentAction = NULL;
+	UDEPRECATED_PawnAction* NewCurrentAction = NULL;
 	int32 Priority = EAIRequestPriority::MAX - 1;
 	do 
 	{
@@ -356,7 +356,7 @@ void UDEPRECATED_UPawnActionsComponent::UpdateCurrentAction()
 	}
 }
 
-void UDEPRECATED_UPawnActionsComponent::UpdateAILogicLock()
+void UDEPRECATED_PawnActionsComponent::UpdateAILogicLock()
 {
 	if (ControlledPawn && ControlledPawn->GetController())
 	{
@@ -390,7 +390,7 @@ void UDEPRECATED_UPawnActionsComponent::UpdateAILogicLock()
 	}
 }
 
-EPawnActionAbortState::Type UDEPRECATED_UPawnActionsComponent::K2_AbortAction(UDEPRECATED_UPawnAction* ActionToAbort)
+EPawnActionAbortState::Type UDEPRECATED_PawnActionsComponent::K2_AbortAction(UDEPRECATED_PawnAction* ActionToAbort)
 {
 	if (ActionToAbort != NULL)
 	{
@@ -399,7 +399,7 @@ EPawnActionAbortState::Type UDEPRECATED_UPawnActionsComponent::K2_AbortAction(UD
 	return EPawnActionAbortState::NeverStarted;
 }
 
-EPawnActionAbortState::Type UDEPRECATED_UPawnActionsComponent::AbortAction(UDEPRECATED_UPawnAction& ActionToAbort)
+EPawnActionAbortState::Type UDEPRECATED_PawnActionsComponent::AbortAction(UDEPRECATED_PawnAction& ActionToAbort)
 {
 	const EPawnActionAbortState::Type AbortState = ActionToAbort.Abort(EAIForceParam::DoNotForce);
 	if (AbortState == EPawnActionAbortState::NeverStarted)
@@ -413,7 +413,7 @@ EPawnActionAbortState::Type UDEPRECATED_UPawnActionsComponent::AbortAction(UDEPR
 	return AbortState;
 }
 
-void UDEPRECATED_UPawnActionsComponent::RemoveEventsForAction(UDEPRECATED_UPawnAction& PawnAction)
+void UDEPRECATED_PawnActionsComponent::RemoveEventsForAction(UDEPRECATED_PawnAction& PawnAction)
 {
 	for (int32 ActionIndex = ActionEvents.Num() - 1; ActionIndex >= 0; --ActionIndex)
 	{
@@ -424,7 +424,7 @@ void UDEPRECATED_UPawnActionsComponent::RemoveEventsForAction(UDEPRECATED_UPawnA
 	}
 }
 
-EPawnActionAbortState::Type UDEPRECATED_UPawnActionsComponent::K2_ForceAbortAction(UDEPRECATED_UPawnAction* ActionToAbort)
+EPawnActionAbortState::Type UDEPRECATED_PawnActionsComponent::K2_ForceAbortAction(UDEPRECATED_PawnAction* ActionToAbort)
 {
 	if (ActionToAbort)
 	{
@@ -433,12 +433,12 @@ EPawnActionAbortState::Type UDEPRECATED_UPawnActionsComponent::K2_ForceAbortActi
 	return EPawnActionAbortState::NeverStarted;
 }
 
-EPawnActionAbortState::Type UDEPRECATED_UPawnActionsComponent::ForceAbortAction(UDEPRECATED_UPawnAction& ActionToAbort)
+EPawnActionAbortState::Type UDEPRECATED_PawnActionsComponent::ForceAbortAction(UDEPRECATED_PawnAction& ActionToAbort)
 {
 	return ActionToAbort.Abort(EAIForceParam::Force);
 }
 
-uint32 UDEPRECATED_UPawnActionsComponent::AbortActionsInstigatedBy(UObject* const Instigator, EAIRequestPriority::Type Priority)
+uint32 UDEPRECATED_PawnActionsComponent::AbortActionsInstigatedBy(UObject* const Instigator, EAIRequestPriority::Type Priority)
 {
 	uint32 AbortedActionsCount = 0;
 
@@ -452,7 +452,7 @@ uint32 UDEPRECATED_UPawnActionsComponent::AbortActionsInstigatedBy(UObject* cons
 	}
 	else
 	{
-		UDEPRECATED_UPawnAction* Action = ActionStacks[Priority].GetTop();
+		UDEPRECATED_PawnAction* Action = ActionStacks[Priority].GetTop();
 		while (Action)
 		{
 			if (Action->GetInstigator() == Instigator)
@@ -479,7 +479,7 @@ uint32 UDEPRECATED_UPawnActionsComponent::AbortActionsInstigatedBy(UObject* cons
 	return AbortedActionsCount;
 }
 
-bool UDEPRECATED_UPawnActionsComponent::K2_PushAction(UDEPRECATED_UPawnAction* NewAction, EAIRequestPriority::Type Priority, UObject* Instigator)
+bool UDEPRECATED_PawnActionsComponent::K2_PushAction(UDEPRECATED_PawnAction* NewAction, EAIRequestPriority::Type Priority, UObject* Instigator)
 {
 	if (NewAction)
 	{
@@ -488,7 +488,7 @@ bool UDEPRECATED_UPawnActionsComponent::K2_PushAction(UDEPRECATED_UPawnAction* N
 	return false;
 }
 
-bool UDEPRECATED_UPawnActionsComponent::PushAction(UDEPRECATED_UPawnAction& NewAction, EAIRequestPriority::Type Priority, UObject* Instigator)
+bool UDEPRECATED_PawnActionsComponent::PushAction(UDEPRECATED_PawnAction& NewAction, EAIRequestPriority::Type Priority, UObject* Instigator)
 {
 	if (NewAction.HasBeenStarted() == false || NewAction.IsFinished() == true)
 	{
@@ -501,7 +501,7 @@ bool UDEPRECATED_UPawnActionsComponent::PushAction(UDEPRECATED_UPawnAction& NewA
 	return false;
 }
 
-bool UDEPRECATED_UPawnActionsComponent::OnEvent(UDEPRECATED_UPawnAction& Action, EPawnActionEventType::Type Event)
+bool UDEPRECATED_PawnActionsComponent::OnEvent(UDEPRECATED_PawnAction& Action, EPawnActionEventType::Type Event)
 {
 	bool bResult = false;
 	const FPawnActionEvent ActionEvent(Action, Event, ActionEventIndex++);
@@ -533,7 +533,7 @@ bool UDEPRECATED_UPawnActionsComponent::OnEvent(UDEPRECATED_UPawnAction& Action,
 	return bResult;
 }
 
-void UDEPRECATED_UPawnActionsComponent::SetControlledPawn(APawn* NewPawn)
+void UDEPRECATED_PawnActionsComponent::SetControlledPawn(APawn* NewPawn)
 {
 	if (ControlledPawn != NULL && ControlledPawn != NewPawn)
 	{
@@ -545,7 +545,7 @@ void UDEPRECATED_UPawnActionsComponent::SetControlledPawn(APawn* NewPawn)
 	}
 }
 
-APawn* UDEPRECATED_UPawnActionsComponent::CacheControlledPawn()
+APawn* UDEPRECATED_PawnActionsComponent::CacheControlledPawn()
 {
 	if (ControlledPawn == NULL)
 	{
@@ -571,7 +571,7 @@ APawn* UDEPRECATED_UPawnActionsComponent::CacheControlledPawn()
 //----------------------------------------------------------------------//
 // blueprint interface
 //----------------------------------------------------------------------//
-bool UDEPRECATED_UPawnActionsComponent::K2_PerformAction(APawn* Pawn, UDEPRECATED_UPawnAction* Action, TEnumAsByte<EAIRequestPriority::Type> Priority)
+bool UDEPRECATED_PawnActionsComponent::K2_PerformAction(APawn* Pawn, UDEPRECATED_PawnAction* Action, TEnumAsByte<EAIRequestPriority::Type> Priority)
 {
 	if (Pawn && Action)
 	{
@@ -580,7 +580,7 @@ bool UDEPRECATED_UPawnActionsComponent::K2_PerformAction(APawn* Pawn, UDEPRECATE
 	return false;
 }
 
-bool UDEPRECATED_UPawnActionsComponent::PerformAction(APawn& Pawn, UDEPRECATED_UPawnAction& Action, TEnumAsByte<EAIRequestPriority::Type> Priority)
+bool UDEPRECATED_PawnActionsComponent::PerformAction(APawn& Pawn, UDEPRECATED_PawnAction& Action, TEnumAsByte<EAIRequestPriority::Type> Priority)
 {
 	bool bSuccess = false;
 
@@ -588,7 +588,7 @@ bool UDEPRECATED_UPawnActionsComponent::PerformAction(APawn& Pawn, UDEPRECATED_U
 
 	if (Pawn.GetController())
 	{
-		UDEPRECATED_UPawnActionsComponent* ActionComp = Pawn.GetController()->FindComponentByClass<UDEPRECATED_UPawnActionsComponent>();
+		UDEPRECATED_PawnActionsComponent* ActionComp = Pawn.GetController()->FindComponentByClass<UDEPRECATED_PawnActionsComponent>();
 		if (ActionComp)
 		{
 			ActionComp->PushAction(Action, Priority);
@@ -603,7 +603,7 @@ bool UDEPRECATED_UPawnActionsComponent::PerformAction(APawn& Pawn, UDEPRECATED_U
 // debug
 //----------------------------------------------------------------------//
 #if ENABLE_VISUAL_LOG
-void UDEPRECATED_UPawnActionsComponent::DescribeSelfToVisLog(FVisualLogEntry* Snapshot) const
+void UDEPRECATED_PawnActionsComponent::DescribeSelfToVisLog(FVisualLogEntry* Snapshot) const
 {
 	static const FString Category = TEXT("PawnActions");
 
@@ -614,7 +614,7 @@ void UDEPRECATED_UPawnActionsComponent::DescribeSelfToVisLog(FVisualLogEntry* Sn
 
 	for (int32 PriorityIndex = 0; PriorityIndex < ActionStacks.Num(); ++PriorityIndex)
 	{
-		const UDEPRECATED_UPawnAction* Action = ActionStacks[PriorityIndex].GetTop();
+		const UDEPRECATED_PawnAction* Action = ActionStacks[PriorityIndex].GetTop();
 		if (Action == NULL)
 		{
 			continue;
@@ -641,7 +641,7 @@ void UDEPRECATED_UPawnActionsComponent::DescribeSelfToVisLog(FVisualLogEntry* Sn
 }
 #endif // ENABLE_VISUAL_LOG
 
-FString UDEPRECATED_UPawnActionsComponent::DescribeEventType(EPawnActionEventType::Type EventType)
+FString UDEPRECATED_PawnActionsComponent::DescribeEventType(EPawnActionEventType::Type EventType)
 {
 	return GetEventName(EventType);
 }
diff --git a/Engine/Source/Runtime/AIModule/Private/BehaviorTree/Tasks/BTTask_PawnActionBase.cpp b/Engine/Source/Runtime/AIModule/Private/BehaviorTree/Tasks/BTTask_PawnActionBase.cpp
--- a/Engine/Source/Runtime/AIModule/Private/BehaviorTree/Tasks/BTTask_PawnActionBase.cpp
+++ b/Engine/Source/Runtime/AIModule/Private/BehaviorTree/Tasks/BTTask_PawnActionBase.cpp
@@ -22,7 +22,7 @@ EBTNodeResult::Type UBTTask_PawnActionBase::AbortTask(UBehaviorTreeComponent& Ow
 		? EBTNodeResult::InProgress : EBTNodeResult::Aborted;
 }
 
-EBTNodeResult::Type UBTTask_PawnActionBase::PushAction(UBehaviorTreeComponent& OwnerComp, UDEPRECATED_UPawnAction& Action)
+EBTNodeResult::Type UBTTask_PawnActionBase::PushAction(UBehaviorTreeComponent& OwnerComp, UDEPRECATED_PawnAction& Action)
 {
 	EBTNodeResult::Type NodeResult = EBTNodeResult::Failed;
 
@@ -47,7 +47,7 @@ EBTNodeResult::Type UBTTask_PawnActionBase::PushAction(UBehaviorTreeComponent& O
 	return NodeResult;
 }
 
-void UBTTask_PawnActionBase::OnActionEvent(UDEPRECATED_UPawnAction& Action, EPawnActionEventType::Type Event)
+void UBTTask_PawnActionBase::OnActionEvent(UDEPRECATED_PawnAction& Action, EPawnActionEventType::Type Event)
 {
 	EPawnActionTaskResult Result = ActionEventHandler(this, Action, Event);
 	if (Result == EPawnActionTaskResult::ActionLost)
@@ -56,7 +56,7 @@ void UBTTask_PawnActionBase::OnActionEvent(UDEPRECATED_UPawnAction& Action, EPaw
 	}
 }
 
-EPawnActionTaskResult UBTTask_PawnActionBase::ActionEventHandler(UBTTaskNode* TaskNode, UDEPRECATED_UPawnAction& Action, EPawnActionEventType::Type Event)
+EPawnActionTaskResult UBTTask_PawnActionBase::ActionEventHandler(UBTTaskNode* TaskNode, UDEPRECATED_PawnAction& Action, EPawnActionEventType::Type Event)
 {
 	EPawnActionTaskResult Result = EPawnActionTaskResult::Unknown;
 
@@ -66,7 +66,7 @@ EPawnActionTaskResult UBTTask_PawnActionBase::ActionEventHandler(UBTTaskNode* Ta
 	{
 		UE_VLOG(Action.GetController(), LogBehaviorTree, Warning, TEXT("%s: Unhandled PawnAction event: %s, can't find owning component!"),
 			*UBehaviorTreeTypes::DescribeNodeHelper(TaskNode),
-			*UDEPRECATED_UPawnActionsComponent::DescribeEventType(Event));
+			*UDEPRECATED_PawnActionsComponent::DescribeEventType(Event));
 
 		return Result;
 	}
@@ -100,14 +100,14 @@ EPawnActionTaskResult UBTTask_PawnActionBase::ActionEventHandler(UBTTaskNode* Ta
 	{
 		UE_VLOG(Action.GetController(), LogBehaviorTree, Verbose, TEXT("%s: Unhandled PawnAction event: %s TaskStatus: %s"),
 			*UBehaviorTreeTypes::DescribeNodeHelper(TaskNode),
-			*UDEPRECATED_UPawnActionsComponent::DescribeEventType(Event),
+			*UDEPRECATED_PawnActionsComponent::DescribeEventType(Event),
 			*UBehaviorTreeTypes::DescribeTaskStatus(TaskStatus));
 	}
 
 	return Result;
 }
 
-void UBTTask_PawnActionBase::OnActionLost(UDEPRECATED_UPawnAction& Action)
+void UBTTask_PawnActionBase::OnActionLost(UDEPRECATED_PawnAction& Action)
 {
 	AAIController* AIOwner = Cast<AAIController>(Action.GetController());
 	UBehaviorTreeComponent* OwnerComp = AIOwner ? Cast<UBehaviorTreeComponent>(AIOwner->BrainComponent) : nullptr;
diff --git a/Engine/Source/Runtime/AIModule/Private/BehaviorTree/Tasks/BTTask_PushPawnAction.cpp b/Engine/Source/Runtime/AIModule/Private/BehaviorTree/Tasks/BTTask_PushPawnAction.cpp
--- a/Engine/Source/Runtime/AIModule/Private/BehaviorTree/Tasks/BTTask_PushPawnAction.cpp
+++ b/Engine/Source/Runtime/AIModule/Private/BehaviorTree/Tasks/BTTask_PushPawnAction.cpp
@@ -14,7 +14,7 @@ UBTTask_PushPawnAction::UBTTask_PushPawnAction(const FObjectInitializer& ObjectI
 
 EBTNodeResult::Type UBTTask_PushPawnAction::ExecuteTask(UBehaviorTreeComponent& OwnerComp, uint8* NodeMemory)
 {
-	UDEPRECATED_UPawnAction* ActionCopy = Action_DEPRECATED ? DuplicateObject<UDEPRECATED_UPawnAction>(Action_DEPRECATED, &OwnerComp) : nullptr;
+	UDEPRECATED_PawnAction* ActionCopy = Action_DEPRECATED ? DuplicateObject<UDEPRECATED_PawnAction>(Action_DEPRECATED, &OwnerComp) : nullptr;
 	if (ActionCopy == nullptr)
 	{
 		return EBTNodeResult::Failed;
