Deprecating properties in structures that referenced UPawnActor.  C# UHT failed to complain about them.

diff --git a/Engine/Source/Runtime/AIModule/Classes/Actions/PawnActionsComponent.h b/Engine/Source/Runtime/AIModule/Classes/Actions/PawnActionsComponent.h
--- a/Engine/Source/Runtime/AIModule/Classes/Actions/PawnActionsComponent.h
+++ b/Engine/Source/Runtime/AIModule/Classes/Actions/PawnActionsComponent.h
@@ -22,7 +22,7 @@ struct AIMODULE_API FPawnActionEvent
 	static const int32 FakeActionIndex = INDEX_NONE;
 
 	UPROPERTY()
-	TObjectPtr<UDEPRECATED_PawnAction> Action;
+	TObjectPtr<UDEPRECATED_PawnAction> Action_DEPRECATED;
 
 	EPawnActionEventType::Type EventType;
 
@@ -31,12 +31,12 @@ struct AIMODULE_API FPawnActionEvent
 	// used to maintain order of equally-important messages
 	uint32 Index;
 
-	FPawnActionEvent() : Action(NULL), EventType(EPawnActionEventType::Invalid), Priority(EAIRequestPriority::MAX), Index(uint32(-1))
+	FPawnActionEvent() : Action_DEPRECATED(NULL), EventType(EPawnActionEventType::Invalid), Priority(EAIRequestPriority::MAX), Index(uint32(-1))
 	{}
 
 	FPawnActionEvent(UDEPRECATED_PawnAction& Action, EPawnActionEventType::Type EventType, uint32 Index);
 
-	bool operator==(const FPawnActionEvent& Other) const { return (Action == Other.Action) && (EventType == Other.EventType) && (Priority == Other.Priority); }
+	bool operator==(const FPawnActionEvent& Other) const { return (Action_DEPRECATED == Other.Action_DEPRECATED) && (EventType == Other.EventType) && (Priority == Other.Priority); }
 };
 
 USTRUCT()
@@ -45,12 +45,12 @@ struct AIMODULE_API FPawnActionStack
 	GENERATED_USTRUCT_BODY()
 
 	FPawnActionStack()
-		: TopAction(nullptr)
+		: TopAction_DEPRECATED(nullptr)
 	{}
 
 private:
 	UPROPERTY()
-	TObjectPtr<UDEPRECATED_PawnAction> TopAction;
+	TObjectPtr<UDEPRECATED_PawnAction> TopAction_DEPRECATED;
 
 public:
 	void Pause();
@@ -64,9 +64,9 @@ public:
 	 *	and if found action will be popped along with all it's siblings */
 	void PopAction(UDEPRECATED_PawnAction& ActionToPop);
 	
-	FORCEINLINE UDEPRECATED_PawnAction* GetTop() const { return TopAction; }
+	FORCEINLINE UDEPRECATED_PawnAction* GetTop() const { return TopAction_DEPRECATED; }
 
-	FORCEINLINE bool IsEmpty() const { return TopAction == NULL; }
+	FORCEINLINE bool IsEmpty() const { return TopAction_DEPRECATED == NULL; }
 
 	//----------------------------------------------------------------------//
 	// Debugging-testing purposes 
diff --git a/Engine/Source/Runtime/AIModule/Private/Actions/PawnActionsComponent.cpp b/Engine/Source/Runtime/AIModule/Private/Actions/PawnActionsComponent.cpp
--- a/Engine/Source/Runtime/AIModule/Private/Actions/PawnActionsComponent.cpp
+++ b/Engine/Source/Runtime/AIModule/Private/Actions/PawnActionsComponent.cpp
@@ -63,7 +63,7 @@ namespace
 }
 
 FPawnActionEvent::FPawnActionEvent(UDEPRECATED_PawnAction& InAction, EPawnActionEventType::Type InEventType, uint32 InIndex)
-	: Action(&InAction), EventType(InEventType), Index(InIndex)
+	: Action_DEPRECATED(&InAction), EventType(InEventType), Index(InIndex)
 {
 	Priority = InAction.GetPriority();
 }
@@ -74,41 +74,41 @@ FPawnActionEvent::FPawnActionEvent(UDEPRECATED_PawnAction& InAction, EPawnAction
 
 void FPawnActionStack::Pause()
 {
-	if (TopAction != NULL)
+	if (TopAction_DEPRECATED != NULL)
 	{
-		TopAction->Pause(NULL);
+		TopAction_DEPRECATED->Pause(NULL);
 	}
 }
 
 void FPawnActionStack::Resume()
 {
-	if (TopAction != NULL)
+	if (TopAction_DEPRECATED != NULL)
 	{
-		TopAction->Resume();
+		TopAction_DEPRECATED->Resume();
 	}
 }
 
 void FPawnActionStack::PushAction(UDEPRECATED_PawnAction& NewTopAction)
 {
-	if (TopAction != NULL)
+	if (TopAction_DEPRECATED != NULL)
 	{
-		if (TopAction->IsPaused() == false && TopAction->HasBeenStarted() == true)
+		if (TopAction_DEPRECATED->IsPaused() == false && TopAction_DEPRECATED->HasBeenStarted() == true)
 		{
-			TopAction->Pause(&NewTopAction);
+			TopAction_DEPRECATED->Pause(&NewTopAction);
 		}
-		ensure(TopAction->ChildAction_DEPRECATED == NULL);
-		TopAction->ChildAction_DEPRECATED = &NewTopAction;
-		NewTopAction.ParentAction_DEPRECATED = TopAction;
+		ensure(TopAction_DEPRECATED->ChildAction_DEPRECATED == NULL);
+		TopAction_DEPRECATED->ChildAction_DEPRECATED = &NewTopAction;
+		NewTopAction.ParentAction_DEPRECATED = TopAction_DEPRECATED;
 	}
 
-	TopAction = &NewTopAction;
+	TopAction_DEPRECATED = &NewTopAction;
 	NewTopAction.OnPushed();
 }
 
 void FPawnActionStack::PopAction(UDEPRECATED_PawnAction& ActionToPop)
 {
 	// first check if it's there
-	UDEPRECATED_PawnAction* CutPoint = TopAction;
+	UDEPRECATED_PawnAction* CutPoint = TopAction_DEPRECATED;
 	while (CutPoint != NULL && CutPoint != &ActionToPop)
 	{
 		CutPoint = CutPoint->ParentAction_DEPRECATED;
@@ -116,7 +116,7 @@ void FPawnActionStack::PopAction(UDEPRECATED_PawnAction& ActionToPop)
 
 	if (CutPoint == &ActionToPop)
 	{
-		UDEPRECATED_PawnAction* ActionBeingRemoved = TopAction;
+		UDEPRECATED_PawnAction* ActionBeingRemoved = TopAction_DEPRECATED;
 		// note StopAction can be null
 		UDEPRECATED_PawnAction* StopAction = ActionToPop.ParentAction_DEPRECATED;
 
@@ -138,14 +138,14 @@ void FPawnActionStack::PopAction(UDEPRECATED_PawnAction& ActionToPop)
 			ActionBeingRemoved = NextAction;
 		}
 
-		TopAction = StopAction;
+		TopAction_DEPRECATED = StopAction;
 	}
 }
 
 int32 FPawnActionStack::GetStackSize() const
 {
 	int32 Size = 0;
-	const UDEPRECATED_PawnAction* TempAction = TopAction;
+	const UDEPRECATED_PawnAction* TempAction = TopAction_DEPRECATED;
 	while (TempAction != nullptr)
 	{
 		TempAction = TempAction->GetParentAction();
@@ -211,7 +211,7 @@ void UDEPRECATED_PawnActionsComponent::TickComponent(float DeltaTime, enum ELeve
 		{
 			FPawnActionEvent& Event = ActionEvents[EventIndex];
 
-			if (Event.Action == nullptr)
+			if (Event.Action_DEPRECATED == nullptr)
 			{
 				UE_VLOG(ControlledPawn, LogPawnAction, Warning, TEXT("NULL action encountered during ActionEvents processing. May result in some notifies not being sent out."));
 				continue;
@@ -224,17 +224,17 @@ void UDEPRECATED_PawnActionsComponent::TickComponent(float DeltaTime, enum ELeve
 				// because of it, we need to operate on copy instead of reference to memory address
 				{
 					FPawnActionEvent EventCopy(Event);
-					EventCopy.Action->Abort(EAIForceParam::Force);
-					ActionStacks[EventCopy.Priority].PopAction(*EventCopy.Action);
+					EventCopy.Action_DEPRECATED->Abort(EAIForceParam::Force);
+					ActionStacks[EventCopy.Priority].PopAction(*EventCopy.Action_DEPRECATED);
 				}
 				break;
 			case EPawnActionEventType::FinishedAborting:
 			case EPawnActionEventType::FinishedExecution:
 			case EPawnActionEventType::FailedToStart:
-				ActionStacks[Event.Priority].PopAction(*Event.Action);
+				ActionStacks[Event.Priority].PopAction(*Event.Action_DEPRECATED);
 				break;
 			case EPawnActionEventType::Push:
-				ActionStacks[Event.Priority].PushAction(*Event.Action);
+				ActionStacks[Event.Priority].PushAction(*Event.Action_DEPRECATED);
 				break;
 			default:
 				break;
@@ -417,7 +417,7 @@ void UDEPRECATED_PawnActionsComponent::RemoveEventsForAction(UDEPRECATED_PawnAct
 {
 	for (int32 ActionIndex = ActionEvents.Num() - 1; ActionIndex >= 0; --ActionIndex)
 	{
-		if (ActionEvents[ActionIndex].Action == &PawnAction)
+		if (ActionEvents[ActionIndex].Action_DEPRECATED == &PawnAction)
 		{
 			ActionEvents.RemoveAtSwap(ActionIndex, /*Count=*/1, /*bAllowShrinking=*/false);
 		}
@@ -468,7 +468,7 @@ uint32 UDEPRECATED_PawnActionsComponent::AbortActionsInstigatedBy(UObject* const
 			const FPawnActionEvent& Event = ActionEvents[ActionIndex];
 			if (Event.Priority == Priority &&
 				Event.EventType == EPawnActionEventType::Push &&
-				Event.Action && Event.Action->GetInstigator() == Instigator)
+				Event.Action_DEPRECATED && Event.Action_DEPRECATED->GetInstigator() == Instigator)
 			{
 				ActionEvents.RemoveAtSwap(ActionIndex, /*Count=*/1, /*bAllowShrinking=*/false);
 				AbortedActionsCount++;
