Engine changes for behavior injection
Moved FNodeInitializationData type to its own file
Exposed injection tag from the RunBehaviorDynamic task
Exposed the selected key from the Blackboard base Decorator
Made FBehaviorTreeInstance move only so the they aren't copied and the NodeMemory array doesn't change memory address so callbacks that pass the NodeMemory ptr can work

diff --git a/Engine/Source/Runtime/AIModule/Classes/BehaviorTree/BehaviorTreeTypes.h b/Engine/Source/Runtime/AIModule/Classes/BehaviorTree/BehaviorTreeTypes.h
--- a/Engine/Source/Runtime/AIModule/Classes/BehaviorTree/BehaviorTreeTypes.h
+++ b/Engine/Source/Runtime/AIModule/Classes/BehaviorTree/BehaviorTreeTypes.h
@@ -324,7 +324,16 @@ struct FBehaviorTreeInstance
 	FBTInstanceDeactivation DeactivationNotify;
 
 	AIMODULE_API FBehaviorTreeInstance();
+
+	UE_DEPRECATED(5.2, "Copying FBehaviorTreeInstance constructor has been deprecated in favor of move-constructor")
 	AIMODULE_API FBehaviorTreeInstance(const FBehaviorTreeInstance& Other);
+
+	UE_DEPRECATED(5.2, "Copying FBehaviorTreeInstance assignement operator has been deprecated in favor of move assignement operator")
+	AIMODULE_API FBehaviorTreeInstance& operator=(const FBehaviorTreeInstance& Other) = default;
+
+	AIMODULE_API FBehaviorTreeInstance(FBehaviorTreeInstance&& Other);
+	AIMODULE_API FBehaviorTreeInstance& operator=(FBehaviorTreeInstance&& Other);
+
 	AIMODULE_API FBehaviorTreeInstance(int32 MemorySize);
 	AIMODULE_API ~FBehaviorTreeInstance();
 
diff --git a/Engine/Source/Runtime/AIModule/Classes/BehaviorTree/Decorators/BTDecorator_BlackboardBase.h b/Engine/Source/Runtime/AIModule/Classes/BehaviorTree/Decorators/BTDecorator_BlackboardBase.h
--- a/Engine/Source/Runtime/AIModule/Classes/BehaviorTree/Decorators/BTDecorator_BlackboardBase.h
+++ b/Engine/Source/Runtime/AIModule/Classes/BehaviorTree/Decorators/BTDecorator_BlackboardBase.h
@@ -26,6 +26,7 @@ class AIMODULE_API UBTDecorator_BlackboardBase : public UBTDecorator
 	virtual FName GetNodeIconName() const override;
 #endif
 
+public:
 	/** get name of selected blackboard key */
 	FName GetSelectedBlackboardKey() const;
 
diff --git a/Engine/Source/Runtime/AIModule/Classes/BehaviorTree/Tasks/BTTask_RunBehaviorDynamic.h b/Engine/Source/Runtime/AIModule/Classes/BehaviorTree/Tasks/BTTask_RunBehaviorDynamic.h
--- a/Engine/Source/Runtime/AIModule/Classes/BehaviorTree/Tasks/BTTask_RunBehaviorDynamic.h
+++ b/Engine/Source/Runtime/AIModule/Classes/BehaviorTree/Tasks/BTTask_RunBehaviorDynamic.h
@@ -32,6 +32,7 @@ class AIMODULE_API UBTTask_RunBehaviorDynamic : public UBTTaskNode
 #endif // WITH_EDITOR
 
 	bool HasMatchingTag(const FGameplayTag& Tag) const;
+	const FGameplayTag& GetInjectionTag() const;
 	bool SetBehaviorAsset(UBehaviorTree* NewBehaviorAsset);
 	
 	/** @returns default subtree asset */
@@ -68,3 +69,7 @@ FORCEINLINE UBehaviorTree* UBTTask_RunBehaviorDynamic::GetDefaultBehaviorAsset()
 	return DefaultBehaviorAsset;
 }
 
+FORCEINLINE const FGameplayTag& UBTTask_RunBehaviorDynamic::GetInjectionTag() const
+{
+	return InjectionTag;
+}
diff --git a/Engine/Source/Runtime/AIModule/Private/BehaviorTree/BehaviorTreeComponent.cpp b/Engine/Source/Runtime/AIModule/Private/BehaviorTree/BehaviorTreeComponent.cpp
--- a/Engine/Source/Runtime/AIModule/Private/BehaviorTree/BehaviorTreeComponent.cpp
+++ b/Engine/Source/Runtime/AIModule/Private/BehaviorTree/BehaviorTreeComponent.cpp
@@ -2603,7 +2603,7 @@ bool UBehaviorTreeComponent::PushInstance(UBehaviorTree& TreeAsset)
 	const bool bLoaded = BTManager->LoadTree(TreeAsset, RootNode, InstanceMemorySize);
 	if (bLoaded)
 	{
-		FBehaviorTreeInstance NewInstance;
+		FBehaviorTreeInstance& NewInstance = InstanceStack.AddDefaulted_GetRef();
 		NewInstance.InstanceIdIndex = UpdateInstanceId(&TreeAsset, ActiveNode, InstanceStack.Num() - 1);
 		NewInstance.RootNode = RootNode;
 		NewInstance.ActiveNode = NULL;
@@ -2622,7 +2622,6 @@ bool UBehaviorTreeComponent::PushInstance(UBehaviorTree& TreeAsset)
 		NewInstance.SetInstanceMemory(InstanceInfo.InstanceMemory);
 		NewInstance.Initialize(*this, *RootNode, NodeInstanceIndex, bFirstTime ? EBTMemoryInit::Initialize : EBTMemoryInit::RestoreSubtree);
 
-		InstanceStack.Push(NewInstance);
 
 		const int32 InstanceIndex = InstanceStack.Num() - 1;
 		ActiveInstanceIdx = IntCastChecked<uint16>(InstanceIndex);
diff --git a/Engine/Source/Runtime/AIModule/Private/BehaviorTree/BehaviorTreeManager.cpp b/Engine/Source/Runtime/AIModule/Private/BehaviorTree/BehaviorTreeManager.cpp
--- a/Engine/Source/Runtime/AIModule/Private/BehaviorTree/BehaviorTreeManager.cpp
+++ b/Engine/Source/Runtime/AIModule/Private/BehaviorTree/BehaviorTreeManager.cpp
@@ -7,6 +7,7 @@
 #include "EngineGlobals.h"
 #include "Engine/Engine.h"
 #include "AISystem.h"
+#include "BehaviorTreeNodeInitializationData.h"
 #include "BehaviorTree/BehaviorTreeTypes.h"
 #include "BehaviorTree/BTNode.h"
 #include "BehaviorTree/BTTaskNode.h"
@@ -60,35 +61,6 @@ uint16 UBehaviorTreeManager::GetAlignedDataSize(uint16 Size)
 	return ((Size + 3) & ~3);
 }
 
-struct FNodeInitializationData
-{
-	UBTNode* Node;
-	UBTCompositeNode* ParentNode;
-	uint16 ExecutionIndex;
-	uint16 DataSize;
-	uint16 SpecialDataSize;
-	uint8 TreeDepth;
-
-	FNodeInitializationData() {}
-	FNodeInitializationData(UBTNode* InNode, UBTCompositeNode* InParentNode,
-		uint16 InExecutionIndex, uint8 InTreeDepth, uint16 NodeMemory, uint16 SpecialNodeMemory = 0)
-		: Node(InNode), ParentNode(InParentNode), ExecutionIndex(InExecutionIndex), TreeDepth(InTreeDepth)
-	{
-		SpecialDataSize = UBehaviorTreeManager::GetAlignedDataSize(SpecialNodeMemory);
-
-		const uint16 NodeMemorySize = NodeMemory + SpecialDataSize;
-		DataSize = (NodeMemorySize <= 2) ? NodeMemorySize : UBehaviorTreeManager::GetAlignedDataSize(NodeMemorySize);
-	}
-
-	struct FMemorySort
-	{
-		FORCEINLINE bool operator()(const FNodeInitializationData& A, const FNodeInitializationData& B) const
-		{
-			return A.DataSize > B.DataSize;
-		}
-	};
-};
-
 static void MergeDecoratorOpsHelper(TArray<FBTDecoratorLogic>& LinkOps, const TArray<FBTDecoratorLogic>& InjectedOps, const int32 NumOriginalDecorators, const int32 NumInjectedDecorators)
 {
 	if (LinkOps.Num() == 0 && InjectedOps.Num() == 0)
@@ -142,10 +114,10 @@ static void MergeDecoratorOpsHelper(TArray<FBTDecoratorLogic>& LinkOps, const TA
 }
 
 static void InitializeNodeHelper(UBTCompositeNode* ParentNode, UBTNode* NodeOb,
-	uint8 TreeDepth, uint16& ExecutionIndex, TArray<FNodeInitializationData>& InitList,
+	uint8 TreeDepth, uint16& ExecutionIndex, TArray<FBehaviorTreeNodeInitializationData>& InitList,
 	UBehaviorTree& TreeAsset, UObject* NodeOuter)
 {
-	InitList.Add(FNodeInitializationData(NodeOb, ParentNode, ExecutionIndex, TreeDepth, NodeOb->GetInstanceMemorySize(), NodeOb->GetSpecialMemorySize()));
+	InitList.Add(FBehaviorTreeNodeInitializationData(NodeOb, ParentNode, ExecutionIndex, TreeDepth, NodeOb->GetInstanceMemorySize(), NodeOb->GetSpecialMemorySize()));
 	{
 		FScopedBTLoggingContext LogContext(NodeOb);
 		NodeOb->InitializeFromAsset(TreeAsset);
@@ -170,7 +142,7 @@ static void InitializeNodeHelper(UBTCompositeNode* ParentNode, UBTNode* NodeOb,
 			UBTService* Service = Cast<UBTService>(StaticDuplicateObject(CompositeOb->Services[ServiceIndex], NodeOuter));;
 			CompositeOb->Services[ServiceIndex] = Service;
 
-			InitList.Add(FNodeInitializationData(Service, CompositeOb, ExecutionIndex, TreeDepth,
+			InitList.Add(FBehaviorTreeNodeInitializationData(Service, CompositeOb, ExecutionIndex, TreeDepth,
 				Service->GetInstanceMemorySize(), Service->GetSpecialMemorySize()));
 
 			FScopedBTLoggingContext LogContext(Service);
@@ -197,7 +169,7 @@ static void InitializeNodeHelper(UBTCompositeNode* ParentNode, UBTNode* NodeOb,
 				UBTDecorator* Decorator = Cast<UBTDecorator>(StaticDuplicateObject(ChildInfo.Decorators[DecoratorIndex], NodeOuter));
 				ChildInfo.Decorators[DecoratorIndex] = Decorator;
 
-				InitList.Add(FNodeInitializationData(Decorator, CompositeOb, ExecutionIndex, TreeDepth,
+				InitList.Add(FBehaviorTreeNodeInitializationData(Decorator, CompositeOb, ExecutionIndex, TreeDepth,
 					Decorator->GetInstanceMemorySize(), Decorator->GetSpecialMemorySize()));
 
 				FScopedBTLoggingContext LogContext(Decorator);
@@ -226,7 +198,7 @@ static void InitializeNodeHelper(UBTCompositeNode* ParentNode, UBTNode* NodeOb,
 					UBTDecorator* Decorator = Cast<UBTDecorator>(StaticDuplicateObject(SubtreeAsset.RootDecorators[DecoratorIndex], NodeOuter));
 					ChildInfo.Decorators.Add(Decorator);
 
-					InitList.Add(FNodeInitializationData(Decorator, CompositeOb, ExecutionIndex, TreeDepth,
+					InitList.Add(FBehaviorTreeNodeInitializationData(Decorator, CompositeOb, ExecutionIndex, TreeDepth,
 						Decorator->GetInstanceMemorySize(), Decorator->GetSpecialMemorySize()));
 
 					// initialize with parent tree
@@ -268,7 +240,7 @@ static void InitializeNodeHelper(UBTCompositeNode* ParentNode, UBTNode* NodeOb,
 					UBTService* Service = Cast<UBTService>(StaticDuplicateObject(ChildInfo.ChildTask->Services[ServiceIndex], NodeOuter));
 					ChildInfo.ChildTask->Services[ServiceIndex] = Service;
 
-					InitList.Add(FNodeInitializationData(Service, CompositeOb, ExecutionIndex, TreeDepth,
+					InitList.Add(FBehaviorTreeNodeInitializationData(Service, CompositeOb, ExecutionIndex, TreeDepth,
 						Service->GetInstanceMemorySize(), Service->GetSpecialMemorySize()));
 
 					FScopedBTLoggingContext LogContext(Service);
@@ -309,7 +281,7 @@ bool UBehaviorTreeManager::LoadTree(UBehaviorTree& Asset, UBTCompositeNode*& Roo
 		TemplateInfo.Asset = &Asset;
 		TemplateInfo.Template = Cast<UBTCompositeNode>(StaticDuplicateObject(Asset.RootNode, this));
 
-		TArray<FNodeInitializationData> InitList;
+		TArray<FBehaviorTreeNodeInitializationData> InitList;
 		uint16 ExecutionIndex = 0;
 		InitializeNodeHelper(NULL, TemplateInfo.Template, 0, ExecutionIndex, InitList, Asset, this);
 
@@ -324,7 +296,7 @@ bool UBehaviorTreeManager::LoadTree(UBehaviorTree& Asset, UBTCompositeNode*& Roo
 		// sort nodes by memory size, so they can be packed better
 		// it still won't protect against structures, that are internally misaligned (-> uint8, uint32)
 		// but since all Engine level nodes are good... 
-		InitList.Sort(FNodeInitializationData::FMemorySort());
+		InitList.Sort(FBehaviorTreeNodeInitializationData::FMemorySort());
 		uint16 MemoryOffset = 0;
 		for (int32 Index = 0; Index < InitList.Num(); Index++)
 		{
@@ -346,7 +318,7 @@ bool UBehaviorTreeManager::LoadTree(UBehaviorTree& Asset, UBTCompositeNode*& Roo
 
 void UBehaviorTreeManager::InitializeMemoryHelper(const TArray<UBTDecorator*>& Nodes, TArray<uint16>& MemoryOffsets, int32& MemorySize, bool bForceInstancing)
 {
-	TArray<FNodeInitializationData> InitList;
+	TArray<FBehaviorTreeNodeInitializationData> InitList;
 	for (int32 NodeIndex = 0; NodeIndex < Nodes.Num(); NodeIndex++)
 	{
 		UBTNode* Node = Nodes[NodeIndex];
@@ -357,7 +329,7 @@ void UBehaviorTreeManager::InitializeMemoryHelper(const TArray<UBTDecorator*>& N
 			Node->ForceInstancing(true);
 		}
 
-		InitList.Add(FNodeInitializationData(Nodes[NodeIndex], NULL, 0, 0, Nodes[NodeIndex]->GetInstanceMemorySize(), Nodes[NodeIndex]->GetSpecialMemorySize()));
+		InitList.Add(FBehaviorTreeNodeInitializationData(Nodes[NodeIndex], NULL, 0, 0, Nodes[NodeIndex]->GetInstanceMemorySize(), Nodes[NodeIndex]->GetSpecialMemorySize()));
 
 		if (bForceInstancing && !bUsesInstancing)
 		{
@@ -365,7 +337,7 @@ void UBehaviorTreeManager::InitializeMemoryHelper(const TArray<UBTDecorator*>& N
 		}
 	}
 
-	InitList.Sort(FNodeInitializationData::FMemorySort());
+	InitList.Sort(FBehaviorTreeNodeInitializationData::FMemorySort());
 
 	uint16 MemoryOffset = 0;
 	MemoryOffsets.AddZeroed(Nodes.Num());
diff --git a/Engine/Source/Runtime/AIModule/Private/BehaviorTree/BehaviorTreeTypes.cpp b/Engine/Source/Runtime/AIModule/Private/BehaviorTree/BehaviorTreeTypes.cpp
--- a/Engine/Source/Runtime/AIModule/Private/BehaviorTree/BehaviorTreeTypes.cpp
+++ b/Engine/Source/Runtime/AIModule/Private/BehaviorTree/BehaviorTreeTypes.cpp
@@ -186,8 +186,21 @@ FBehaviorTreeInstance::FBehaviorTreeInstance()
 }
 
 FBehaviorTreeInstance::FBehaviorTreeInstance(const FBehaviorTreeInstance& Other)
-{ 
-	*this = Other; 
+	: RootNode(Other.RootNode)
+	, ActiveNode(Other.ActiveNode)
+	, InstanceIdIndex(Other.InstanceIdIndex)
+	, ActiveNodeType(Other.ActiveNodeType)
+	, DeactivationNotify(Other.DeactivationNotify)
+{
+	ActiveAuxNodes = Other.ActiveAuxNodes;
+	ParallelTasks = Other.ParallelTasks;
+	InstanceMemory = Other.InstanceMemory;
+
+#if DO_ENSURE
+	bIteratingNodes = Other.bIteratingNodes;
+	ParallelTaskIndex = Other.ParallelTaskIndex;
+#endif // DO_ENSURE
+
 	IncMemoryStats();
 	INC_DWORD_STAT(STAT_AI_BehaviorTree_NumInstances);
 }
@@ -199,6 +212,49 @@ FBehaviorTreeInstance::FBehaviorTreeInstance(int32 MemorySize)
 	INC_DWORD_STAT(STAT_AI_BehaviorTree_NumInstances);
 }
 
+FBehaviorTreeInstance::FBehaviorTreeInstance(FBehaviorTreeInstance&& Other)
+	: RootNode(Other.RootNode)
+	, ActiveNode(Other.ActiveNode)
+	, InstanceIdIndex(Other.InstanceIdIndex)
+	, ActiveNodeType(Other.ActiveNodeType)
+	, DeactivationNotify(Other.DeactivationNotify)
+{
+	ActiveAuxNodes = MoveTemp(Other.ActiveAuxNodes);
+	ParallelTasks = MoveTemp(Other.ParallelTasks);
+	InstanceMemory = MoveTemp(Other.InstanceMemory);
+
+#if DO_ENSURE
+	bIteratingNodes = Other.bIteratingNodes;
+	ParallelTaskIndex = Other.ParallelTaskIndex;
+#endif // DO_ENSURE
+
+	IncMemoryStats();
+	INC_DWORD_STAT(STAT_AI_BehaviorTree_NumInstances);
+}
+
+FBehaviorTreeInstance& FBehaviorTreeInstance::operator=(FBehaviorTreeInstance&& Other)
+{
+	RootNode = Other.RootNode;
+	ActiveNode = Other.ActiveNode;
+	InstanceIdIndex = Other.InstanceIdIndex;
+	ActiveNodeType = Other.ActiveNodeType;
+	DeactivationNotify = Other.DeactivationNotify;
+
+	ActiveAuxNodes = MoveTemp(Other.ActiveAuxNodes);
+	ParallelTasks = MoveTemp(Other.ParallelTasks);
+	InstanceMemory = MoveTemp(Other.InstanceMemory);
+
+#if DO_ENSURE
+	bIteratingNodes = Other.bIteratingNodes;
+	ParallelTaskIndex = Other.ParallelTaskIndex;
+#endif // DO_ENSURE
+
+	IncMemoryStats();
+	INC_DWORD_STAT(STAT_AI_BehaviorTree_NumInstances);
+
+	return *this;
+}
+
 FBehaviorTreeInstance::~FBehaviorTreeInstance()
 {
 	DecMemoryStats();
diff --git a/Engine/Source/Runtime/AIModule/Public/BehaviorTreeNodeInitializationData.h b/Engine/Source/Runtime/AIModule/Public/BehaviorTreeNodeInitializationData.h
new file mode 100644
--- /dev/null
+++ b/Engine/Source/Runtime/AIModule/Public/BehaviorTreeNodeInitializationData.h
@@ -0,0 +1,36 @@
+// Copyright Epic Games, Inc. All Rights Reserved.
+#pragma once
+
+#include "BehaviorTree/BehaviorTreeManager.h"
+
+class UBTNode;
+class UBTCompositeNode;
+
+struct FBehaviorTreeNodeInitializationData
+{
+	UBTNode* Node;
+	UBTCompositeNode* ParentNode;
+	uint16 ExecutionIndex;
+	uint16 DataSize;
+	uint16 SpecialDataSize;
+	uint8 TreeDepth;
+
+	FBehaviorTreeNodeInitializationData() {}
+	FBehaviorTreeNodeInitializationData(UBTNode* InNode, UBTCompositeNode* InParentNode,
+		uint16 InExecutionIndex, uint8 InTreeDepth, uint16 NodeMemory, uint16 SpecialNodeMemory = 0)
+		: Node(InNode), ParentNode(InParentNode), ExecutionIndex(InExecutionIndex), TreeDepth(InTreeDepth)
+	{
+		SpecialDataSize = UBehaviorTreeManager::GetAlignedDataSize(SpecialNodeMemory);
+
+		const uint16 NodeMemorySize = NodeMemory + SpecialDataSize;
+		DataSize = (NodeMemorySize <= 2) ? NodeMemorySize : UBehaviorTreeManager::GetAlignedDataSize(NodeMemorySize);
+	}
+
+	struct FMemorySort
+	{
+		FORCEINLINE bool operator()(const FBehaviorTreeNodeInitializationData& A, const FBehaviorTreeNodeInitializationData& B) const
+		{
+			return A.DataSize > B.DataSize;
+		}
+	};
+};
