Made IsCompatibleWith virtual for BlackboardComponent
Exposed FBlackboardKey in AIModule API

diff --git a/Engine/Source/Runtime/AIModule/Classes/BehaviorTree/BlackboardComponent.h b/Engine/Source/Runtime/AIModule/Classes/BehaviorTree/BlackboardComponent.h
--- a/Engine/Source/Runtime/AIModule/Classes/BehaviorTree/BlackboardComponent.h
+++ b/Engine/Source/Runtime/AIModule/Classes/BehaviorTree/BlackboardComponent.h
@@ -97,7 +97,7 @@ public:
 	bool InitializeBlackboard(UBlackboardData& NewAsset);
 	
 	/** @return true if component can be used with specified blackboard asset */
-	bool IsCompatibleWith(const UBlackboardData* TestAsset) const;
+	virtual bool IsCompatibleWith(const UBlackboardData* TestAsset) const;
 
 	UFUNCTION(BlueprintCallable, Category="AI|Components|Blackboard")
 	UObject* GetValueAsObject(const FName& KeyName) const;
diff --git a/Engine/Source/Runtime/AIModule/Classes/BehaviorTree/BlackboardData.h b/Engine/Source/Runtime/AIModule/Classes/BehaviorTree/BlackboardData.h
--- a/Engine/Source/Runtime/AIModule/Classes/BehaviorTree/BlackboardData.h
+++ b/Engine/Source/Runtime/AIModule/Classes/BehaviorTree/BlackboardData.h
@@ -11,7 +11,7 @@
 
 /** blackboard entry definition */
 USTRUCT()
-struct FBlackboardEntry
+struct AIMODULE_API FBlackboardEntry
 {
 	GENERATED_USTRUCT_BODY()
 
