Fix for NavLink Proxy gyzmo does not update after clicking "Copy End Points from Simple Link to Smart Link" butto.

I also added transactional state to the functionality.

diff --git a/Engine/Source/Runtime/AIModule/Private/Navigation/NavLinkProxy.cpp b/Engine/Source/Runtime/AIModule/Private/Navigation/NavLinkProxy.cpp
--- a/Engine/Source/Runtime/AIModule/Private/Navigation/NavLinkProxy.cpp
+++ b/Engine/Source/Runtime/AIModule/Private/Navigation/NavLinkProxy.cpp
@@ -16,9 +16,15 @@
 #include "VisualLogger/VisualLogger.h"
 #include "NavigationOctree.h"
 #include "ObjectEditorUtils.h"
+#include "EditorSupportDelegates.h"
+#if WITH_EDITOR
+#include "ScopedTransaction.h"
+#endif
 
 #include UE_INLINE_GENERATED_CPP_BY_NAME(NavLinkProxy)
 
+#define LOCTEXT_NAMESPACE "NavLink"
+
 ANavLinkProxy::ANavLinkProxy(const FObjectInitializer& ObjectInitializer) : Super(ObjectInitializer)
 {
 	USceneComponent* SceneComponent = CreateDefaultSubobject<USceneComponent>(TEXT("PositionComponent"));
@@ -297,7 +303,18 @@ void ANavLinkProxy::CopyEndPointsFromSimpleLinkToSmartLink()
 {
 	if (PointLinks.Num() && SmartLinkComp)
 	{
-		SmartLinkComp->SetLinkData(PointLinks[0].Left, PointLinks[0].Right, PointLinks[0].Direction);
+		{
+			const FScopedTransaction Transaction(LOCTEXT("SetLinkData", "Set Link Data"));
+			SmartLinkComp->Modify();
+			SmartLinkComp->SetLinkData(PointLinks[0].Left, PointLinks[0].Right, PointLinks[0].Direction);
+		}
+		if (EdRenderComp)
+		{
+			EdRenderComp->MarkRenderStateDirty();
+		}
+		FEditorSupportDelegates::RedrawAllViewports.Broadcast();
 	}
 }
 #endif // WITH_EDITOR
+
+#undef LOCTEXT_NAMESPACE
\ No newline at end of file
