Replaced some enable-if/disable-if overloads with if constexpr blocks in a single function.

diff --git a/Engine/Plugins/Enterprise/DataprepEditor/Source/DataprepCore/Private/SelectionSystem/DataprepSelectionSystemUtils.h b/Engine/Plugins/Enterprise/DataprepEditor/Source/DataprepCore/Private/SelectionSystem/DataprepSelectionSystemUtils.h
--- a/Engine/Plugins/Enterprise/DataprepEditor/Source/DataprepCore/Private/SelectionSystem/DataprepSelectionSystemUtils.h
+++ b/Engine/Plugins/Enterprise/DataprepEditor/Source/DataprepCore/Private/SelectionSystem/DataprepSelectionSystemUtils.h
@@ -124,17 +124,18 @@ namespace DataprepSelectionSystemUtils
 	}
 
 	template<class FetchedDataType>
-	static typename TEnableIf<TIsFilterDataDisplayable<FetchedDataType>::Value>::Type UpdateReportedDataIfSupported(const FetchedDataType& FetchedData, FDataprepSelectionInfo& OutInfo)
+	static void UpdateReportedDataIfSupported(const FetchedDataType& FetchedData, FDataprepSelectionInfo& OutInfo)
 	{
-		OutInfo.FetchedData.Set<FetchedDataType>( FetchedData );
-		OutInfo.bWasDataFetchedAndCached = true;
-	}
-
-	template<class FetchedDataType>
-	static typename TEnableIf<!TIsFilterDataDisplayable<FetchedDataType>::Value>::Type UpdateReportedDataIfSupported(const FetchedDataType& FetchedData, FDataprepSelectionInfo& OutInfo)
-	{
-		// Do nothing since the data type is not supported
-		OutInfo.bWasDataFetchedAndCached = false;
+		if constexpr (TIsFilterDataDisplayable<FetchedDataType>::Value)
+		{
+			OutInfo.FetchedData.Set<FetchedDataType>( FetchedData );
+			OutInfo.bWasDataFetchedAndCached = true;
+		}
+		else
+		{
+			// Do nothing since the data type is not supported
+			OutInfo.bWasDataFetchedAndCached = false;
+		}
 	}
 
 	/**
diff --git a/Engine/Plugins/Enterprise/DatasmithImporter/Source/DatasmithImporter/Private/ObjectElements/DatasmithUElementsUtils.h b/Engine/Plugins/Enterprise/DatasmithImporter/Source/DatasmithImporter/Private/ObjectElements/DatasmithUElementsUtils.h
--- a/Engine/Plugins/Enterprise/DatasmithImporter/Source/DatasmithImporter/Private/ObjectElements/DatasmithUElementsUtils.h
+++ b/Engine/Plugins/Enterprise/DatasmithImporter/Source/DatasmithImporter/Private/ObjectElements/DatasmithUElementsUtils.h
@@ -49,37 +49,19 @@ public:
 	 * });
 	 */
 	template<typename TargetType, typename Func>
-	static inline typename TEnableIf<std::is_same_v<TargetType, IDatasmithLevelVariantSetsElement>, bool>::Type ForVariantElement(const TSharedPtr<IDatasmithScene>& RootObject, Func Function)
+	static inline bool ForVariantElement(const TSharedPtr<IDatasmithScene>& RootObject, Func Function)
 	{
 		for (int32 Index = 0; Index < RootObject->GetLevelVariantSetsCount(); ++Index)
 		{
-			bool bKeepIterating = Function(RootObject->GetLevelVariantSets(Index));
-			if (!bKeepIterating)
+			bool bKeepIterating;
+			if constexpr (std::is_same_v<TargetType, IDatasmithLevelVariantSetsElement>)
 			{
-				return false;
+				bKeepIterating = Function(RootObject->GetLevelVariantSets(Index));
 			}
-		}
-		return true;
-	}
-	template<typename TargetType, typename Func>
-	static inline typename TEnableIf<!std::is_same_v<TargetType, IDatasmithLevelVariantSetsElement>, bool>::Type ForVariantElement(const TSharedPtr<IDatasmithScene>& RootObject, Func Function)
-	{
-		for (int32 Index = 0; Index < RootObject->GetLevelVariantSetsCount(); ++Index)
-		{
-			bool bKeepIterating = ForVariantElement<TargetType>(RootObject->GetLevelVariantSets(Index), Function);
-			if (!bKeepIterating)
+			else
 			{
-				return false;
+				bKeepIterating = ForVariantElement<TargetType>(RootObject->GetLevelVariantSets(Index), Function);
 			}
-		}
-		return true;
-	}
-	template<typename TargetType, typename Func>
-	static inline typename TEnableIf<std::is_same_v<TargetType, IDatasmithVariantSetElement>, bool>::Type ForVariantElement(const TSharedPtr<IDatasmithLevelVariantSetsElement>& RootObject, Func Function)
-	{
-		for (int32 Index = 0; Index < RootObject->GetVariantSetsCount(); ++Index)
-		{
-			bool bKeepIterating = Function(RootObject->GetVariantSet(Index));
 			if (!bKeepIterating)
 			{
 				return false;
@@ -88,11 +70,19 @@ public:
 		return true;
 	}
 	template<typename TargetType, typename Func>
-	static inline typename TEnableIf<!std::is_same_v<TargetType, IDatasmithVariantSetElement>, bool>::Type ForVariantElement(const TSharedPtr<IDatasmithLevelVariantSetsElement>& RootObject, Func Function)
+	static inline bool ForVariantElement(const TSharedPtr<IDatasmithLevelVariantSetsElement>& RootObject, Func Function)
 	{
 		for (int32 Index = 0; Index < RootObject->GetVariantSetsCount(); ++Index)
 		{
-			bool bKeepIterating = ForVariantElement<TargetType>(RootObject->GetVariantSet(Index), Function);
+			bool bKeepIterating;
+			if constexpr (std::is_same_v<TargetType, IDatasmithVariantSetElement>)
+			{
+				bKeepIterating = Function(RootObject->GetVariantSet(Index));
+			}
+			else
+			{
+				bKeepIterating = ForVariantElement<TargetType>(RootObject->GetVariantSet(Index), Function);
+			}
 			if (!bKeepIterating)
 			{
 				return false;
@@ -102,37 +92,19 @@ public:
 	}
 
 	template<typename TargetType, typename Func>
-	static inline typename TEnableIf<std::is_same_v<TargetType, IDatasmithVariantElement>, bool>::Type ForVariantElement(const TSharedPtr<IDatasmithVariantSetElement>& RootObject, Func Function)
+	static inline bool ForVariantElement(const TSharedPtr<IDatasmithVariantSetElement>& RootObject, Func Function)
 	{
 		for (int32 Index = 0; Index < RootObject->GetVariantsCount(); ++Index)
 		{
-			bool bKeepIterating = Function(RootObject->GetVariant(Index));
-			if (!bKeepIterating)
+			bool bKeepIterating;
+			if constexpr (std::is_same_v<TargetType, IDatasmithVariantElement>)
 			{
-				return false;
+				bKeepIterating = Function(RootObject->GetVariant(Index));
 			}
-		}
-		return true;
-	}
-	template<typename TargetType, typename Func>
-	static inline typename TEnableIf<!std::is_same_v<TargetType, IDatasmithVariantElement>, bool>::Type ForVariantElement(const TSharedPtr<IDatasmithVariantSetElement>& RootObject, Func Function)
-	{
-		for (int32 Index = 0; Index < RootObject->GetVariantsCount(); ++Index)
-		{
-			bool bKeepIterating = ForVariantElement<TargetType>(RootObject->GetVariant(Index), Function);
-			if (!bKeepIterating)
+			else
 			{
-				return false;
+				bKeepIterating = ForVariantElement<TargetType>(RootObject->GetVariant(Index), Function);
 			}
-		}
-		return true;
-	}
-	template<typename TargetType, typename Func>
-	static inline typename TEnableIf<std::is_same_v<TargetType, IDatasmithActorBindingElement>, bool>::Type ForVariantElement(const TSharedPtr<IDatasmithVariantElement>& RootObject, Func Function)
-	{
-		for (int32 Index = 0; Index < RootObject->GetActorBindingsCount(); ++Index)
-		{
-			bool bKeepIterating = Function(RootObject->GetActorBinding(Index));
 			if (!bKeepIterating)
 			{
 				return false;
@@ -141,11 +113,19 @@ public:
 		return true;
 	}
 	template<typename TargetType, typename Func>
-	static inline typename TEnableIf<!std::is_same_v<TargetType, IDatasmithActorBindingElement>, bool>::Type ForVariantElement(const TSharedPtr<IDatasmithVariantElement>& RootObject, Func Function)
+	static inline bool ForVariantElement(const TSharedPtr<IDatasmithVariantElement>& RootObject, Func Function)
 	{
 		for (int32 Index = 0; Index < RootObject->GetActorBindingsCount(); ++Index)
 		{
-			bool bKeepIterating = ForVariantElement<TargetType>(RootObject->GetActorBinding(Index), Function);
+			bool bKeepIterating;
+			if constexpr (std::is_same_v<TargetType, IDatasmithActorBindingElement>)
+			{
+				bKeepIterating = Function(RootObject->GetActorBinding(Index));
+			}
+			else
+			{
+				bKeepIterating = ForVariantElement<TargetType>(RootObject->GetActorBinding(Index), Function);
+			}
 			if (!bKeepIterating)
 			{
 				return false;
diff --git a/Engine/Source/Editor/BlueprintGraph/Classes/BlueprintNodeBinder.h b/Engine/Source/Editor/BlueprintGraph/Classes/BlueprintNodeBinder.h
--- a/Engine/Source/Editor/BlueprintGraph/Classes/BlueprintNodeBinder.h
+++ b/Engine/Source/Editor/BlueprintGraph/Classes/BlueprintNodeBinder.h
@@ -94,39 +94,40 @@ public:
 		return !bIsUObject && Property.IsValid() && Property->IsA(InClass);
 	}
 	template <typename T>
-	typename TEnableIf<TIsDerivedFrom<T, UObject>::IsDerived, bool>::Type IsA() const
+	bool IsA() const
 	{
-		if (bIsUObject && Object.IsValid())
+		if constexpr (std::is_base_of_v<UObject, T>)
 		{
-			return Object->IsA(T::StaticClass());
+			if (bIsUObject && Object.IsValid())
+			{
+				return Object->IsA(T::StaticClass());
+			}
 		}
-		return false;
-	}
-
-	template <typename T>
-	typename TEnableIf<!TIsDerivedFrom<T, UObject>::IsDerived, bool>::Type IsA() const
-	{
-		if (!bIsUObject && Property.IsValid())
+		else
 		{
-			return Property->IsA(T::StaticClass());
+			if (!bIsUObject && Property.IsValid())
+			{
+				return Property->IsA(T::StaticClass());
+			}
 		}
 		return false;
 	}
 	template <typename T>
-	typename TEnableIf<TIsDerivedFrom<T, UObject>::IsDerived, T*>::Type Get() const
+	T* Get() const
 	{
-		if (bIsUObject && Object.IsValid())
+		if constexpr (std::is_base_of_v<UObject, T>)
 		{
-			return Cast<T>(Object.Get());
+			if (bIsUObject && Object.IsValid())
+			{
+				return Cast<T>(Object.Get());
+			}
 		}
-		return nullptr;
-	}
-	template <typename T>
-	typename TEnableIf<!TIsDerivedFrom<T, UObject>::IsDerived, T*>::Type Get() const
-	{
-		if (!bIsUObject && Property.IsValid())
+		else
 		{
-			return CastField<T>(Property.Get());
+			if (!bIsUObject && Property.IsValid())
+			{
+				return CastField<T>(Property.Get());
+			}
 		}
 		return nullptr;
 	}
diff --git a/Engine/Source/Runtime/AIModule/Public/GraphAStar.h b/Engine/Source/Runtime/AIModule/Public/GraphAStar.h
--- a/Engine/Source/Runtime/AIModule/Public/GraphAStar.h
+++ b/Engine/Source/Runtime/AIModule/Public/GraphAStar.h
@@ -72,9 +72,18 @@ struct CQuery##ConditionalFunctionName	\
 {	\
 	template<typename TemplateClass> auto Requires(TemplateClassParameter& Obj) -> decltype(Obj.ConditionalFunctionName());	\
 };	\
-template <typename TemplateClass> using ReturnType##ConditionalFunctionName = decltype((DeclVal<TemplateClass>().*DeclVal<decltype(&TemplateClass::ConditionalFunctionName)>())()); \
-template <typename TemplateClass> static FORCEINLINE typename TEnableIf< TModels_V<CQuery##ConditionalFunctionName, TemplateClass>, ReturnType##ConditionalFunctionName<TemplateClass>>::Type ConditionalFunctionName(TemplateClassParameter & Obj) { return Obj.ConditionalFunctionName(); } \
-template <typename TemplateClass> static FORCEINLINE typename TEnableIf<!TModels_V<CQuery##ConditionalFunctionName, TemplateClass>, ConditionalReturnType>::Type ConditionalFunctionName(TemplateClassParameter& Obj) { return DefaultImpl; }
+template <typename TemplateClass> \
+static FORCEINLINE decltype(auto) ConditionalFunctionName(TemplateClassParameter & Obj) \
+{ \
+	if constexpr (TModels_V<CQuery##ConditionalFunctionName, TemplateClass>) \
+	{ \
+		return Obj.ConditionalFunctionName(); \
+	} \
+	else \
+	{ \
+		return (ConditionalReturnType)(DefaultImpl); \
+	} \
+}
 #define DECLARE_OPTIONALLY_IMPLEMENTED_TEMPLATE_CLASS_FUNCTION( ConditionalReturnType, ConditionalFunctionName, DefaultImpl )  DECLARE_OPTIONALLY_IMPLEMENTED_TEMPLATE_CLASS_FUNCTION_IMPL( TemplateClass, TemplateClass, ConditionalReturnType, ConditionalFunctionName, DefaultImpl )
 #define DECLARE_OPTIONALLY_IMPLEMENTED_TEMPLATE_CLASS_FUNCTION_CONST( ConditionalReturnType, ConditionalFunctionName, DefaultImpl ) DECLARE_OPTIONALLY_IMPLEMENTED_TEMPLATE_CLASS_FUNCTION_IMPL( TemplateClass, const TemplateClass, ConditionalReturnType, ConditionalFunctionName, DefaultImpl )
 
@@ -83,9 +92,18 @@ struct CQuery##ConditionalFunctionName	\
 {	\
 	template<typename TemplateClass> auto Requires(TemplateClassParameter& Obj, ConditionalParamType1 Param1) -> decltype(Obj.ConditionalFunctionName(Param1));	\
 };	\
-template <typename TemplateClass> using ReturnType##ConditionalFunctionName = decltype((DeclVal<TemplateClass>().*DeclVal<decltype(&TemplateClass::ConditionalFunctionName)>())(DeclVal<ConditionalParamType1>())); \
-template <typename TemplateClass> static FORCEINLINE typename TEnableIf< TModels_V<CQuery##ConditionalFunctionName, TemplateClass>, ReturnType##ConditionalFunctionName<TemplateClass>>::Type ConditionalFunctionName(TemplateClassParameter & Obj, ConditionalParamType1 Param1) { return Obj.ConditionalFunctionName(Param1); } \
-template <typename TemplateClass> static FORCEINLINE typename TEnableIf<!TModels_V<CQuery##ConditionalFunctionName, TemplateClass>, ConditionalReturnType>::Type ConditionalFunctionName(TemplateClassParameter& Obj, ConditionalParamType1 Param1) { return DefaultImpl; }
+template <typename TemplateClass> \
+static FORCEINLINE decltype(auto) ConditionalFunctionName(TemplateClassParameter & Obj, ConditionalParamType1 Param1) \
+{ \
+	if constexpr (TModels_V<CQuery##ConditionalFunctionName, TemplateClass>) \
+	{ \
+		return Obj.ConditionalFunctionName(Param1); \
+	} \
+	else \
+	{ \
+		return (ConditionalReturnType)(DefaultImpl); \
+	} \
+}
 #define DECLARE_OPTIONALLY_IMPLEMENTED_TEMPLATE_CLASS_FUNCTION_1PARAM( ConditionalReturnType, ConditionalFunctionName, ConditionalParamType1, DefaultImpl) DECLARE_OPTIONALLY_IMPLEMENTED_TEMPLATE_CLASS_FUNCTION_1PARAM_IMPL( TemplateClass, TemplateClass, ConditionalReturnType, ConditionalFunctionName, ConditionalParamType1, DefaultImpl) 
 #define DECLARE_OPTIONALLY_IMPLEMENTED_TEMPLATE_CLASS_FUNCTION_1PARAM_CONST( ConditionalReturnType, ConditionalFunctionName, ConditionalParamType1, DefaultImpl) DECLARE_OPTIONALLY_IMPLEMENTED_TEMPLATE_CLASS_FUNCTION_1PARAM_IMPL( TemplateClass, const TemplateClass, ConditionalReturnType, ConditionalFunctionName, ConditionalParamType1, DefaultImpl) 
 
@@ -94,9 +112,18 @@ struct CQuery##ConditionalFunctionName	\
 {	\
 	template<typename TemplateClass> auto Requires(TemplateClassParameter& Obj, ConditionalParamType1 Param1, ConditionalParamType2 Param2) -> decltype(Obj.ConditionalFunctionName(Param1,Param2));	\
 };	\
-template <typename TemplateClass> using ReturnType##ConditionalFunctionName = decltype((DeclVal<TemplateClass>().*DeclVal<decltype(&TemplateClass::ConditionalFunctionName)>())(DeclVal<ConditionalParamType1>(), DeclVal<ConditionalParamType2>())); \
-template <typename TemplateClass> static FORCEINLINE typename TEnableIf< TModels_V<CQuery##ConditionalFunctionName, TemplateClass>, ReturnType##ConditionalFunctionName<TemplateClass>>::Type ConditionalFunctionName(TemplateClassParameter & Obj, ConditionalParamType1 Param1, ConditionalParamType2 Param2) { return Obj.ConditionalFunctionName(Param1, Param2); } \
-template <typename TemplateClass> static FORCEINLINE typename TEnableIf<!TModels_V<CQuery##ConditionalFunctionName, TemplateClass>, ConditionalReturnType>::Type ConditionalFunctionName(TemplateClassParameter& Obj, ConditionalParamType1 Param1, ConditionalParamType2 Param2) { return DefaultImpl; }
+template <typename TemplateClass> \
+static FORCEINLINE decltype(auto) ConditionalFunctionName(TemplateClassParameter & Obj, ConditionalParamType1 Param1, ConditionalParamType2 Param2) \
+{ \
+	if constexpr (TModels_V<CQuery##ConditionalFunctionName, TemplateClass>) \
+	{ \
+		return Obj.ConditionalFunctionName(Param1, Param2); \
+	} \
+	else \
+	{ \
+		return (ConditionalReturnType)(DefaultImpl); \
+	} \
+}
 #define DECLARE_OPTIONALLY_IMPLEMENTED_TEMPLATE_CLASS_FUNCTION_2PARAMS( ConditionalReturnType, ConditionalFunctionName, ConditionalParamType1, ConditionalParamType2, DefaultImpl) DECLARE_OPTIONALLY_IMPLEMENTED_TEMPLATE_CLASS_FUNCTION_2PARAMS_IMPL( TemplateClass, TemplateClass, ConditionalReturnType, ConditionalFunctionName, ConditionalParamType1, ConditionalParamType2, DefaultImpl)
 #define DECLARE_OPTIONALLY_IMPLEMENTED_TEMPLATE_CLASS_FUNCTION_2PARAMS_CONST( ConditionalReturnType, ConditionalFunctionName, ConditionalParamType1, ConditionalParamType2, DefaultImpl) DECLARE_OPTIONALLY_IMPLEMENTED_TEMPLATE_CLASS_FUNCTION_2PARAMS_IMPL(TemplateClass, const TemplateClass, ConditionalReturnType, ConditionalFunctionName, ConditionalParamType1, ConditionalParamType2, DefaultImpl)
 
@@ -105,8 +132,18 @@ struct CQuery##QueryFunctionName	\
 {	\
 	template<typename TemplateClass> auto Requires(const TemplateClass& Obj) -> decltype(Obj.ConditionalFunctionName());	\
 };	\
-template <typename TemplateClass> static FORCEINLINE typename TEnableIf< TModels_V<CQuery##QueryFunctionName, TemplateClass>, QueryReturnType>::Type QueryFunctionName(const TemplateClass& Obj, QueryParam) { return QueryImpl; } \
-template <typename TemplateClass> static FORCEINLINE typename TEnableIf<!TModels_V<CQuery##QueryFunctionName, TemplateClass>, QueryReturnType>::Type QueryFunctionName(const TemplateClass& Obj, QueryParam) { return QueryDefaultImpl; }
+template <typename TemplateClass> \
+static FORCEINLINE QueryReturnType QueryFunctionName(const TemplateClass& Obj, QueryParam) \
+{ \
+	if constexpr (TModels_V<CQuery##QueryFunctionName, TemplateClass>) \
+	{ \
+		return QueryImpl; \
+	} \
+	else \
+	{ \
+		return QueryDefaultImpl; \
+	} \
+}
 
 template<bool DoRangeCheck>
 class FRangeChecklessAllocator : public FDefaultAllocator
diff --git a/Engine/Source/Runtime/AudioExtensions/Public/IAudioProxyInitializer.h b/Engine/Source/Runtime/AudioExtensions/Public/IAudioProxyInitializer.h
--- a/Engine/Source/Runtime/AudioExtensions/Public/IAudioProxyInitializer.h
+++ b/Engine/Source/Runtime/AudioExtensions/Public/IAudioProxyInitializer.h
@@ -135,25 +135,21 @@ public:
 
 namespace Audio
 {
-	// SFINAE used to optionally invoke subclasses of IAudioProxyDataFactory when we can.
-	template<typename UClassToUse, typename std::enable_if_t<std::is_base_of_v<IAudioProxyDataFactory, UClassToUse>, bool> = true>
+	template <typename UClassToUse>
 	IAudioProxyDataFactory* CastToProxyDataFactory(UObject* InObject)
 	{
-		if (InObject)
+		if constexpr (std::is_base_of_v<IAudioProxyDataFactory, UClassToUse>)
 		{
-			UClassToUse* DowncastObject = Cast<UClassToUse>(InObject);
-			if (ensureAlways(DowncastObject))
+			if (InObject)
 			{
-				return static_cast<IAudioProxyDataFactory*>(DowncastObject);
+				UClassToUse* DowncastObject = Cast<UClassToUse>(InObject);
+				if (ensureAlways(DowncastObject))
+				{
+					return static_cast<IAudioProxyDataFactory*>(DowncastObject);
+				}
 			}
 		}
 
 		return nullptr;
 	}
-
-	template<typename UClassToUse, typename std::enable_if_t<!std::is_base_of_v<IAudioProxyDataFactory, UClassToUse>, bool>::Type = true>
-	IAudioProxyDataFactory* CastToProxyDataFactory(UObject* InObject)
-	{
-		return nullptr;
-	}
 } // namespace Audio
diff --git a/Engine/Source/Runtime/Core/Public/HAL/IConsoleManager.h b/Engine/Source/Runtime/Core/Public/HAL/IConsoleManager.h
--- a/Engine/Source/Runtime/Core/Public/HAL/IConsoleManager.h
+++ b/Engine/Source/Runtime/Core/Public/HAL/IConsoleManager.h
@@ -1682,29 +1682,31 @@ private:
 	EConsoleVariableFlags Flags = EConsoleVariableFlags::ECVF_Default;
 
 	template<class Y>
-	typename TEnableIf<!std::is_same_v<T, Y>, Y>::Type GetImpl() const
+	Y GetImpl() const
 	{
-		check(false);
-		return Y();
-	}
-
-	template<class Y>
-	typename TEnableIf<std::is_same_v<T, Y>, Y>::Type GetImpl() const
-	{
-		return GetValueOnAnyThread();
-	}
-
-	template<class Y>
-	typename TEnableIf<!std::is_same_v<T, Y>, TConsoleVariableData<Y>*>::Type AsImpl()
-	{
-		check(false);
-		return nullptr;
+		if constexpr (std::is_same_v<T, Y>)
+		{
+			return GetValueOnAnyThread();
+		}
+		else
+		{
+			check(false);
+			return Y();
+		}
 	}
 
 	template<class Y>
-	typename TEnableIf<std::is_same_v<T, Y>, TConsoleVariableData<T>*>::Type AsImpl()
+	TConsoleVariableData<T>* AsImpl()
 	{
-		return &Value;
+		if constexpr (std::is_same_v<T, Y>)
+		{
+			return &Value;
+		}
+		else
+		{
+			check(false);
+			return nullptr;
+		}
 	}
 };
 #endif // NO_CVARS
diff --git a/Engine/Source/Runtime/Core/Public/Logging/LogMacros.h b/Engine/Source/Runtime/Core/Public/Logging/LogMacros.h
--- a/Engine/Source/Runtime/Core/Public/Logging/LogMacros.h
+++ b/Engine/Source/Runtime/Core/Public/Logging/LogMacros.h
@@ -161,25 +161,17 @@ private:
 	namespace UEAsserts_Private
 	{
 		template <int32 VerbosityToCheck, typename CategoryType>
-		FORCEINLINE
-			typename TEnableIf<
-				((VerbosityToCheck & ELogVerbosity::VerbosityMask) <= CategoryType::CompileTimeVerbosity &&
-				(VerbosityToCheck & ELogVerbosity::VerbosityMask) <= ELogVerbosity::COMPILED_IN_MINIMUM_VERBOSITY),
-				bool>::Type
-			IsLogActive(const CategoryType& Category)
+		FORCEINLINE bool IsLogActive(const CategoryType& Category)
 		{
-			return !Category.IsSuppressed((ELogVerbosity::Type)VerbosityToCheck);
-		}
-
-		template <int32 VerbosityToCheck, typename CategoryType>
-		FORCEINLINE
-			typename TEnableIf<
-				!((VerbosityToCheck & ELogVerbosity::VerbosityMask) <= CategoryType::CompileTimeVerbosity &&
-				(VerbosityToCheck & ELogVerbosity::VerbosityMask) <= ELogVerbosity::COMPILED_IN_MINIMUM_VERBOSITY),
-				bool>::Type
-			IsLogActive(const CategoryType& Category)
-		{
-			return false;
+			if constexpr (((VerbosityToCheck & ELogVerbosity::VerbosityMask) <= CategoryType::CompileTimeVerbosity &&
+				(VerbosityToCheck & ELogVerbosity::VerbosityMask) <= ELogVerbosity::COMPILED_IN_MINIMUM_VERBOSITY))
+			{
+				return !Category.IsSuppressed((ELogVerbosity::Type)VerbosityToCheck);
+			}
+			else
+			{
+				return false;
+			}
 		}
 	}
 
diff --git a/Engine/Source/Runtime/Core/Public/Math/GenericOctree.h b/Engine/Source/Runtime/Core/Public/Math/GenericOctree.h
--- a/Engine/Source/Runtime/Core/Public/Math/GenericOctree.h
+++ b/Engine/Source/Runtime/Core/Public/Math/GenericOctree.h
@@ -1047,14 +1047,16 @@ private:
 
 	// Function overload set which calls the V2 version if it's supported or the old version if it's not
 	template <typename Semantics>
-	typename TEnableIf<!TModels_V<COctreeSemanticsV2, Semantics>>::Type SetOctreeSemanticsElementId(const ElementType& Element, FOctreeElementId2 Id)
+	void SetOctreeSemanticsElementId(const ElementType& Element, FOctreeElementId2 Id)
 	{
-		Semantics::SetElementId(Element, Id);
-	}
-	template <typename Semantics>
-	typename TEnableIf<TModels_V<COctreeSemanticsV2, Semantics>>::Type SetOctreeSemanticsElementId(const ElementType& Element, FOctreeElementId2 Id)
-	{
-		Semantics::SetElementId(static_cast<typename Semantics::FOctree&>(*this), Element, Id);
+		if constexpr (TModels_V<COctreeSemanticsV2, Semantics>)
+		{
+			Semantics::SetElementId(static_cast<typename Semantics::FOctree&>(*this), Element, Id);
+		}
+		else
+		{
+			Semantics::SetElementId(Element, Id);
+		}
 	}
 
 protected:
@@ -1782,16 +1784,18 @@ private:
 			-> decltype(Semantics::SetElementId(OctreeInstance, Element, Id));
 	};
 
-	// Function overload set which calls the V2 version if it's supported or the old version if it's not
+	// Calls the V2 version if it's supported or the old version if it's not
 	template <typename Semantics>
-	typename TEnableIf<!TModels_V<COctreeSemanticsV2, Semantics>>::Type SetOctreeSemanticsElementId(const ElementType& Element, FOctreeElementId Id)
+	void SetOctreeSemanticsElementId(const ElementType& Element, FOctreeElementId Id)
 	{
-		Semantics::SetElementId(Element, Id);
-	}
-	template <typename Semantics>
-	typename TEnableIf<TModels_V<COctreeSemanticsV2, Semantics>>::Type SetOctreeSemanticsElementId(const ElementType& Element, FOctreeElementId Id)
-	{
-		Semantics::SetElementId(*this, Element, Id);
+		if constexpr (TModels_V<COctreeSemanticsV2, Semantics>)
+		{
+			Semantics::SetElementId(*this, Element, Id);
+		}
+		else
+		{
+			Semantics::SetElementId(Element, Id);
+		}
 	}
 
 protected:
diff --git a/Engine/Source/Runtime/Core/Public/Misc/HashBuilder.h b/Engine/Source/Runtime/Core/Public/Misc/HashBuilder.h
--- a/Engine/Source/Runtime/Core/Public/Misc/HashBuilder.h
+++ b/Engine/Source/Runtime/Core/Public/Misc/HashBuilder.h
@@ -34,16 +34,17 @@ public:
 	}
 
 	template <typename T>
-	typename TEnableIf<!TModels_V<CGetTypeHashable, T>, FHashBuilder&>::Type Append(const T& InData)
+	FHashBuilder& Append(const T& InData)
 	{
-		return AppendRaw(InData);
-	}
-
-	template <typename T>
-	typename TEnableIf<TModels_V<CGetTypeHashable, T>, FHashBuilder&>::Type Append(const T& InData)
-	{
-		Hash = HashCombineFast(Hash, GetTypeHash(InData));
-		return *this;
+		if constexpr (TModels_V<CGetTypeHashable, T>)
+		{
+			Hash = HashCombineFast(Hash, GetTypeHash(InData));
+			return *this;
+		}
+		else
+		{
+			return AppendRaw(InData);
+		}
 	}
 
 	template <typename T>
diff --git a/Engine/Source/Runtime/Core/Public/ProfilingDebugging/FormatArgsTrace.h b/Engine/Source/Runtime/Core/Public/ProfilingDebugging/FormatArgsTrace.h
--- a/Engine/Source/Runtime/Core/Public/ProfilingDebugging/FormatArgsTrace.h
+++ b/Engine/Source/Runtime/Core/Public/ProfilingDebugging/FormatArgsTrace.h
@@ -52,22 +52,23 @@ private:
 		enum { Value = TAnd<TIsPointer<T>, TIsCharType<std::remove_cv_t<typename TRemovePointer<T>::Type>>>::Value };
 	};
 
-	template <typename T>
-	constexpr static typename TEnableIf<!TIsStringArgument<T>::Value, uint64>::Type GetArgumentEncodedSize(T Argument)
-	{
-		return sizeof(T);
-	}
-
 	template <typename T, typename CharType = std::remove_cv_t<typename TRemovePointer<T>::Type>>
-	static typename TEnableIf<TIsStringArgument<T>::Value, uint64>::Type GetArgumentEncodedSize(T Argument)
+	constexpr static uint64 GetArgumentEncodedSize(T Argument)
 	{
-		if (Argument != nullptr)
+		if constexpr (TIsStringArgument<T>::Value)
 		{
-			return (TCString<CharType>::Strlen(Argument) + 1) * sizeof(CharType);
+			if (Argument != nullptr)
+			{
+				return (TCString<CharType>::Strlen(Argument) + 1) * sizeof(CharType);
+			}
+			else
+			{
+				return sizeof(CharType);
+			}
 		}
 		else
 		{
-			return sizeof(CharType);
+			return sizeof(T);
 		}
 	}
 
@@ -83,54 +84,47 @@ private:
 	}
 
 	template <typename T>
-	static typename TEnableIf<TAnd<TNot<TIsFloatingPoint<T>>, TNot<TIsStringArgument<T>>>::Value>::Type EncodeArgumentInternal(uint8*& TypeCodesPtr, uint8*& PayloadPtr, T Argument)
-	{
-		*TypeCodesPtr++ = FormatArgTypeCode_CategoryInteger | sizeof(T);
-
-#if PLATFORM_SUPPORTS_UNALIGNED_LOADS
-		*reinterpret_cast<T*>(PayloadPtr) = Argument;
-#else
-		// For ARM targets, it's possible that using __packed here would be preferable
-		// but I have not checked the codegen -- it's possible that the compiler generates
-		// the same code for this fixed size memcpy
-		memcpy(PayloadPtr, &Argument, sizeof Argument);
-#endif
-
-		PayloadPtr += sizeof(T);
-	}
-
-	template <typename T>
-	static typename TEnableIf<TIsFloatingPoint<T>::Value>::Type EncodeArgumentInternal(uint8*& TypeCodesPtr, uint8*& PayloadPtr, T Argument)
+	static void EncodeArgumentInternal(uint8*& TypeCodesPtr, uint8*& PayloadPtr, T Argument)
 	{
-		*TypeCodesPtr++ = FormatArgTypeCode_CategoryFloatingPoint | sizeof(T);
+		if constexpr (TIsStringArgument<T>::Value)
+		{
+			using CharType = std::remove_cv_t<typename TRemovePointer<T>::Type>;
+
+			*TypeCodesPtr++ = FormatArgTypeCode_CategoryString | sizeof(CharType);
+			if (Argument != nullptr)
+			{
+				uint16 Length = (uint16)((TCString<CharType>::Strlen(Argument) + 1) * sizeof(CharType));
+				memcpy(PayloadPtr, Argument, Length);
+				PayloadPtr += Length;
+			}
+			else
+			{
+				CharType Terminator { 0 };
+				memcpy(PayloadPtr, &Terminator, sizeof(CharType));
+				PayloadPtr += sizeof(CharType);
+			}
+		}
+		else
+		{
+			if constexpr (std::is_floating_point_v<T>)
+			{
+				*TypeCodesPtr++ = FormatArgTypeCode_CategoryFloatingPoint | sizeof(T);
+			}
+			else
+			{
+				*TypeCodesPtr++ = FormatArgTypeCode_CategoryInteger | sizeof(T);
+			}
 
 #if PLATFORM_SUPPORTS_UNALIGNED_LOADS
-		*reinterpret_cast<T*>(PayloadPtr) = Argument;
+			*reinterpret_cast<T*>(PayloadPtr) = Argument;
 #else
-		// For ARM targets, it's possible that using __packed here would be preferable
-		// but I have not checked the codegen -- it's possible that the compiler generates
-		// the same code for this fixed size memcpy
-		memcpy(PayloadPtr, &Argument, sizeof Argument);
+			// For ARM targets, it's possible that using __packed here would be preferable
+			// but I have not checked the codegen -- it's possible that the compiler generates
+			// the same code for this fixed size memcpy
+			memcpy(PayloadPtr, &Argument, sizeof Argument);
 #endif
 
-		PayloadPtr += sizeof(T);
-	}
-
-	template <typename T, typename CharType = std::remove_cv_t<typename TRemovePointer<T>::Type>>
-	static typename TEnableIf<TIsStringArgument<T>::Value>::Type EncodeArgumentInternal(uint8*& TypeCodesPtr, uint8*& PayloadPtr, T Argument)
-	{
-		*TypeCodesPtr++ = FormatArgTypeCode_CategoryString | sizeof(CharType);
-		if (Argument != nullptr)
-		{
-			uint16 Length = (uint16)((TCString<CharType>::Strlen(Argument) + 1) * sizeof(CharType));
-			memcpy(PayloadPtr, Argument, Length);
-			PayloadPtr += Length;
-		}
-		else
-		{
-			CharType Terminator { 0 };
-			memcpy(PayloadPtr, &Terminator, sizeof(CharType));
-			PayloadPtr += sizeof(CharType);
+			PayloadPtr += sizeof(T);
 		}
 	}
 
diff --git a/Engine/Source/Runtime/Core/Public/Serialization/MemoryLayout.h b/Engine/Source/Runtime/Core/Public/Serialization/MemoryLayout.h
--- a/Engine/Source/Runtime/Core/Public/Serialization/MemoryLayout.h
+++ b/Engine/Source/Runtime/Core/Public/Serialization/MemoryLayout.h
@@ -179,16 +179,17 @@ struct THasCustomDefaultObject
 };
 
 template<typename T>
-typename TEnableIf<THasCustomDefaultObject<T>::Value, const T*>::Type InternalGetDefaultObject()
+const T* InternalGetDefaultObject()
 {
-	return GetDefault<T>();
-}
-
-template<typename T>
-typename TEnableIf<!THasCustomDefaultObject<T>::Value, const T*>::Type InternalGetDefaultObject()
-{
-	static const T Default;
-	return &Default;
+	if constexpr (THasCustomDefaultObject<T>::Value)
+	{
+		return GetDefault<T>();
+	}
+	else
+	{
+		static const T Default;
+		return &Default;
+	}
 }
 
 template<typename T, ETypeLayoutInterface::Type InterfaceType>
diff --git a/Engine/Source/Runtime/Core/Public/Templates/UnrealTemplate.h b/Engine/Source/Runtime/Core/Public/Templates/UnrealTemplate.h
--- a/Engine/Source/Runtime/Core/Public/Templates/UnrealTemplate.h
+++ b/Engine/Source/Runtime/Core/Public/Templates/UnrealTemplate.h
@@ -559,7 +559,7 @@ template <typename T>
 struct TUseBitwiseSwap
 {
 	// We don't use bitwise swapping for 'register' types because this will force them into memory and be slower.
-	enum { Value = !TOrValue<__is_enum(T), TIsPointer<T>, TIsArithmetic<T>>::Value };
+	enum { Value = !std::is_enum_v<T> || std::is_pointer_v<T> || std::is_arithmetic_v<T> };
 };
 
 
@@ -567,23 +567,24 @@ struct TUseBitwiseSwap
  * Swap two values.  Assumes the types are trivially relocatable.
  */
 template <typename T>
-inline typename TEnableIf<TUseBitwiseSwap<T>::Value>::Type Swap(T& A, T& B)
+inline void Swap(T& A, T& B)
 {
-	if (LIKELY(&A != &B))
+	if constexpr (TUseBitwiseSwap<T>::Value)
 	{
-		TTypeCompatibleBytes<T> Temp;
-		FMemory::Memcpy(&Temp, &A, sizeof(T));
-		FMemory::Memcpy(&A, &B, sizeof(T));
-		FMemory::Memcpy(&B, &Temp, sizeof(T));
+		if (LIKELY(&A != &B))
+		{
+			TTypeCompatibleBytes<T> Temp;
+			FMemory::Memcpy(&Temp, &A, sizeof(T));
+			FMemory::Memcpy(&A, &B, sizeof(T));
+			FMemory::Memcpy(&B, &Temp, sizeof(T));
+		}
+	}
+	else
+	{
+		T Temp = MoveTemp(A);
+		A = MoveTemp(B);
+		B = MoveTemp(Temp);
 	}
-}
-
-template <typename T>
-inline typename TEnableIf<!TUseBitwiseSwap<T>::Value>::Type Swap(T& A, T& B)
-{
-	T Temp = MoveTemp(A);
-	A = MoveTemp(B);
-	B = MoveTemp(Temp);
 }
 
 template <typename T>
diff --git a/Engine/Source/Runtime/CoreUObject/Public/UObject/Class.h b/Engine/Source/Runtime/CoreUObject/Public/UObject/Class.h
--- a/Engine/Source/Runtime/CoreUObject/Public/UObject/Class.h
+++ b/Engine/Source/Runtime/CoreUObject/Public/UObject/Class.h
@@ -1535,14 +1535,16 @@ public:
 	static COREUOBJECT_API void DeferCppStructOps(FTopLevelAssetPath Target, ICppStructOps* InCppStructOps);
 
 	template<class CPPSTRUCT>
-	static typename TEnableIf<!DISABLE_ABSTRACT_CONSTRUCT>::Type DeferCppStructOps(FTopLevelAssetPath Target)
+	static void DeferCppStructOps(FTopLevelAssetPath Target)
 	{
-		DeferCppStructOps(Target, new UScriptStruct::TCppStructOps<CPPSTRUCT>);
-	}
-	template<class CPPSTRUCT>
-	static typename TEnableIf<DISABLE_ABSTRACT_CONSTRUCT>::Type DeferCppStructOps(FTopLevelAssetPath Target)
-	{
-		DeferCppStructOps(Target, nullptr);
+		if constexpr (DISABLE_ABSTRACT_CONSTRUCT)
+		{
+			DeferCppStructOps(Target, nullptr);
+		}
+		else
+		{
+			DeferCppStructOps(Target, new UScriptStruct::TCppStructOps<CPPSTRUCT>);
+		}
 	}
 
 	/** Look for the CppStructOps and hook it up **/
diff --git a/Engine/Source/Runtime/CoreUObject/Public/UObject/Field.h b/Engine/Source/Runtime/CoreUObject/Public/UObject/Field.h
--- a/Engine/Source/Runtime/CoreUObject/Public/UObject/Field.h
+++ b/Engine/Source/Runtime/CoreUObject/Public/UObject/Field.h
@@ -320,26 +320,23 @@ public:
 	}
 
 	template <typename T>
-	typename TEnableIf<TIsDerivedFrom<T, UObject>::IsDerived, T*>::Type Get() const
+	T* Get() const
 	{
 		static_assert(sizeof(T) > 0, "T must not be an incomplete type");
 		if (IsA(T::StaticClass()))
 		{
-			return static_cast<T*>(Container.Object);
+			if constexpr (std::is_base_of_v<UObject, T>)
+			{
+				return static_cast<T*>(Container.Object);
+			}
+			else
+			{
+				return static_cast<T*>(Container.Field);
+			}
 		}
 		return nullptr;
 	}
 
-	template <typename T>
-	typename TEnableIf<!TIsDerivedFrom<T, UObject>::IsDerived, T*>::Type Get() const
-	{
-		static_assert(sizeof(T) > 0, "T must not be an incomplete type");
-		if (IsA(T::StaticClass()))
-		{
-			return static_cast<T*>(Container.Field);
-		}
-		return nullptr;
-	}
 	UObject* ToUObject() const
 	{
 		if (bIsUObject)
diff --git a/Engine/Source/Runtime/RenderCore/Public/RenderGraphPass.h b/Engine/Source/Runtime/RenderCore/Public/RenderGraphPass.h
--- a/Engine/Source/Runtime/RenderCore/Public/RenderGraphPass.h
+++ b/Engine/Source/Runtime/RenderCore/Public/RenderGraphPass.h
@@ -596,15 +596,16 @@ public:
 
 private:
 	template<class T>
-	typename TEnableIf<!std::is_same_v<T, FRDGPass>, void>::Type ExecuteLambdaFunc(FRHIComputeCommandList& RHICmdList)
+	void ExecuteLambdaFunc(FRHIComputeCommandList& RHICmdList)
 	{
-		ExecuteLambda(static_cast<TRHICommandList&>(RHICmdList));
-	}
-
-	template<class T>
-	typename TEnableIf<std::is_same_v<T, FRDGPass>, void>::Type ExecuteLambdaFunc(FRHIComputeCommandList& RHICmdList)
-	{
-		ExecuteLambda(this, static_cast<TRHICommandList&>(RHICmdList));
+		if constexpr (std::is_same_v<T, FRDGPass>)
+		{
+			ExecuteLambda(this, static_cast<TRHICommandList&>(RHICmdList));
+		}
+		else
+		{
+			ExecuteLambda(static_cast<TRHICommandList&>(RHICmdList));
+		}
 	}
 
 	void Execute(FRHIComputeCommandList& RHICmdList) override
