Fixed a compile issue that got introduced by a bad merge.

diff --git a/Engine/Source/Runtime/AIModule/Classes/BehaviorTree/Blackboard/BlackboardKeyEnums.h b/Engine/Source/Runtime/AIModule/Classes/BehaviorTree/Blackboard/BlackboardKeyEnums.h
--- a/Engine/Source/Runtime/AIModule/Classes/BehaviorTree/Blackboard/BlackboardKeyEnums.h
+++ b/Engine/Source/Runtime/AIModule/Classes/BehaviorTree/Blackboard/BlackboardKeyEnums.h
@@ -2,6 +2,7 @@
 
 #pragma once
 
+#include "BlackboardKeyEnums.generated.h"
 
 namespace EBlackboardCompare
 {
