Removing unnecessary uses of Engine/Private

diff --git a/Engine/Source/Runtime/AIModule/AIModule.Build.cs b/Engine/Source/Runtime/AIModule/AIModule.Build.cs
--- a/Engine/Source/Runtime/AIModule/AIModule.Build.cs
+++ b/Engine/Source/Runtime/AIModule/AIModule.Build.cs
@@ -8,19 +8,6 @@ namespace UnrealBuildTool.Rules
     {
         public AIModule(ReadOnlyTargetRules Target) : base(Target)
         {
-            PublicIncludePaths.AddRange(
-                new string[] {
-                    "Runtime/NavigationSystem/Public",
-                    "Runtime/AIModule/Public",
-                }
-                );
-
-            PrivateIncludePaths.AddRange(
-                new string[] {
-					Path.Combine(GetModuleDirectory("Engine"), "Private"),
-                }
-                );
-
             PublicDependencyModuleNames.AddRange(
                 new string[] {
                     "Core",
@@ -39,12 +26,6 @@ namespace UnrealBuildTool.Rules
                 }
                 );
 
-            DynamicallyLoadedModuleNames.AddRange(
-                new string[] {
-					// ... add any modules that your module loads dynamically here ...
-				}
-                );
-
             if (Target.bBuildEditor == true)
             {
 				PrivateDependencyModuleNames.Add("EditorFramework");
diff --git a/Engine/Source/Runtime/MRMesh/MRMesh.build.cs b/Engine/Source/Runtime/MRMesh/MRMesh.build.cs
--- a/Engine/Source/Runtime/MRMesh/MRMesh.build.cs
+++ b/Engine/Source/Runtime/MRMesh/MRMesh.build.cs
@@ -26,9 +26,6 @@ namespace UnrealBuildTool.Rules
 				PrivateDependencyModuleNames.Add("UnrealEd");
 			}
 			
-			// Used for including the private Chaos headers
-			PrivateIncludePaths.Add(Path.Combine(GetModuleDirectory("Engine"), "Private"));
-
 			PrivateIncludePathModuleNames.Add("DerivedDataCache");
 		}
 	}
diff --git a/Engine/Source/Runtime/MRMesh/Private/MRMeshComponent.cpp b/Engine/Source/Runtime/MRMesh/Private/MRMeshComponent.cpp
--- a/Engine/Source/Runtime/MRMesh/Private/MRMeshComponent.cpp
+++ b/Engine/Source/Runtime/MRMesh/Private/MRMeshComponent.cpp
@@ -27,19 +27,10 @@
 #include "UObject/ConstructorHelpers.h"
 #include "Materials/MaterialInstanceDynamic.h"
 
-#include "PhysicsEngine/Experimental/ChaosDerivedData.h"
 #include "SceneInterface.h"
 
 #include UE_INLINE_GENERATED_CPP_BY_NAME(MRMeshComponent)
 
-// See ISpatialAcceleration.h, the header is not included to avoid having a dependency on Chaos
-#ifndef CHAOS_SERIALIZE_OUT
-#define CHAOS_SERIALIZE_OUT WITH_EDITOR
-#endif
-
-// When using the Chaos code path, we're using a private class FChaosDerivedDataCooker which is not exported
-#define SUPPORTS_PHYSICS_COOKING (CHAOS_SERIALIZE_OUT && IS_MONOLITHIC)
-
 DECLARE_CYCLE_STAT(TEXT("MrMesh SetCollisionProfileName"), STAT_MrMesh_SetCollisionProfileName, STATGROUP_Physics);
 DECLARE_CYCLE_STAT(TEXT("Update Collision"), STAT_UpdateCollision, STATGROUP_MRMESH);
 
diff --git a/Engine/Source/Runtime/NavigationSystem/NavigationSystem.Build.cs b/Engine/Source/Runtime/NavigationSystem/NavigationSystem.Build.cs
--- a/Engine/Source/Runtime/NavigationSystem/NavigationSystem.Build.cs
+++ b/Engine/Source/Runtime/NavigationSystem/NavigationSystem.Build.cs
@@ -8,15 +8,6 @@ namespace UnrealBuildTool.Rules
     {
         public NavigationSystem(ReadOnlyTargetRules Target) : base(Target)
         {
-			PublicIncludePathModuleNames.Add("NavigationSystem");
-
-            PrivateIncludePaths.AddRange(
-                new string[] {
-					Path.Combine(GetModuleDirectory("NavigationSystem"), "Private"),
-					Path.Combine(GetModuleDirectory("Engine"), "Private"),
-                }
-                );
-
             PublicDependencyModuleNames.AddRange(
                 new string[] {
                     "Core",
