Cut out all AI-related UE_DEPRECATED(4.x) code

diff --git a/Engine/Source/Runtime/AIModule/Classes/BehaviorTree/BTCompositeNode.h b/Engine/Source/Runtime/AIModule/Classes/BehaviorTree/BTCompositeNode.h
--- a/Engine/Source/Runtime/AIModule/Classes/BehaviorTree/BTCompositeNode.h
+++ b/Engine/Source/Runtime/AIModule/Classes/BehaviorTree/BTCompositeNode.h
@@ -96,10 +96,6 @@ class AIMODULE_API UBTCompositeNode : public UBTNode
 	UPROPERTY()
 	TArray<TObjectPtr<UBTService>> Services;
 
-	/** delegate for finding next child to execute */
-	UE_DEPRECATED(4.21, "OnNextChild is no longer being used. Please override UBTCompositeNode::GetNextChildHandler instead")
-	FGetNextChildDelegate OnNextChild;
-
 	~UBTCompositeNode();
 
 	/** fill in data about tree structure */
diff --git a/Engine/Source/Runtime/AIModule/Classes/BehaviorTree/BehaviorTreeTypes.h b/Engine/Source/Runtime/AIModule/Classes/BehaviorTree/BehaviorTreeTypes.h
--- a/Engine/Source/Runtime/AIModule/Classes/BehaviorTree/BehaviorTreeTypes.h
+++ b/Engine/Source/Runtime/AIModule/Classes/BehaviorTree/BehaviorTreeTypes.h
@@ -650,9 +650,6 @@ public:
 	FORCEINLINE void InvalidateResolvedKey() { SelectedKeyID = FBlackboard::InvalidKey; }
 
 	friend FBlackboardDecoratorDetails;
-
-	UE_DEPRECATED(4.24, "This version of AddClassFilter is deprecated. Please provide AllowedClass as TSubclassOf<UObject>")
-	void AddClassFilter(UObject* Owner, FName PropertyName, TSubclassOf<UClass> AllowedClass);
 };
 
 UCLASS(Abstract)
diff --git a/Engine/Source/Runtime/AIModule/Classes/EnvironmentQuery/Tests/EnvQueryTest_Overlap.h b/Engine/Source/Runtime/AIModule/Classes/EnvironmentQuery/Tests/EnvQueryTest_Overlap.h
--- a/Engine/Source/Runtime/AIModule/Classes/EnvironmentQuery/Tests/EnvQueryTest_Overlap.h
+++ b/Engine/Source/Runtime/AIModule/Classes/EnvironmentQuery/Tests/EnvQueryTest_Overlap.h
@@ -30,10 +30,4 @@ protected:
 
 	bool RunOverlap(const FVector& ItemPos, const FCollisionShape& CollisionShape, const TArray<AActor*>& IgnoredActors, UWorld* World, enum ECollisionChannel Channel, const FCollisionQueryParams& Params) const;
 	bool RunOverlapBlocking(const FVector& ItemPos, const FCollisionShape& CollisionShape, const TArray<AActor*>& IgnoredActors, UWorld* World, enum ECollisionChannel Channel, const FCollisionQueryParams& Params) const;
-
-    UE_DEPRECATED(4.23, "This method is no longer called by RunTest and has been replaced by an overload using a list of actors to ignore. It now calls that overload but you should use the new overload instead.")
-	bool RunOverlap(const FVector& ItemPos, const FCollisionShape& CollisionShape, AActor* ItemActor, UWorld* World, enum ECollisionChannel Channel, const FCollisionQueryParams& Params);
-
-	UE_DEPRECATED(4.23, "This method is no longer called by RunTest and has been replaced by an overload using a list of actors to ignore. It now calls that overload but you should use the new overload instead.")
-	bool RunOverlapBlocking(const FVector& ItemPos, const FCollisionShape& CollisionShape, AActor* ItemActor, UWorld* World, enum ECollisionChannel Channel, const FCollisionQueryParams& Params);
 };
diff --git a/Engine/Source/Runtime/AIModule/Classes/Navigation/CrowdFollowingComponent.h b/Engine/Source/Runtime/AIModule/Classes/Navigation/CrowdFollowingComponent.h
--- a/Engine/Source/Runtime/AIModule/Classes/Navigation/CrowdFollowingComponent.h
+++ b/Engine/Source/Runtime/AIModule/Classes/Navigation/CrowdFollowingComponent.h
@@ -156,10 +156,6 @@ public:
 
 protected:
 #if WITH_EDITORONLY_DATA
-	UE_DEPRECATED(4.26, "This member is deprecated and no longer used. Use AvoidanceInterface instead.")
-	UPROPERTY(transient)
-	TObjectPtr<class UCharacterMovementComponent> CharacterMovement;
-
 	/** DEPRECATED: Group mask for this agent - use IRVOAvoidanceInterface instead */
 	UPROPERTY()
 	FNavAvoidanceMask AvoidanceGroup_DEPRECATED;
diff --git a/Engine/Source/Runtime/AIModule/Classes/Perception/AIPerceptionComponent.h b/Engine/Source/Runtime/AIModule/Classes/Perception/AIPerceptionComponent.h
--- a/Engine/Source/Runtime/AIModule/Classes/Perception/AIPerceptionComponent.h
+++ b/Engine/Source/Runtime/AIModule/Classes/Perception/AIPerceptionComponent.h
@@ -151,12 +151,6 @@ struct AIMODULE_API FActorPerceptionInfo
 		return LastSensedStimuli.IsValidIndex(Sense) && (LastSensedStimuli[Sense].IsValid() && (LastSensedStimuli[Sense].IsExpired() == false)) ? LastSensedStimuli[Sense].ReceiverLocation : FAISystem::InvalidLocation;
 	}
 
-	UE_DEPRECATED(4.23, "This method is identical to IsSenseActive and will be removed in future versions. Please use IsSenseActive to check for a currently active stimuli or HasKnownStimulusOfSense for an active or remembered stimuli.")
-	FORCEINLINE bool IsSenseRegistered(FAISenseID Sense) const
-	{
-		return LastSensedStimuli.IsValidIndex(Sense) && LastSensedStimuli[Sense].IsActive();
-	}
-
 	/** Indicates a currently active or "remembered" stimuli for a given sense
 	* @param Sense	The AISenseID of the sense
 	*
@@ -360,11 +354,7 @@ public:
 	/** If SenseToUse is none all actors ever perceived in any way (and not forgotten yet) will get fetched */
 	UFUNCTION(BlueprintCallable, Category = "AI|Perception")
 	void GetKnownPerceivedActors(TSubclassOf<UAISense> SenseToUse, TArray<AActor*>& OutActors) const;
-	
-	UE_DEPRECATED(4.13, "GetPerceivedActors is deprecated. Use GetCurrentlyPerceivedActors or GetKnownPerceivedActors")
-	UFUNCTION(BlueprintCallable, Category = "AI|Perception")
-	void GetPerceivedActors(TSubclassOf<UAISense> SenseToUse, TArray<AActor*>& OutActors) const;
-	
+
 	/** Retrieves whatever has been sensed about given actor */
 	UFUNCTION(BlueprintCallable, Category = "AI|Perception")
 	bool GetActorsPerception(AActor* Actor, FActorPerceptionBlueprintInfo& Info);
diff --git a/Engine/Source/Runtime/AIModule/Classes/Perception/AIPerceptionSystem.h b/Engine/Source/Runtime/AIModule/Classes/Perception/AIPerceptionSystem.h
--- a/Engine/Source/Runtime/AIModule/Classes/Perception/AIPerceptionSystem.h
+++ b/Engine/Source/Runtime/AIModule/Classes/Perception/AIPerceptionSystem.h
@@ -53,11 +53,6 @@ protected:
 	// not a UPROPERTY on purpose so that we have a control over when stuff gets removed from the map
 	TMap<const AActor*, FPerceptionStimuliSource> RegisteredStimuliSources;
 
-#if WITH_EDITORONLY_DATA
-	UE_DEPRECATED(4.25, "This property will be removed in future versions. UnregisterSource is called by AActor.OnEndPlay delegate and will perform the cleanup.")
-	uint32 bStimuliSourcesRefreshRequired : 1;
-#endif
-
 	uint32 bHandlePawnNotification : 1;
 
 	struct FDelayedStimulus
@@ -198,9 +193,6 @@ protected:
 	 *	its stimuli (@see MarkForStimulusProcessing)*/
 	bool AgeStimuli(const float Amount);
 
-	UE_DEPRECATED(4.26, "Parameterless AgeStimuli has been deprecated, please use the other AgeStimuli flavor. To get behavior identical to the old one calls AgeStimuli(PerceptionAgingRate).")
-	void AgeStimuli();
-
 	friend class UAISense;
 	FORCEINLINE AIPerception::FListenerMap& GetListenersMap() { return ListenerContainer; }
 
diff --git a/Engine/Source/Runtime/AIModule/Classes/Perception/AISense.h b/Engine/Source/Runtime/AIModule/Classes/Perception/AISense.h
--- a/Engine/Source/Runtime/AIModule/Classes/Perception/AISense.h
+++ b/Engine/Source/Runtime/AIModule/Classes/Perception/AISense.h
@@ -23,11 +23,6 @@ class AIMODULE_API UAISense : public UObject
 	static const float SuspendNextUpdate;
 
 protected:
-	UE_DEPRECATED(4.23, "This property will be removed in future versions. Use AISenseConfig::MaxAge instead.")
-	/** age past which stimulus of this sense are "forgotten". (DEPRECATED: This property will be removed in future versions. Use AISenseConfig::MaxAge instead.)*/
-	UPROPERTY(VisibleAnywhere, BlueprintReadOnly, Category = "AI Perception", config)
-	float DefaultExpirationAge;
-
 	UPROPERTY(EditDefaultsOnly, BlueprintReadOnly, Category = "AI Perception", config)
 	EAISenseNotifyType NotifyType;
 
@@ -103,9 +98,6 @@ public:
 	virtual void RegisterSource(AActor& SourceActors){}
 	virtual void UnregisterSource(AActor& SourceActors){}
 
-	UE_DEPRECATED(4.25, "This method is no longer used and will be removed in future versions. UnregisterSource is called by AActor.OnEndPlay delegate and will perform the cleanup.")
-	virtual void CleanseInvalidSources() {}
-
 	virtual void RegisterWrappedEvent(UAISenseEvent& PerceptionEvent);
 	virtual FAISenseID UpdateSenseID();
 
@@ -118,14 +110,6 @@ public:
 	FORCEINLINE void OnListenerRemoved(const FPerceptionListener& RemovedListener) { OnListenerRemovedDelegate.ExecuteIfBound(RemovedListener); }
 	virtual void OnListenerConfigUpdated(const FPerceptionListener& UpdatedListener) { OnListenerUpdate(UpdatedListener); }
 
-	UE_DEPRECATED(4.23, "This method will be removed in future versions. Perception relies on AISenseConfig::MaxAge so the value returned is no longer used by the perception system.")
-	FORCEINLINE float GetDefaultExpirationAge() const
-	{
-		PRAGMA_DISABLE_DEPRECATION_WARNINGS
-		return DefaultExpirationAge; 
-		PRAGMA_ENABLE_DEPRECATION_WARNINGS
-	}
-
 	bool WantsNewPawnNotification() const { return bWantsNewPawnNotification; }
 	bool ShouldAutoRegisterAllPawnsAsSources() const { return bAutoRegisterAllPawnsAsSources; }
 
diff --git a/Engine/Source/Runtime/AIModule/Classes/Perception/AISense_Sight.h b/Engine/Source/Runtime/AIModule/Classes/Perception/AISense_Sight.h
--- a/Engine/Source/Runtime/AIModule/Classes/Perception/AISense_Sight.h
+++ b/Engine/Source/Runtime/AIModule/Classes/Perception/AISense_Sight.h
@@ -318,28 +318,10 @@ protected:
 	float CalcQueryImportance(const FPerceptionListener& Listener, const FVector& TargetLocation, const float SightRadiusSq) const;
 	bool RegisterNewQuery(const FPerceptionListener& Listener, const IGenericTeamAgentInterface* ListenersTeamAgent, const AActor& TargetActor, const FAISightTarget::FTargetId& TargetId, const FVector& TargetLocation, const FDigestedSightProperties& PropDigest, const TFunction<void(FAISightQuery&)>& OnAddedFunc);
 
-	// Deprecated methods
-public:
-	UE_DEPRECATED(4.25, "Not needed anymore done automatically at the beginning of each update.")
-	FORCEINLINE void SortQueries() {}
-
 protected:
 	enum FQueriesOperationPostProcess
 	{
 		DontSort,
 		Sort
 	};
-	UE_DEPRECATED(4.25, "Use RemoveAllQueriesByListener without unneeded PostProcess parameter.")
-	void RemoveAllQueriesByListener(const FPerceptionListener& Listener, FQueriesOperationPostProcess PostProcess) { RemoveAllQueriesByListener(Listener); }
-	UE_DEPRECATED(4.25, "Use RemoveAllQueriesByListener without unneeded PostProcess parameter.")
-	void RemoveAllQueriesByListener(const FPerceptionListener& Listener, FQueriesOperationPostProcess PostProcess, TFunctionRef<void(const FAISightQuery&)> OnRemoveFunc) { RemoveAllQueriesByListener(Listener, [&](const FAISightQuery& query) { OnRemoveFunc(query); }); }
-	UE_DEPRECATED(4.25, "Use RemoveAllQueriesToTarget without unneeded PostProcess parameter.")
-	void RemoveAllQueriesToTarget(const FAISightTarget::FTargetId& TargetId, FQueriesOperationPostProcess PostProcess) { RemoveAllQueriesToTarget(TargetId); }
-	UE_DEPRECATED(4.25, "Use RemoveAllQueriesToTarget without unneeded PostProcess parameter.")
-	void RemoveAllQueriesToTarget(const FAISightTarget::FTargetId& TargetId, FQueriesOperationPostProcess PostProcess, TFunctionRef<void(const FAISightQuery&)> OnRemoveFunc) { RemoveAllQueriesToTarget(TargetId, [&](const FAISightQuery& query) { OnRemoveFunc(query); }); }
-	UE_DEPRECATED(4.25, "Use RegisterTarget without unneeded PostProcess parameter.")
-	bool RegisterTarget(AActor& TargetActor, FQueriesOperationPostProcess PostProcess) { return RegisterTarget(TargetActor); }
-	UE_DEPRECATED(4.25, "Use RegisterTarget without unneeded PostProcess parameter.")
-	bool RegisterTarget(AActor& TargetActor, FQueriesOperationPostProcess PostProcess, TFunctionRef<void(FAISightQuery&)> OnAddedFunc) { return RegisterTarget(TargetActor, [&](FAISightQuery& query) { OnAddedFunc(query); }); }
-
 };
diff --git a/Engine/Source/Runtime/AIModule/Classes/Perception/AISightTargetInterface.h b/Engine/Source/Runtime/AIModule/Classes/Perception/AISightTargetInterface.h
--- a/Engine/Source/Runtime/AIModule/Classes/Perception/AISightTargetInterface.h
+++ b/Engine/Source/Runtime/AIModule/Classes/Perception/AISightTargetInterface.h
@@ -70,13 +70,5 @@ class AIMODULE_API IAISightTargetInterface
 		OutSightStrength = 0;
 		return false; 
 	}
-
-	UE_DEPRECATED(4.27, "This function is deprecated. Use the new CanBeSeenFrom method signature")
-	virtual bool CanBeSeenFrom(const FVector& ObserverLocation, FVector& OutSeenLocation, int32& NumberOfLoSChecksPerformed, float& OutSightStrength, const AActor* IgnoreActor = nullptr) const final
-	{
-		NumberOfLoSChecksPerformed = 0;
-		OutSightStrength = 0;
-		return false;
-	}
 };
 
diff --git a/Engine/Source/Runtime/AIModule/Private/BehaviorTree/BehaviorTreeTypes.cpp b/Engine/Source/Runtime/AIModule/Private/BehaviorTree/BehaviorTreeTypes.cpp
--- a/Engine/Source/Runtime/AIModule/Private/BehaviorTree/BehaviorTreeTypes.cpp
+++ b/Engine/Source/Runtime/AIModule/Private/BehaviorTree/BehaviorTreeTypes.cpp
@@ -634,13 +634,6 @@ void FBlackboardKeySelector::AddNameFilter(UObject* Owner, FName PropertyName)
 	AllowedTypes.Add(NewObject<UBlackboardKeyType_Name>(Owner, *FilterName));
 }
 
-// deprecated
-void FBlackboardKeySelector::AddClassFilter(UObject* Owner, FName PropertyName, TSubclassOf<UClass> AllowedClass)
-{
-	TSubclassOf<UObject> AllowedObjectClass = AllowedClass;
-	AddClassFilter(Owner, PropertyName, AllowedObjectClass);
-}
-
 //----------------------------------------------------------------------//
 // UBehaviorTreeTypes
 //----------------------------------------------------------------------//
diff --git a/Engine/Source/Runtime/AIModule/Private/EnvironmentQuery/Tests/EnvQueryTest_Overlap.cpp b/Engine/Source/Runtime/AIModule/Private/EnvironmentQuery/Tests/EnvQueryTest_Overlap.cpp
--- a/Engine/Source/Runtime/AIModule/Private/EnvironmentQuery/Tests/EnvQueryTest_Overlap.cpp
+++ b/Engine/Source/Runtime/AIModule/Private/EnvironmentQuery/Tests/EnvQueryTest_Overlap.cpp
@@ -113,20 +113,6 @@ FText UEnvQueryTest_Overlap::GetDescriptionDetails() const
 	return FText::FromString(FString::Printf(TEXT("Using a %s where %s in channel: %s"), *ShapeDesc, *SizeDesc, *ChannelDesc));
 }
 
-bool UEnvQueryTest_Overlap::RunOverlap(const FVector& ItemPos, const FCollisionShape& CollisionShape, AActor* ItemActor, UWorld* World, enum ECollisionChannel Channel, const FCollisionQueryParams& Params)
-{
-	TArray<AActor*> Actors;
-	Actors.Add(ItemActor);
-	return const_cast<UEnvQueryTest_Overlap*>(this)->RunOverlap(ItemPos, CollisionShape, Actors, World, Channel, Params);
-}
-
-bool UEnvQueryTest_Overlap::RunOverlapBlocking(const FVector& ItemPos, const FCollisionShape& CollisionShape, AActor* ItemActor, UWorld* World, enum ECollisionChannel Channel, const FCollisionQueryParams& Params)
-{
-	TArray<AActor*> Actors;
-	Actors.Add(ItemActor);
-	return const_cast<UEnvQueryTest_Overlap*>(this)->RunOverlapBlocking(ItemPos, CollisionShape, Actors, World, Channel, Params);
-}
-
 bool UEnvQueryTest_Overlap::RunOverlap(const FVector& ItemPos, const FCollisionShape& CollisionShape, const TArray<AActor*>& IgnoredActors, UWorld* World, enum ECollisionChannel Channel, const FCollisionQueryParams& Params) const
 {
 	FCollisionQueryParams OverlapParams(Params);
diff --git a/Engine/Source/Runtime/AIModule/Private/Perception/AIPerceptionComponent.cpp b/Engine/Source/Runtime/AIModule/Private/Perception/AIPerceptionComponent.cpp
--- a/Engine/Source/Runtime/AIModule/Private/Perception/AIPerceptionComponent.cpp
+++ b/Engine/Source/Runtime/AIModule/Private/Perception/AIPerceptionComponent.cpp
@@ -888,12 +888,3 @@ void UAIPerceptionComponent::DescribeSelfToVisLog(FVisualLogEntry* Snapshot) con
 
 }
 #endif // ENABLE_VISUAL_LOG
-
-//----------------------------------------------------------------------//
-// deprecated
-//----------------------------------------------------------------------//
-void UAIPerceptionComponent::GetPerceivedActors(TSubclassOf<UAISense> SenseToUse, TArray<AActor*>& OutActors) const
-{
-	GetCurrentlyPerceivedActors(SenseToUse, OutActors);
-}
-
diff --git a/Engine/Source/Runtime/AIModule/Private/Perception/AIPerceptionSystem.cpp b/Engine/Source/Runtime/AIModule/Private/Perception/AIPerceptionSystem.cpp
--- a/Engine/Source/Runtime/AIModule/Private/Perception/AIPerceptionSystem.cpp
+++ b/Engine/Source/Runtime/AIModule/Private/Perception/AIPerceptionSystem.cpp
@@ -40,12 +40,6 @@ UAIPerceptionSystem::UAIPerceptionSystem(const FObjectInitializer& ObjectInitial
 	, NextStimuliAgingTick(0.)
 	, CurrentTime(0.)
 {
-PRAGMA_DISABLE_DEPRECATION_WARNINGS
-#if WITH_EDITORONLY_DATA
-	bStimuliSourcesRefreshRequired = false;
-#endif
-PRAGMA_ENABLE_DEPRECATION_WARNINGS
-
 	StimuliSourceEndPlayDelegate.BindDynamic(this, &UAIPerceptionSystem::OnPerceptionStimuliSourceEndPlay);
 }
 
@@ -697,11 +691,3 @@ void UAIPerceptionSystem::ReportPerceptionEvent(UObject* WorldContextObject, UAI
 		PerceptionSys->ReportEvent(PerceptionEvent);
 	}
 }
-
-//----------------------------------------------------------------------//
-// Deprecated
-//----------------------------------------------------------------------//
-void UAIPerceptionSystem::AgeStimuli()
-{
-	AgeStimuli(PerceptionAgingRate);
-}
diff --git a/Engine/Source/Runtime/AIModule/Private/Perception/AISense.cpp b/Engine/Source/Runtime/AIModule/Private/Perception/AISense.cpp
--- a/Engine/Source/Runtime/AIModule/Private/Perception/AISense.cpp
+++ b/Engine/Source/Runtime/AIModule/Private/Perception/AISense.cpp
@@ -29,10 +29,6 @@ UAISense::UAISense(const FObjectInitializer& ObjectInitializer)
 	, TimeUntilNextUpdate(SuspendNextUpdate)
 	, SenseID(FAISenseID::InvalidID())
 {
-	PRAGMA_DISABLE_DEPRECATION_WARNINGS
-	DefaultExpirationAge = FAIStimulus::NeverHappenedAge;
-	PRAGMA_ENABLE_DEPRECATION_WARNINGS
-
 	bNeedsForgettingNotification = false;
 
 	if (HasAnyFlags(RF_ClassDefaultObject) == false)
diff --git a/Engine/Source/Runtime/Engine/Classes/AI/Navigation/NavigationTypes.h b/Engine/Source/Runtime/Engine/Classes/AI/Navigation/NavigationTypes.h
--- a/Engine/Source/Runtime/Engine/Classes/AI/Navigation/NavigationTypes.h
+++ b/Engine/Source/Runtime/Engine/Classes/AI/Navigation/NavigationTypes.h
@@ -408,17 +408,6 @@ struct ENGINE_API FNavDataConfig : public FNavAgentProperties
 	UPROPERTY(EditAnywhere, BlueprintReadOnly, Category = Querying)
 	FVector DefaultQueryExtent;
 
-	UE_DEPRECATED(4.24, "FNavDataConfig.NavigationDataClass is deprecated and setting it directly has no effect. Please use setter and getter functions instead.")
-	UPROPERTY(Transient)
-	mutable TSubclassOf<AActor> NavigationDataClass;
-
-#if WITH_EDITOR
-	// used to be a UPROPERTY, but had to remove it so that it doesn't interfere
-	// with property redirects
-	UE_DEPRECATED(4.24, "FNavDataConfig.NavigationDataClassName is deprecated. Please use setter and getter functions instead.")
-	FSoftClassPath NavigationDataClassName;
-#endif // WITH_EDITOR
-
 protected:
 	/** Class to use when spawning navigation data instance */
 	UPROPERTY(EditAnywhere, BlueprintReadWrite, Category = Navigation, meta = (MetaClass = "/Script/NavigationSystem.NavigationData"))
diff --git a/Engine/Source/Runtime/Engine/Classes/AI/NavigationSystemBase.h b/Engine/Source/Runtime/Engine/Classes/AI/NavigationSystemBase.h
--- a/Engine/Source/Runtime/Engine/Classes/AI/NavigationSystemBase.h
+++ b/Engine/Source/Runtime/Engine/Classes/AI/NavigationSystemBase.h
@@ -161,10 +161,6 @@ namespace FNavigationSystem
 	/**	Retrieves the transform the Navigation System is using to convert coords
 	 *	from FromCoordType to ToCoordType */
 	ENGINE_API const FTransform& GetCoordTransform(const ENavigationCoordSystem::Type FromCoordType, const ENavigationCoordSystem::Type ToCoordType);
-	UE_DEPRECATED(4.22, "FNavigationSystem::GetCoordTransformTo is deprecated. Use FNavigationSystem::GetCoordTransform instead")
-	ENGINE_API const FTransform& GetCoordTransformTo(const ENavigationCoordSystem::Type CoordType);
-	UE_DEPRECATED(4.22, "FNavigationSystem::GetCoordTransformFrom is deprecated. Use FNavigationSystem::GetCoordTransform instead")
-	ENGINE_API const FTransform& GetCoordTransformFrom(const ENavigationCoordSystem::Type CoordType);
 
 	ENGINE_API bool WantsComponentChangeNotifies();
 
@@ -249,9 +245,6 @@ public:
 	 */
 	virtual INavigationDataInterface* GetMainNavData() const { return nullptr; }
 
-	UE_DEPRECATED(4.20, "GetMainNavData is deprecated. Use FNavigationSystem::GetCurrent<UNavigationSystemV1>()->GetDefaultNavDataInstance instead")
-	INavigationDataInterface* GetMainNavData(int) { return nullptr; }
-
 	virtual void SetBuildBounds(const FBox& Bounds) PURE_VIRTUAL(UNavigationSystemBase::SetBuildBounds, );
 
 	virtual FBox GetNavigableWorldBounds() const PURE_VIRTUAL(UNavigationSystemBase::GetNavigableWorldBounds, return FBox(ForceInit););
@@ -276,10 +269,6 @@ protected:
 	 *	@param bAddInverse if true (default) will also set coord transform in 
 	 *		the reverse order using Transform.Inverse() */
 	static void SetCoordTransform(const ENavigationCoordSystem::Type FromCoordType, const ENavigationCoordSystem::Type ToCoordType, const FTransform& Transform, bool bAddInverse = true);
-	UE_DEPRECATED(4.22, "FNavigationSystem::SetCoordTransformTo is deprecated. Use FNavigationSystem::SetCoordTransform instead")
-	static void SetCoordTransformTo(const ENavigationCoordSystem::Type CoordType, const FTransform& Transform);
-	UE_DEPRECATED(4.22, "FNavigationSystem::SetCoordTransformFrom is deprecated. Use FNavigationSystem::SetCoordTransform instead")
-	static void SetCoordTransformFrom(const ENavigationCoordSystem::Type CoordType, const FTransform& Transform);
 	static void SetWantsComponentChangeNotifies(const bool bEnable);
 	static void SetDefaultWalkableArea(TSubclassOf<UNavAreaBase> InAreaClass);
 	static void SetDefaultObstacleArea(TSubclassOf<UNavAreaBase> InAreaClass);
diff --git a/Engine/Source/Runtime/Engine/Private/AI/NavigationSystemBase.cpp b/Engine/Source/Runtime/Engine/Private/AI/NavigationSystemBase.cpp
--- a/Engine/Source/Runtime/Engine/Private/AI/NavigationSystemBase.cpp
+++ b/Engine/Source/Runtime/Engine/Private/AI/NavigationSystemBase.cpp
@@ -287,16 +287,6 @@ namespace FNavigationSystem
 		return CoordTypeTransforms;
 	}
 
-	const FTransform& GetCoordTransformTo(const ENavigationCoordSystem::Type CoordType)
-	{
-		return GetCoordTransform(ENavigationCoordSystem::Unreal, CoordType);
-	}
-
-	const FTransform& GetCoordTransformFrom(const ENavigationCoordSystem::Type CoordType)
-	{
-		return GetCoordTransform(CoordType, ENavigationCoordSystem::Unreal);
-	}
-
 	const FTransform& GetCoordTransform(const ENavigationCoordSystem::Type FromCoordType, const ENavigationCoordSystem::Type ToCoordType)
 	{
 		return GetCoordTypeTransforms().Get(FromCoordType, ToCoordType);
@@ -369,16 +359,6 @@ void FNavigationLockContext::UnlockUpdates()
 //----------------------------------------------------------------------//
 // UNavigationSystemBase
 //----------------------------------------------------------------------//
-void UNavigationSystemBase::SetCoordTransformTo(const ENavigationCoordSystem::Type CoordType, const FTransform& Transform)
-{
-	SetCoordTransform(ENavigationCoordSystem::Unreal, CoordType, Transform);
-}
-
-void UNavigationSystemBase::SetCoordTransformFrom(const ENavigationCoordSystem::Type CoordType, const FTransform& Transform)
-{
-	SetCoordTransform(CoordType, ENavigationCoordSystem::Unreal, Transform);
-}
-
 void UNavigationSystemBase::SetCoordTransform(const ENavigationCoordSystem::Type FromCoordType, const ENavigationCoordSystem::Type ToCoordType, const FTransform& Transform, bool bAddInverse)
 {
 	FNavigationSystem::GetCoordTypeTransforms().Get(FromCoordType, ToCoordType) = Transform;
diff --git a/Engine/Source/Runtime/Engine/Public/AI/NavDataGenerator.h b/Engine/Source/Runtime/Engine/Public/AI/NavDataGenerator.h
--- a/Engine/Source/Runtime/Engine/Public/AI/NavDataGenerator.h
+++ b/Engine/Source/Runtime/Engine/Public/AI/NavDataGenerator.h
@@ -32,10 +32,6 @@ public:
 	/** Asks generator to update navigation affected by DirtyAreas */
 	virtual void RebuildDirtyAreas(const TArray<FNavigationDirtyArea>& DirtyAreas) {}
 
-	/** determines whether this generator is performing navigation building actions at the moment*/
-	UE_DEPRECATED(4.26, "This function is deprecated. Please use IsBuildInProgressCheckDirty")
-	virtual bool IsBuildInProgress(bool bCheckDirtyToo = false) const { return IsBuildInProgressCheckDirty(); }
-
 	/** determines whether this generator is performing navigation building actions at the moment, dirty areas are also checked */
 	virtual bool IsBuildInProgressCheckDirty() const { return false; }
 
@@ -57,9 +53,5 @@ public:
 #if ENABLE_VISUAL_LOG
 	virtual void ExportNavigationData(const FString& FileName) const {}
 	virtual void GrabDebugSnapshot(struct FVisualLogEntry* Snapshot, const FBox& BoundingBox, const FName& CategoryName, ELogVerbosity::Type Verbosity) const {}
-
-	UE_DEPRECATED(4.19, "This version of GrabDebugSnapshot has been deprecated. Please use the other version of the function.")
-	virtual void GrabDebugSnapshot(struct FVisualLogEntry* Snapshot, const FBox& BoundingBox, const struct FLogCategoryBase& Category, ELogVerbosity::Type Verbosity) const {}
 #endif
-
 };
diff --git a/Engine/Source/Runtime/NavigationSystem/Private/NavMesh/PImplRecastNavMesh.cpp b/Engine/Source/Runtime/NavigationSystem/Private/NavMesh/PImplRecastNavMesh.cpp
--- a/Engine/Source/Runtime/NavigationSystem/Private/NavMesh/PImplRecastNavMesh.cpp
+++ b/Engine/Source/Runtime/NavigationSystem/Private/NavMesh/PImplRecastNavMesh.cpp
@@ -1124,14 +1124,6 @@ void FPImplRecastNavMesh::Raycast(const FVector& StartLoc, const FVector& EndLoc
 	}
 }
 
-// DEPRECATED
-ENavigationQueryResult::Type FPImplRecastNavMesh::FindPath(const FVector& StartLoc, const FVector& EndLoc, FNavMeshPath& Path, const FNavigationQueryFilter& Filter, const UObject* Owner) const
-{
-	constexpr FVector::FReal CostLimit = TNumericLimits<FVector::FReal>::Max();
-	constexpr bool bRequireNavigableEndLocation = true; 
-	return FindPath(StartLoc, EndLoc, CostLimit, bRequireNavigableEndLocation, Path, Filter, Owner);
-}
-
 // DEPRECATED
 ENavigationQueryResult::Type FPImplRecastNavMesh::FindPath(const FVector& StartLoc, const FVector& EndLoc, const FVector::FReal CostLimit, FNavMeshPath& Path, const FNavigationQueryFilter& InQueryFilter, const UObject* Owner) const
 {
diff --git a/Engine/Source/Runtime/NavigationSystem/Private/NavigationSystem.cpp b/Engine/Source/Runtime/NavigationSystem/Private/NavigationSystem.cpp
--- a/Engine/Source/Runtime/NavigationSystem/Private/NavigationSystem.cpp
+++ b/Engine/Source/Runtime/NavigationSystem/Private/NavigationSystem.cpp
@@ -3946,10 +3946,6 @@ bool UNavigationSystemV1::IsNavigationBuildInProgress()
 }
 
 //deprecated
-bool UNavigationSystemV1::IsNavigationBuildInProgress(bool bCheckDirtyToo)
-{
-	return IsNavigationBuildInProgress();
-}
 
 void UNavigationSystemV1::OnNavigationGenerationFinished(ANavigationData& NavData)
 {
diff --git a/Engine/Source/Runtime/NavigationSystem/Public/NavMesh/PImplRecastNavMesh.h b/Engine/Source/Runtime/NavigationSystem/Public/NavMesh/PImplRecastNavMesh.h
--- a/Engine/Source/Runtime/NavigationSystem/Public/NavMesh/PImplRecastNavMesh.h
+++ b/Engine/Source/Runtime/NavigationSystem/Public/NavMesh/PImplRecastNavMesh.h
@@ -80,15 +80,10 @@ public:
 
 	/** Supported queries */
 
-	// @TODONAV
-	/** Generates path from the given query. Synchronous. */
-	UE_DEPRECATED(4.25, "Please use FindPath with the added CostLimit and bRequireNavigableEndLocation parameters (TNumericLimits<FVector::FReal>::Max() and true can be used as default respectively).")
-	ENavigationQueryResult::Type FindPath(const FVector& StartLoc, const FVector& EndLoc, FNavMeshPath& Path, const FNavigationQueryFilter& Filter, const UObject* Owner) const;
-
 	/** Generates path from the given query. Synchronous. */
 	UE_DEPRECATED(5.2, "Please use FindPath with the added bRequireNavigableEndLocation parameter (true can be used as default).")
 	ENavigationQueryResult::Type FindPath(const FVector& StartLoc, const FVector& EndLoc, const FVector::FReal CostLimit, FNavMeshPath& Path, const FNavigationQueryFilter& Filter, const UObject* Owner) const;
-
+																																									  
 	/** Generates path from the given query. Synchronous. */
 	ENavigationQueryResult::Type FindPath(const FVector& StartLoc, const FVector& EndLoc, const FVector::FReal CostLimit, const bool bRequireNavigableEndLocation, FNavMeshPath& Path, const FNavigationQueryFilter& Filter, const UObject* Owner) const;
 
diff --git a/Engine/Source/Runtime/NavigationSystem/Public/NavMesh/RecastNavMeshGenerator.h b/Engine/Source/Runtime/NavigationSystem/Public/NavMesh/RecastNavMeshGenerator.h
--- a/Engine/Source/Runtime/NavigationSystem/Public/NavMesh/RecastNavMeshGenerator.h
+++ b/Engine/Source/Runtime/NavigationSystem/Public/NavMesh/RecastNavMeshGenerator.h
@@ -66,8 +66,6 @@ struct FRecastBuildConfig : public rcConfig
 	/** chunk size for ChunkyMonotone partitioning */
 	int32 TileCacheChunkSize;
 
-	UE_DEPRECATED(4.24, "FRecastBuildConfig.PolyMaxHeight has been deprecated as it has no use")
-	int32 PolyMaxHeight;
 	/** indicates what's the limit of navmesh polygons per tile. This value is calculated from other
 	 *	factors - DO NOT SET IT TO ARBITRARY VALUE */
 	int32 MaxPolysPerTile;
@@ -99,7 +97,6 @@ struct FRecastBuildConfig : public rcConfig
 		bFilterLowSpanSequences = false;
 		bFilterLowSpanFromTileCache = false;
 		// Still initializing, even though the property is deprecated, to avoid static analysis warnings
-		PolyMaxHeight = 10;
 		MaxPolysPerTile = -1;
 		AgentIndex = 0;
 		TileResolution = ENavigationDataResolution::Default;
diff --git a/Engine/Source/Runtime/NavigationSystem/Public/NavigationSystem.h b/Engine/Source/Runtime/NavigationSystem/Public/NavigationSystem.h
--- a/Engine/Source/Runtime/NavigationSystem/Public/NavigationSystem.h
+++ b/Engine/Source/Runtime/NavigationSystem/Public/NavigationSystem.h
@@ -966,10 +966,6 @@ public:
 	virtual void OnNavigationBoundsAdded(ANavMeshBoundsVolume* NavVolume);
 	virtual void OnNavigationBoundsRemoved(ANavMeshBoundsVolume* NavVolume);
 
-	/** determines whether any generator is performing navigation building actions at the moment */
-	UE_DEPRECATED(4.26, "This function is deprecated.  Please use IsNavigationBuildInProgress")
-	bool IsNavigationBuildInProgress(bool bCheckDirtyToo);
-
 	/** determines whether any generator is performing navigation building actions at the moment, dirty areas are also checked */
 	bool IsNavigationBuildInProgress();
 	
@@ -1317,12 +1313,10 @@ public:
 	//----------------------------------------------------------------------//
 public:
 	// Note that this function was only deprecated for blueprint in 5.1
-	UE_DEPRECATED(4.22, "This version is deprecated.  Please use GetRandomLocationInNavigableRadius instead")
+	// Note2: originally deprecated as 4.22, bumped up to 5.1 as per comment above.
+	UE_DEPRECATED(5.1, "This version is deprecated.  Please use GetRandomLocationInNavigableRadius instead")
 	UFUNCTION(BlueprintPure, Category = "AI|Navigation", meta = (WorldContext = "WorldContextObject", DisplayName = "GetRandomPointInNavigableRadius", ScriptName = "GetRandomPointInNavigableRadius", DeprecatedFunction, DeprecationMessage = "GetRandomPointInNavigableRadius is deprecated. Use GetRandomLocationInNavigableRadius instead"))
 	static bool K2_GetRandomPointInNavigableRadius(UObject* WorldContextObject, const FVector& Origin, FVector& RandomLocation, float Radius, ANavigationData* NavData = NULL, TSubclassOf<UNavigationQueryFilter> FilterClass = NULL);
-	
-	UE_DEPRECATED(4.26, "This version of RemoveNavigationBuildLock is deprecated. Please use the new version")
-	void RemoveNavigationBuildLock(uint8 Flags, bool bSkipRebuildInEditor) { RemoveNavigationBuildLock(Flags, bSkipRebuildInEditor ? ELockRemovalRebuildAction::RebuildIfNotInEditor : ELockRemovalRebuildAction::Rebuild);}
 };
 
 //----------------------------------------------------------------------//
