Added an option to skip EQS volume test, if no volumes are found.
Related to UDN case #00537283

diff --git a/Engine/Source/Runtime/AIModule/Classes/EnvironmentQuery/Tests/EnvQueryTest_Volume.h b/Engine/Source/Runtime/AIModule/Classes/EnvironmentQuery/Tests/EnvQueryTest_Volume.h
--- a/Engine/Source/Runtime/AIModule/Classes/EnvironmentQuery/Tests/EnvQueryTest_Volume.h
+++ b/Engine/Source/Runtime/AIModule/Classes/EnvironmentQuery/Tests/EnvQueryTest_Volume.h
@@ -39,4 +39,8 @@ private:
 	/** If bDoComplexVolumeTest is set, it will use a full volume physic test (not only a bounding box test)  */
 	UPROPERTY(EditAnywhere, Category = Volumes, meta = (AllowPrivateAccess = "true"))
 	uint32 bDoComplexVolumeTest : 1;
+
+	/** If no volumes are returned to evaluate the points against then just ignore the test rather than failing each item */
+	UPROPERTY(EditAnywhere, Category = Volumes, meta = (AllowPrivateAccess = "true"))
+	uint32 bSkipTestIfNoVolumes : 1;
 };
diff --git a/Engine/Source/Runtime/AIModule/Private/EnvironmentQuery/Tests/EnvQueryTest_Volume.cpp b/Engine/Source/Runtime/AIModule/Private/EnvironmentQuery/Tests/EnvQueryTest_Volume.cpp
--- a/Engine/Source/Runtime/AIModule/Private/EnvironmentQuery/Tests/EnvQueryTest_Volume.cpp
+++ b/Engine/Source/Runtime/AIModule/Private/EnvironmentQuery/Tests/EnvQueryTest_Volume.cpp
@@ -16,6 +16,7 @@ UEnvQueryTest_Volume::UEnvQueryTest_Volume(const FObjectInitializer& ObjectIniti
 	SetWorkOnFloatValues(false);
 
 	bDoComplexVolumeTest = false;
+	bSkipTestIfNoVolumes = false;
 }
 
 #if WITH_EDITOR
@@ -83,22 +84,29 @@ void UEnvQueryTest_Volume::RunTest(FEnvQueryInstance& QueryInstance) const
 
 	const bool bWantsInside = BoolValue.GetValue();	
 
-	for (FEnvQueryInstance::ItemIterator It(this, QueryInstance); It; ++It)
+	// Initially Items are marked as "passed", so in principle we only need to run the code below if there are volumes 
+	// to test and we're not allowed to skip this test.
+	// We can however also skip if there are no volumes, but bWantsInside == false since the code below would only result
+	// in assigning "passed" to all Items, which is the default state, as described.
+	if (ContextVolumes.Num() > 0 && (bSkipTestIfNoVolumes == false || bWantsInside == false))
 	{
-		const FVector ItemLocation = GetItemLocation(QueryInstance, It.GetIndex());
-		bool bPointIsInsideAVolume = false;
-		for (const FVolumeBoundsCache& VolumeAndBounds : ContextVolumes)
+		for (FEnvQueryInstance::ItemIterator It(this, QueryInstance); It; ++It)
 		{
-			if (VolumeAndBounds.IsPointInsideBoundingBox(ItemLocation))
+			const FVector ItemLocation = GetItemLocation(QueryInstance, It.GetIndex());
+			bool bPointIsInsideAVolume = false;
+			for (const FVolumeBoundsCache& VolumeAndBounds : ContextVolumes)
 			{
-				if (!bDoComplexVolumeTest || VolumeAndBounds.EncompassesPoint(ItemLocation))
+				if (VolumeAndBounds.IsPointInsideBoundingBox(ItemLocation))
 				{
-					bPointIsInsideAVolume = true;
-					break;
+					if (!bDoComplexVolumeTest || VolumeAndBounds.EncompassesPoint(ItemLocation))
+					{
+						bPointIsInsideAVolume = true;
+						break;
+					}
 				}
-			}			
-		}	
-		It.SetScore(TestPurpose, FilterType, bPointIsInsideAVolume, bWantsInside);
+			}
+			It.SetScore(TestPurpose, FilterType, bPointIsInsideAVolume, bWantsInside);
+		}
 	}
 }
 
