UEnvQueryTest_Volume::RunTest skip condition fix

diff --git a/Engine/Source/Runtime/AIModule/Private/EnvironmentQuery/Tests/EnvQueryTest_Volume.cpp b/Engine/Source/Runtime/AIModule/Private/EnvironmentQuery/Tests/EnvQueryTest_Volume.cpp
--- a/Engine/Source/Runtime/AIModule/Private/EnvironmentQuery/Tests/EnvQueryTest_Volume.cpp
+++ b/Engine/Source/Runtime/AIModule/Private/EnvironmentQuery/Tests/EnvQueryTest_Volume.cpp
@@ -88,7 +88,7 @@ void UEnvQueryTest_Volume::RunTest(FEnvQueryInstance& QueryInstance) const
 	// to test and we're not allowed to skip this test.
 	// We can however also skip if there are no volumes, but bWantsInside == false since the code below would only result
 	// in assigning "passed" to all Items, which is the default state, as described.
-	if (ContextVolumes.Num() > 0 && (bSkipTestIfNoVolumes == false || bWantsInside == false))
+	if (ContextVolumes.Num() > 0 || (bSkipTestIfNoVolumes == false && bWantsInside == true))
 	{
 		for (FEnvQueryInstance::ItemIterator It(this, QueryInstance); It; ++It)
 		{
