Fix crash in Editor if BaseClass of UBlackboardKeyType_Object in KeyType is set to None.

diff --git a/Engine/Source/Runtime/AIModule/Classes/BehaviorTree/Blackboard/BlackboardKeyType.h b/Engine/Source/Runtime/AIModule/Classes/BehaviorTree/Blackboard/BlackboardKeyType.h
--- a/Engine/Source/Runtime/AIModule/Classes/BehaviorTree/Blackboard/BlackboardKeyType.h
+++ b/Engine/Source/Runtime/AIModule/Classes/BehaviorTree/Blackboard/BlackboardKeyType.h
@@ -9,6 +9,8 @@
 #include "BlackboardKeyEnums.h"
 #include "BlackboardKeyType.generated.h"
 
+DECLARE_LOG_CATEGORY_EXTERN(LogBlackboard, Log, All);
+
 class UBlackboardComponent;
 
 struct FBlackboardInstancedKeyMemory
diff --git a/Engine/Source/Runtime/AIModule/Classes/BehaviorTree/Blackboard/BlackboardKeyType_Object.h b/Engine/Source/Runtime/AIModule/Classes/BehaviorTree/Blackboard/BlackboardKeyType_Object.h
--- a/Engine/Source/Runtime/AIModule/Classes/BehaviorTree/Blackboard/BlackboardKeyType_Object.h
+++ b/Engine/Source/Runtime/AIModule/Classes/BehaviorTree/Blackboard/BlackboardKeyType_Object.h
@@ -16,7 +16,7 @@ class AIMODULE_API UBlackboardKeyType_Object : public UBlackboardKeyType
 	typedef UObject* FDataType;
 	static const FDataType InvalidValue;
 
-	UPROPERTY(Category=Blackboard, EditDefaultsOnly, meta=(AllowAbstract="1"))
+	UPROPERTY(Category=Blackboard, EditDefaultsOnly, NoClear, meta=(AllowAbstract="1"))
 	TObjectPtr<UClass> BaseClass;
 
 	static UObject* GetValue(const UBlackboardKeyType_Object* KeyOb, const uint8* RawData);
diff --git a/Engine/Source/Runtime/AIModule/Private/BehaviorTree/Blackboard/BlackboardKeyType_Object.cpp b/Engine/Source/Runtime/AIModule/Private/BehaviorTree/Blackboard/BlackboardKeyType_Object.cpp
--- a/Engine/Source/Runtime/AIModule/Private/BehaviorTree/Blackboard/BlackboardKeyType_Object.cpp
+++ b/Engine/Source/Runtime/AIModule/Private/BehaviorTree/Blackboard/BlackboardKeyType_Object.cpp
@@ -58,7 +58,8 @@ FString UBlackboardKeyType_Object::DescribeSelf() const
 bool UBlackboardKeyType_Object::IsAllowedByFilter(UBlackboardKeyType* FilterOb) const
 {
 	UBlackboardKeyType_Object* FilterObject = Cast<UBlackboardKeyType_Object>(FilterOb);
-	return (FilterObject && (FilterObject->BaseClass == BaseClass || BaseClass->IsChildOf(FilterObject->BaseClass)));
+	UE_CLOG(BaseClass == nullptr, LogBlackboard, Warning, TEXT("Cannot assign %s to a Blackboard Key with null base class. Should the base class be changed?"), *FilterObject->DescribeSelf());
+	return (FilterObject && (FilterObject->BaseClass == BaseClass || (BaseClass != nullptr && BaseClass->IsChildOf(FilterObject->BaseClass))));
 }
 
 bool UBlackboardKeyType_Object::GetLocation(const UBlackboardComponent& OwnerComp, const uint8* RawData, FVector& Location) const
diff --git a/Engine/Source/Runtime/AIModule/Private/BehaviorTree/BlackboardComponent.cpp b/Engine/Source/Runtime/AIModule/Private/BehaviorTree/BlackboardComponent.cpp
--- a/Engine/Source/Runtime/AIModule/Private/BehaviorTree/BlackboardComponent.cpp
+++ b/Engine/Source/Runtime/AIModule/Private/BehaviorTree/BlackboardComponent.cpp
@@ -19,7 +19,7 @@
 
 #include UE_INLINE_GENERATED_CPP_BY_NAME(BlackboardComponent)
 
-DEFINE_LOG_CATEGORY_STATIC(LogBlackboard, Log, All);
+DEFINE_LOG_CATEGORY(LogBlackboard);
 
 UBlackboardComponent::UBlackboardComponent(const FObjectInitializer& ObjectInitializer) : Super(ObjectInitializer)
 {
