Fix Horde static analysis warning for checking FilterObject pointer

diff --git a/Engine/Source/Runtime/AIModule/Private/BehaviorTree/Blackboard/BlackboardKeyType_Object.cpp b/Engine/Source/Runtime/AIModule/Private/BehaviorTree/Blackboard/BlackboardKeyType_Object.cpp
--- a/Engine/Source/Runtime/AIModule/Private/BehaviorTree/Blackboard/BlackboardKeyType_Object.cpp
+++ b/Engine/Source/Runtime/AIModule/Private/BehaviorTree/Blackboard/BlackboardKeyType_Object.cpp
@@ -58,7 +58,7 @@ FString UBlackboardKeyType_Object::DescribeSelf() const
 bool UBlackboardKeyType_Object::IsAllowedByFilter(UBlackboardKeyType* FilterOb) const
 {
 	UBlackboardKeyType_Object* FilterObject = Cast<UBlackboardKeyType_Object>(FilterOb);
-	UE_CLOG(BaseClass == nullptr, LogBlackboard, Warning, TEXT("Cannot assign %s to a Blackboard Key with null base class. Should the base class be changed?"), *FilterObject->DescribeSelf());
+	UE_CLOG(BaseClass == nullptr, LogBlackboard, Warning, TEXT("Cannot assign %s to a Blackboard Key with null base class. Should the base class be changed?"), FilterObject ? *FilterObject->DescribeSelf() : TEXT("NULL"));
 	return (FilterObject && (FilterObject->BaseClass == BaseClass || (BaseClass != nullptr && BaseClass->IsChildOf(FilterObject->BaseClass))));
 }
 
