Add option to select which nav data to use for EQS in LWM

diff --git a/Engine/Source/Runtime/AIModule/Private/EnvironmentQuery/EnvQueryTypes.cpp b/Engine/Source/Runtime/AIModule/Private/EnvironmentQuery/EnvQueryTypes.cpp
--- a/Engine/Source/Runtime/AIModule/Private/EnvironmentQuery/EnvQueryTypes.cpp
+++ b/Engine/Source/Runtime/AIModule/Private/EnvironmentQuery/EnvQueryTypes.cpp
@@ -278,18 +278,19 @@ namespace FEQSHelpers
 	const ANavigationData* FindNavigationDataForQuery(FEnvQueryInstance& QueryInstance)
 	{
 		const UNavigationSystemV1* NavSys = FNavigationSystem::GetCurrent<UNavigationSystemV1>(QueryInstance.World);
-		
 		if (NavSys == nullptr)
 		{
 			return nullptr;
 		}
 
 		// try to match navigation agent for querier
-		INavAgentInterface* NavAgent = QueryInstance.Owner.IsValid() ? Cast<INavAgentInterface>(QueryInstance.Owner.Get()) : NULL;
-		if (NavAgent)
+		if (const INavAgentInterface* NavAgent = QueryInstance.Owner.IsValid() ? Cast<INavAgentInterface>(QueryInstance.Owner.Get()) : nullptr)
 		{
 			const FNavAgentProperties& NavAgentProps = NavAgent->GetNavAgentPropertiesRef();
-			return NavSys->GetNavDataForProps(NavAgentProps, NavAgent->GetNavAgentLocation());
+			if (NavAgentProps.IsValid() || NavAgentProps.PreferredNavData.IsValid())
+			{
+				return NavSys->GetNavDataForProps(NavAgentProps, NavAgent->GetNavAgentLocation());
+			}
 		}
 
 		return NavSys->GetDefaultNavDataInstance();
