PR #10250: Hierarchicalhashgrid2d Remove function fix (Contributed by [REDACTED_PERSON])

diff --git a/Engine/Source/Runtime/AIModule/Public/HierarchicalHashGrid2D.h b/Engine/Source/Runtime/AIModule/Public/HierarchicalHashGrid2D.h
--- a/Engine/Source/Runtime/AIModule/Public/HierarchicalHashGrid2D.h
+++ b/Engine/Source/Runtime/AIModule/Public/HierarchicalHashGrid2D.h
@@ -227,7 +227,7 @@ public:
 	void Remove(const ItemIDType ID, const FBox& OldBounds)
 	{
 		const FCellLocation OldLocation = CalcCellLocation(OldBounds);
-		Remove(OldLocation, ID);
+		Remove(ID, OldLocation);
 	}
 
 	/** Removes item based on the cell location it was added with.
