Fix bIsObservingBB not being set correctly upon blueprint compilation. The BP reinstantiation code is restoring the old values over the newly calculated one.
To fix this, removed the fact that the boolean is a UPROPERTY so that the resintantiator does not restore the old value and upon instantiation, we are now getching the value from the CDO.

diff --git a/Engine/Source/Runtime/AIModule/Classes/BehaviorTree/Decorators/BTDecorator_BlueprintBase.h b/Engine/Source/Runtime/AIModule/Classes/BehaviorTree/Decorators/BTDecorator_BlueprintBase.h
--- a/Engine/Source/Runtime/AIModule/Classes/BehaviorTree/Decorators/BTDecorator_BlueprintBase.h
+++ b/Engine/Source/Runtime/AIModule/Classes/BehaviorTree/Decorators/BTDecorator_BlueprintBase.h
@@ -109,7 +109,6 @@ protected:
 	uint32 bCheckConditionOnlyBlackBoardChanges : 1;
 
 	/** gets set to true if decorator declared BB keys it can potentially observe */
-	UPROPERTY()
 	uint32 bIsObservingBB : 1;
 	
 	/** set if ReceiveTick is implemented by blueprint */
diff --git a/Engine/Source/Runtime/AIModule/Private/BehaviorTree/Decorators/BTDecorator_BlueprintBase.cpp b/Engine/Source/Runtime/AIModule/Private/BehaviorTree/Decorators/BTDecorator_BlueprintBase.cpp
--- a/Engine/Source/Runtime/AIModule/Private/BehaviorTree/Decorators/BTDecorator_BlueprintBase.cpp
+++ b/Engine/Source/Runtime/AIModule/Private/BehaviorTree/Decorators/BTDecorator_BlueprintBase.cpp
@@ -43,6 +43,11 @@ void UBTDecorator_BlueprintBase::InitializeProperties()
 
 		bIsObservingBB = BlueprintNodeHelpers::HasAnyBlackboardSelectors(this, StopAtClass);
 	}
+	else
+	{
+		// Make sure if the users started to observe a BB key after we instantiated this node, that this value get propagated correctly
+		bIsObservingBB = GetClass()->GetDefaultObject<UBTDecorator_BlueprintBase>()->bIsObservingBB;
+	}
 }
 
 void UBTDecorator_BlueprintBase::PostInitProperties()
