AI Sense:
* Fix compilation of SIGHT_LOG_SEGMENT call
* Improve comments readability for listener delegates

#rnx

diff --git a/Engine/Source/Runtime/AIModule/Classes/Perception/AISense.h b/Engine/Source/Runtime/AIModule/Classes/Perception/AISense.h
--- a/Engine/Source/Runtime/AIModule/Classes/Perception/AISense.h
+++ b/Engine/Source/Runtime/AIModule/Classes/Perception/AISense.h
@@ -49,13 +49,13 @@ private:
 	FAISenseID SenseID;
 
 protected:
-	/**	If bound will be called when new FPerceptionListener gets registers with AIPerceptionSystem */
+	/**	Called when a new FPerceptionListener registers with AIPerceptionSystem */
 	FOnPerceptionListenerUpdateDelegate OnNewListenerDelegate;
 
-	/**	If bound will be called when a FPerceptionListener's in AIPerceptionSystem change */
+	/**	Called when a registered FPerceptionListener changes */
 	FOnPerceptionListenerUpdateDelegate OnListenerUpdateDelegate;
 
-	/**	If bound will be called when a FPerceptionListener's in removed from AIPerceptionSystem */
+	/**	Called when a FPerceptionListener is removed from AIPerceptionSystem */
 	FOnPerceptionListenerUpdateDelegate OnListenerRemovedDelegate;
 				
 public:
diff --git a/Engine/Source/Runtime/AIModule/Classes/Perception/AISense_Sight.h b/Engine/Source/Runtime/AIModule/Classes/Perception/AISense_Sight.h
--- a/Engine/Source/Runtime/AIModule/Classes/Perception/AISense_Sight.h
+++ b/Engine/Source/Runtime/AIModule/Classes/Perception/AISense_Sight.h
@@ -297,7 +297,11 @@ protected:
 
 	EVisibilityResult ComputeVisibility(UWorld* World, FAISightQuery& SightQuery, FPerceptionListener& Listener, const AActor* ListenerActor, FAISightTarget& Target, AActor* TargetActor, const FDigestedSightProperties& PropDigest, float& OutStimulusStrength, FVector& OutSeenLocation, int32& OutNumberOfLoSChecksPerformed, int32& OutNumberOfAsyncLosCheckRequested) const;
 	virtual bool ShouldAutomaticallySeeTarget(const FDigestedSightProperties& PropDigest, FAISightQuery* SightQuery, FPerceptionListener& Listener, AActor* TargetActor, float& OutStimulusStrength) const;
+	
+	UE_DEPRECATED(5.3, "Please use the UpdateQueryVisibilityStatus version which takes an Actor& instead.")
 	void UpdateQueryVisibilityStatus(FAISightQuery& SightQuery, FPerceptionListener& Listener, const bool bIsVisible, const FVector& SeenLocation, const float StimulusStrength, AActor* TargetActor, const FVector& TargetLocation) const;
+	void UpdateQueryVisibilityStatus(FAISightQuery& SightQuery, FPerceptionListener& Listener, const bool bIsVisible, const FVector& SeenLocation, const float StimulusStrength, AActor& TargetActor, const FVector& TargetLocation) const;
+
 	void OnPendingCanBeSeenQueryProcessed(const FAISightQueryID& QueryID, const bool bIsVisible, const float StimulusStrength, const FVector& SeenLocation, const TOptional<int32>& UserData);
 	void OnPendingTraceQueryProcessed(const FTraceHandle& TraceHandle, FTraceDatum& TraceDatum);
 	void OnPendingQueryProcessed(const int32 SightQueryIndex, const bool bIsVisible, const float StimulusStrength, const FVector& SeenLocation, const TOptional<int32>& UserData, const TOptional<AActor*> InTargetActor = NullOpt);
diff --git a/Engine/Source/Runtime/AIModule/Private/Perception/AISense_Sight.cpp b/Engine/Source/Runtime/AIModule/Private/Perception/AISense_Sight.cpp
--- a/Engine/Source/Runtime/AIModule/Private/Perception/AISense_Sight.cpp
+++ b/Engine/Source/Runtime/AIModule/Private/Perception/AISense_Sight.cpp
@@ -382,7 +382,7 @@ float UAISense_Sight::Update()
 				const bool bIsVisible = VisibilityResult == EVisibilityResult::Visible;
 				const bool bWasVisible = SightQuery->GetLastResult();
 				const FVector TargetLocation = TargetActor->GetActorLocation();
-				UpdateQueryVisibilityStatus(*SightQuery, Listener, bIsVisible, SeenLocation, StimulusStrength, TargetActor, TargetLocation);
+				UpdateQueryVisibilityStatus(*SightQuery, Listener, bIsVisible, SeenLocation, StimulusStrength, *TargetActor, TargetLocation);
 
 				const float SightRadiusSq = bWasVisible ? PropDigest.LoseSightRadiusSq : PropDigest.SightRadiusSq;
 				SightQuery->Importance = CalcQueryImportance(Listener, TargetLocation, SightRadiusSq);
@@ -577,11 +577,19 @@ UAISense_Sight::EVisibilityResult UAISense_Sight::ComputeVisibility(UWorld* Worl
 }
 
 void UAISense_Sight::UpdateQueryVisibilityStatus(FAISightQuery& SightQuery, FPerceptionListener& Listener, const bool bIsVisible, const FVector& SeenLocation, const float StimulusStrength, AActor* TargetActor, const FVector& TargetLocation) const
+{
+	if (TargetActor)
+	{
+		UpdateQueryVisibilityStatus(SightQuery, Listener, bIsVisible, SeenLocation, StimulusStrength, *TargetActor, TargetLocation);
+	}
+}
+
+void UAISense_Sight::UpdateQueryVisibilityStatus(FAISightQuery& SightQuery, FPerceptionListener& Listener, const bool bIsVisible, const FVector& SeenLocation, const float StimulusStrength, AActor& TargetActor, const FVector& TargetLocation) const
 {
 	if (bIsVisible)
 	{
 		const bool bHasValidSeenLocation = SeenLocation != FAISystem::InvalidLocation;
-		Listener.RegisterStimulus(TargetActor, FAIStimulus(*this, StimulusStrength, bHasValidSeenLocation ? SeenLocation : SightQuery.LastSeenLocation, Listener.CachedLocation));
+		Listener.RegisterStimulus(&TargetActor, FAIStimulus(*this, StimulusStrength, bHasValidSeenLocation ? SeenLocation : SightQuery.LastSeenLocation, Listener.CachedLocation));
 		SightQuery.SetLastResult(true);
 		if (bHasValidSeenLocation)
 		{
@@ -591,12 +599,12 @@ void UAISense_Sight::UpdateQueryVisibilityStatus(FAISightQuery& SightQuery, FPer
 	// communicate failure only if we've seen given actor before
 	else if (SightQuery.GetLastResult())
 	{
-		Listener.RegisterStimulus(TargetActor, FAIStimulus(*this, 0.f, TargetLocation, Listener.CachedLocation, FAIStimulus::SensingFailed));
+		Listener.RegisterStimulus(&TargetActor, FAIStimulus(*this, 0.f, TargetLocation, Listener.CachedLocation, FAIStimulus::SensingFailed));
 		SightQuery.SetLastResult(false);
 		SightQuery.LastSeenLocation = FAISystem::InvalidLocation;
 	}
 
-	SIGHT_LOG_SEGMENT(ListenerPtr->GetOwner(), Listener.CachedLocation, TargetLocation, bIsVisible ? FColor::Green : FColor::Red, TEXT("TargetID %d"), Target.TargetId);
+	SIGHT_LOG_SEGMENT(Listener.GetBodyActor(), Listener.CachedLocation, TargetLocation, bIsVisible ? FColor::Green : FColor::Red, TEXT("Target: %s"), *TargetActor.GetName());
 }
 
 void UAISense_Sight::OnPendingCanBeSeenQueryProcessed(const FAISightQueryID& QueryID, const bool bIsVisible, const float StimulusStrength, const FVector& SeenLocation, const TOptional<int32>& UserData)
@@ -677,7 +685,7 @@ void UAISense_Sight::OnPendingQueryProcessed(const int32 SightQueryIndex, const
 
 	const bool bWasVisible = SightQuery.GetLastResult();
 	const FVector TargetLocation = TargetActor->GetActorLocation();
-	UpdateQueryVisibilityStatus(SightQuery, *Listener, bIsVisible, SeenLocation, StimulusStrength, TargetActor, TargetLocation);
+	UpdateQueryVisibilityStatus(SightQuery, *Listener, bIsVisible, SeenLocation, StimulusStrength, *TargetActor, TargetLocation);
 
 	if (UserData.IsSet())
 	{
