Fixed mismatched copy constructor/assignment operator in FSequentialIDBase.

diff --git a/Engine/Source/Runtime/AIModule/Public/SequentialID.h b/Engine/Source/Runtime/AIModule/Public/SequentialID.h
--- a/Engine/Source/Runtime/AIModule/Public/SequentialID.h
+++ b/Engine/Source/Runtime/AIModule/Public/SequentialID.h
@@ -10,7 +10,6 @@ struct AIMODULE_API FSequentialIDBase
 	GENERATED_BODY()
 
 	FSequentialIDBase() = default;
-	FSequentialIDBase(const FSequentialIDBase& Other) = default;
 
 	explicit FSequentialIDBase(uint32 InID) : Value(InID) {}
 
