Fixed a bug where asynchronous sight queries were not getting properly removed when the source of the sight query was unregistering.

#rnx

diff --git a/Engine/Source/Runtime/AIModule/Private/Perception/AISense_Sight.cpp b/Engine/Source/Runtime/AIModule/Private/Perception/AISense_Sight.cpp
--- a/Engine/Source/Runtime/AIModule/Private/Perception/AISense_Sight.cpp
+++ b/Engine/Source/Runtime/AIModule/Private/Perception/AISense_Sight.cpp
@@ -735,7 +735,7 @@ void UAISense_Sight::UnregisterSource(AActor& SourceActor)
 	FAISightTarget AsTarget;
 
 	if (ObservedTargets.RemoveAndCopyValue(AsTargetId, AsTarget) 
-		&& (SightQueriesInRange.Num() + SightQueriesOutOfRange.Num()) > 0)
+		&& (SightQueriesInRange.Num() + SightQueriesOutOfRange.Num() + SightQueriesPending.Num()) > 0)
 	{
 		AActor* TargetActor = AsTarget.Target.Get();
 
