Fix crash in AAIController: Check that path following component is not null
#rnx

diff --git a/Engine/Source/Runtime/AIModule/Private/AIController.cpp b/Engine/Source/Runtime/AIModule/Private/AIController.cpp
--- a/Engine/Source/Runtime/AIModule/Private/AIController.cpp
+++ b/Engine/Source/Runtime/AIModule/Private/AIController.cpp
@@ -744,12 +744,19 @@ void AAIController::StopMovement()
 {
 	// @note FPathFollowingResultFlags::ForcedScript added to make AITask_MoveTo instances 
 	// not ignore OnRequestFinished notify that's going to be sent out due to this call
-	PathFollowingComponent->AbortMove(*this, FPathFollowingResultFlags::MovementStop | FPathFollowingResultFlags::ForcedScript);
+	if (PathFollowingComponent)
+	{
+	    PathFollowingComponent->AbortMove(*this, FPathFollowingResultFlags::MovementStop | FPathFollowingResultFlags::ForcedScript);
+	}
 }
 
 bool AAIController::ShouldPostponePathUpdates() const
 {
-	return GetPathFollowingComponent()->HasStartedNavLinkMove() || Super::ShouldPostponePathUpdates();
+	if (PathFollowingComponent && PathFollowingComponent->HasStartedNavLinkMove())
+	{
+		return true;
+	}
+	return Super::ShouldPostponePathUpdates();
 }
 
 bool AAIController::BuildPathfindingQuery(const FAIMoveRequest& MoveRequest, FPathFindingQuery& Query) const
