Add Perception System gameplay debugger category to be able to track the perception usage

diff --git a/Engine/Source/Runtime/AIModule/Classes/Perception/AIPerceptionSystem.h b/Engine/Source/Runtime/AIModule/Classes/Perception/AIPerceptionSystem.h
--- a/Engine/Source/Runtime/AIModule/Classes/Perception/AIPerceptionSystem.h
+++ b/Engine/Source/Runtime/AIModule/Classes/Perception/AIPerceptionSystem.h
@@ -15,6 +15,7 @@
 #include "AISubsystem.h"
 #include "AIPerceptionSystem.generated.h"
 
+class FGameplayDebuggerCategory;
 class UAIPerceptionComponent;
 class UAISenseEvent;
 
@@ -163,7 +164,11 @@ public:
 
 	UFUNCTION(BlueprintCallable, Category = "AI|Perception", meta = (WorldContext="WorldContextObject"))
 	static TSubclassOf<UAISense> GetSenseClassForStimulus(UObject* WorldContextObject, const FAIStimulus& Stimulus);
-	
+
+#if WITH_GAMEPLAY_DEBUGGER_MENU
+	virtual void DescribeSelfToGameplayDebugger(FGameplayDebuggerCategory& DebuggerCategory) const;
+#endif // WITH_GAMEPLAY_DEBUGGER_MENU
+
 protected:
 	
 	UFUNCTION()
diff --git a/Engine/Source/Runtime/AIModule/Classes/Perception/AISense.h b/Engine/Source/Runtime/AIModule/Classes/Perception/AISense.h
--- a/Engine/Source/Runtime/AIModule/Classes/Perception/AISense.h
+++ b/Engine/Source/Runtime/AIModule/Classes/Perception/AISense.h
@@ -10,6 +10,7 @@
 #include "AISense.generated.h"
 
 class APawn;
+class FGameplayDebuggerCategory;
 class UAIPerceptionSystem;
 class UAISenseEvent;
 
@@ -113,6 +114,10 @@ public:
 	bool WantsNewPawnNotification() const { return bWantsNewPawnNotification; }
 	bool ShouldAutoRegisterAllPawnsAsSources() const { return bAutoRegisterAllPawnsAsSources; }
 
+#if WITH_GAMEPLAY_DEBUGGER_MENU
+	virtual void DescribeSelfToGameplayDebugger(const UAIPerceptionSystem& PerceptionSystem, FGameplayDebuggerCategory& DebuggerCategory) const;
+#endif // WITH_GAMEPLAY_DEBUGGER_MENU
+
 protected:
 	friend UAIPerceptionSystem;
 	/** gets called when perception system gets notified about new spawned pawn. 
diff --git a/Engine/Source/Runtime/AIModule/Classes/Perception/AISense_Sight.h b/Engine/Source/Runtime/AIModule/Classes/Perception/AISense_Sight.h
--- a/Engine/Source/Runtime/AIModule/Classes/Perception/AISense_Sight.h
+++ b/Engine/Source/Runtime/AIModule/Classes/Perception/AISense_Sight.h
@@ -292,6 +292,10 @@ public:
 	virtual void OnListenerForgetsActor(const FPerceptionListener& Listener, AActor& ActorToForget) override;
 	virtual void OnListenerForgetsAll(const FPerceptionListener& Listener) override;
 
+#if WITH_GAMEPLAY_DEBUGGER_MENU
+	virtual void DescribeSelfToGameplayDebugger(const UAIPerceptionSystem& PerceptionSystem, FGameplayDebuggerCategory& DebuggerCategory) const override;
+#endif // WITH_GAMEPLAY_DEBUGGER_MENU
+
 protected:
 	virtual float Update() override;
 
diff --git a/Engine/Source/Runtime/AIModule/Private/AIModule.cpp b/Engine/Source/Runtime/AIModule/Private/AIModule.cpp
--- a/Engine/Source/Runtime/AIModule/Private/AIModule.cpp
+++ b/Engine/Source/Runtime/AIModule/Private/AIModule.cpp
@@ -25,6 +25,7 @@
 #include "GameplayDebugger/GameplayDebuggerCategory_EQS.h"
 #include "GameplayDebugger/GameplayDebuggerCategory_Navmesh.h"
 #include "GameplayDebugger/GameplayDebuggerCategory_Perception.h"
+#include "GameplayDebugger/GameplayDebuggerCategory_PerceptionSystem.h"
 #include "GameplayDebugger/GameplayDebuggerCategory_NavLocalGrid.h"
 #endif // WITH_GAMEPLAY_DEBUGGER
 
@@ -80,6 +81,7 @@ void FAIModule::StartupModule()
 	GameplayDebuggerModule.RegisterCategory("EQS", IGameplayDebugger::FOnGetCategory::CreateStatic(&FGameplayDebuggerCategory_EQS::MakeInstance));
 	GameplayDebuggerModule.RegisterCategory("Navmesh", IGameplayDebugger::FOnGetCategory::CreateStatic(&FGameplayDebuggerCategory_Navmesh::MakeInstance), EGameplayDebuggerCategoryState::Disabled, 0);
 	GameplayDebuggerModule.RegisterCategory("Perception", IGameplayDebugger::FOnGetCategory::CreateStatic(&FGameplayDebuggerCategory_Perception::MakeInstance));
+	GameplayDebuggerModule.RegisterCategory("PerceptionSystem", IGameplayDebugger::FOnGetCategory::CreateStatic(&FGameplayDebuggerCategory_PerceptionSystem::MakeInstance));
 	GameplayDebuggerModule.RegisterCategory("NavGrid", IGameplayDebugger::FOnGetCategory::CreateStatic(&FGameplayDebuggerCategory_NavLocalGrid::MakeInstance), EGameplayDebuggerCategoryState::Hidden);
 	GameplayDebuggerModule.NotifyCategoriesChanged();
 #endif // WITH_GAMEPLAY_DEBUGGER
diff --git a/Engine/Source/Runtime/AIModule/Private/GameplayDebugger/GameplayDebuggerCategory_PerceptionSystem.cpp b/Engine/Source/Runtime/AIModule/Private/GameplayDebugger/GameplayDebuggerCategory_PerceptionSystem.cpp
new file mode 100644
--- /dev/null
+++ b/Engine/Source/Runtime/AIModule/Private/GameplayDebugger/GameplayDebuggerCategory_PerceptionSystem.cpp
@@ -0,0 +1,32 @@
+// Copyright Epic Games, Inc. All Rights Reserved.
+
+#include "GameplayDebugger/GameplayDebuggerCategory_PerceptionSystem.h"
+
+#if WITH_GAMEPLAY_DEBUGGER_MENU
+
+#include "GameFramework/PlayerController.h"
+#include "Perception/AIPerceptionSystem.h"
+
+FGameplayDebuggerCategory_PerceptionSystem::FGameplayDebuggerCategory_PerceptionSystem()
+{
+	bShowOnlyWithDebugActor = false;
+}
+
+TSharedRef<FGameplayDebuggerCategory> FGameplayDebuggerCategory_PerceptionSystem::MakeInstance()
+{
+	return MakeShareable(new FGameplayDebuggerCategory_PerceptionSystem());
+}
+
+void FGameplayDebuggerCategory_PerceptionSystem::CollectData(APlayerController* OwnerPC, AActor* DebugActor)
+{
+	UWorld* World = GetDataWorld(OwnerPC, DebugActor);
+	check(World);
+
+	const UAIPerceptionSystem* PerceptionSystem = UAIPerceptionSystem::GetCurrent(*World);
+	if (PerceptionSystem != nullptr)
+	{
+		PerceptionSystem->DescribeSelfToGameplayDebugger(*this);
+	}
+}
+
+#endif // WITH_GAMEPLAY_DEBUGGER_MENU
diff --git a/Engine/Source/Runtime/AIModule/Private/Perception/AIPerceptionSystem.cpp b/Engine/Source/Runtime/AIModule/Private/Perception/AIPerceptionSystem.cpp
--- a/Engine/Source/Runtime/AIModule/Private/Perception/AIPerceptionSystem.cpp
+++ b/Engine/Source/Runtime/AIModule/Private/Perception/AIPerceptionSystem.cpp
@@ -1,18 +1,20 @@
 // Copyright Epic Games, Inc. All Rights Reserved.
 
 #include "Perception/AIPerceptionSystem.h"
+
+#include "AISystem.h"
+#include "Engine/Engine.h"
 #include "EngineGlobals.h"
 #include "EngineUtils.h"
 #include "GameFramework/Pawn.h"
-#include "TimerManager.h"
-#include "Engine/Engine.h"
-#include "AISystem.h"
-#include "Perception/AISense_Hearing.h"
-#include "Perception/AISenseConfig.h"
+#include "GameplayDebuggerCategory.h"
 #include "Perception/AIPerceptionComponent.h"
-#include "VisualLogger/VisualLogger.h"
-#include "ProfilingDebugging/CsvProfiler.h"
+#include "Perception/AISenseConfig.h"
 #include "Perception/AISenseEvent.h"
+#include "Perception/AISense_Hearing.h"
+#include "ProfilingDebugging/CsvProfiler.h"
+#include "TimerManager.h"
+#include "VisualLogger/VisualLogger.h"
 
 #include UE_INLINE_GENERATED_CPP_BY_NAME(AIPerceptionSystem)
 
@@ -691,3 +693,21 @@ void UAIPerceptionSystem::ReportPerceptionEvent(UObject* WorldContextObject, UAI
 		PerceptionSys->ReportEvent(PerceptionEvent);
 	}
 }
+
+//----------------------------------------------------------------------//
+// Debug
+//----------------------------------------------------------------------//
+#if WITH_GAMEPLAY_DEBUGGER_MENU
+void UAIPerceptionSystem::DescribeSelfToGameplayDebugger(FGameplayDebuggerCategory& DebuggerCategory) const
+{
+	DebuggerCategory.AddTextLine(FString::Printf(TEXT("%d Listeners"), ListenerContainer.Num()));
+
+	for (const UAISense* Sense : Senses)
+	{
+		if (Sense)
+		{
+			Sense->DescribeSelfToGameplayDebugger(*this, DebuggerCategory);
+		}
+	}
+}
+#endif // WITH_GAMEPLAY_DEBUGGER_MENU
\ No newline at end of file
diff --git a/Engine/Source/Runtime/AIModule/Private/Perception/AISense.cpp b/Engine/Source/Runtime/AIModule/Private/Perception/AISense.cpp
--- a/Engine/Source/Runtime/AIModule/Private/Perception/AISense.cpp
+++ b/Engine/Source/Runtime/AIModule/Private/Perception/AISense.cpp
@@ -105,6 +105,13 @@ void UAISense::RegisterWrappedEvent(UAISenseEvent& PerceptionEvent)
 	UE_VLOG(GetPerceptionSystem(), LogAIPerception, Error, TEXT("%s did not override UAISense::RegisterWrappedEvent!"), *GetName());
 }
 
+#if WITH_GAMEPLAY_DEBUGGER_MENU
+void UAISense::DescribeSelfToGameplayDebugger(const UAIPerceptionSystem& PerceptionSystem, FGameplayDebuggerCategory& DebuggerCategory) const
+{
+	DebuggerCategory.AddTextLine(FString::Printf(TEXT("%s"), *SenseID.Name.ToString()));
+}
+#endif // WITH_GAMEPLAY_DEBUGGER_MENU
+
 //----------------------------------------------------------------------//
 // 
 //----------------------------------------------------------------------//
diff --git a/Engine/Source/Runtime/AIModule/Private/Perception/AISense_Sight.cpp b/Engine/Source/Runtime/AIModule/Private/Perception/AISense_Sight.cpp
--- a/Engine/Source/Runtime/AIModule/Private/Perception/AISense_Sight.cpp
+++ b/Engine/Source/Runtime/AIModule/Private/Perception/AISense_Sight.cpp
@@ -1,16 +1,18 @@
 // Copyright Epic Games, Inc. All Rights Reserved.
 
 #include "Perception/AISense_Sight.h"
-#include "EngineDefines.h"
-#include "EngineGlobals.h"
+
+#include "AIHelpers.h"
+#include "AISystem.h"
 #include "CollisionQueryParams.h"
 #include "Engine/Engine.h"
-#include "AISystem.h"
-#include "AIHelpers.h"
+#include "EngineDefines.h"
+#include "EngineGlobals.h"
+#include "GameplayDebuggerCategory.h"
 #include "Perception/AIPerceptionComponent.h"
-#include "VisualLogger/VisualLogger.h"
-#include "Perception/AISightTargetInterface.h"
 #include "Perception/AISenseConfig_Sight.h"
+#include "Perception/AISightTargetInterface.h"
+#include "VisualLogger/VisualLogger.h"
 
 #include UE_INLINE_GENERATED_CPP_BY_NAME(AISense_Sight)
 
@@ -1101,3 +1103,18 @@ void UAISense_Sight::OnListenerForgetsAll(const FPerceptionListener& Listener)
 	ForEach(SightQueriesPending, ForgetPreviousResult);
 }
 
+#if WITH_GAMEPLAY_DEBUGGER_MENU
+void UAISense_Sight::DescribeSelfToGameplayDebugger(const UAIPerceptionSystem& PerceptionSystem, FGameplayDebuggerCategory& DebuggerCategory) const
+{
+	const int32 TotalQueriesCount = SightQueriesInRange.Num() + SightQueriesOutOfRange.Num() + SightQueriesPending.Num();
+	DebuggerCategory.AddTextLine(
+		FString::Printf(TEXT("%s: %d Targets, %d Queries (InRange:%d, OutOfRange:%d, Pending:%d)"),
+			*GetSenseID().Name.ToString(),
+			ObservedTargets.Num(),
+			TotalQueriesCount,
+			SightQueriesInRange.Num(),
+			SightQueriesOutOfRange.Num(),
+			SightQueriesPending.Num())
+		);
+}
+#endif // WITH_GAMEPLAY_DEBUGGER_MENU
\ No newline at end of file
diff --git a/Engine/Source/Runtime/AIModule/Public/GameplayDebugger/GameplayDebuggerCategory_PerceptionSystem.h b/Engine/Source/Runtime/AIModule/Public/GameplayDebugger/GameplayDebuggerCategory_PerceptionSystem.h
new file mode 100644
--- /dev/null
+++ b/Engine/Source/Runtime/AIModule/Public/GameplayDebugger/GameplayDebuggerCategory_PerceptionSystem.h
@@ -0,0 +1,22 @@
+// Copyright Epic Games, Inc. All Rights Reserved.
+
+#pragma once
+
+#if WITH_GAMEPLAY_DEBUGGER_MENU
+#include "GameplayDebuggerCategory.h"
+
+class AActor;
+class APlayerController;
+
+class FGameplayDebuggerCategory_PerceptionSystem : public FGameplayDebuggerCategory
+{
+public:
+	FGameplayDebuggerCategory_PerceptionSystem();
+
+	AIMODULE_API static TSharedRef<FGameplayDebuggerCategory> MakeInstance();
+
+protected:
+	virtual void CollectData(APlayerController* OwnerPC, AActor* DebugActor) override;
+};
+
+#endif // WITH_GAMEPLAY_DEBUGGER_MENU
