buildfix: missing ifdef around gameplay debugger include

diff --git a/Engine/Source/Runtime/AIModule/Private/Perception/AIPerceptionSystem.cpp b/Engine/Source/Runtime/AIModule/Private/Perception/AIPerceptionSystem.cpp
--- a/Engine/Source/Runtime/AIModule/Private/Perception/AIPerceptionSystem.cpp
+++ b/Engine/Source/Runtime/AIModule/Private/Perception/AIPerceptionSystem.cpp
@@ -7,7 +7,6 @@
 #include "EngineGlobals.h"
 #include "EngineUtils.h"
 #include "GameFramework/Pawn.h"
-#include "GameplayDebuggerCategory.h"
 #include "Perception/AIPerceptionComponent.h"
 #include "Perception/AISenseConfig.h"
 #include "Perception/AISenseEvent.h"
@@ -16,6 +15,11 @@
 #include "TimerManager.h"
 #include "VisualLogger/VisualLogger.h"
 
+#if WITH_GAMEPLAY_DEBUGGER_MENU
+#include "GameplayDebuggerTypes.h"
+#include "GameplayDebuggerCategory.h"
+#endif // WITH_GAMEPLAY_DEBUGGER_MENU
+
 #include UE_INLINE_GENERATED_CPP_BY_NAME(AIPerceptionSystem)
 
 DECLARE_CYCLE_STAT(TEXT("Perception System"),STAT_AI_PerceptionSys,STATGROUP_AI);
diff --git a/Engine/Source/Runtime/AIModule/Private/Perception/AISense_Sight.cpp b/Engine/Source/Runtime/AIModule/Private/Perception/AISense_Sight.cpp
--- a/Engine/Source/Runtime/AIModule/Private/Perception/AISense_Sight.cpp
+++ b/Engine/Source/Runtime/AIModule/Private/Perception/AISense_Sight.cpp
@@ -8,12 +8,16 @@
 #include "Engine/Engine.h"
 #include "EngineDefines.h"
 #include "EngineGlobals.h"
-#include "GameplayDebuggerCategory.h"
 #include "Perception/AIPerceptionComponent.h"
 #include "Perception/AISenseConfig_Sight.h"
 #include "Perception/AISightTargetInterface.h"
 #include "VisualLogger/VisualLogger.h"
 
+#if WITH_GAMEPLAY_DEBUGGER_MENU
+#include "GameplayDebuggerTypes.h"
+#include "GameplayDebuggerCategory.h"
+#endif // WITH_GAMEPLAY_DEBUGGER_MENU
+
 #include UE_INLINE_GENERATED_CPP_BY_NAME(AISense_Sight)
 
 #define AISENSE_SIGHT_TIMESLICING_DEBUG 0
