Added bRequireNavigableEndLocation property to UBTTask_MoveTo to setup SetRequireNavigableEndLocation on the move request

diff --git a/Engine/Source/Runtime/AIModule/Classes/BehaviorTree/Tasks/BTTask_MoveTo.h b/Engine/Source/Runtime/AIModule/Classes/BehaviorTree/Tasks/BTTask_MoveTo.h
--- a/Engine/Source/Runtime/AIModule/Classes/BehaviorTree/Tasks/BTTask_MoveTo.h
+++ b/Engine/Source/Runtime/AIModule/Classes/BehaviorTree/Tasks/BTTask_MoveTo.h
@@ -65,6 +65,10 @@ class AIMODULE_API UBTTask_MoveTo : public UBTTask_BlackboardBase
 	UPROPERTY(Category = Node, EditAnywhere, AdvancedDisplay)
 	uint32 bTrackMovingGoal : 1;
 
+	/** if set, the goal location will need to be navigable */
+	UPROPERTY(Category = Node, EditAnywhere, AdvancedDisplay)
+	uint32 bRequireNavigableEndLocation : 1;
+
 	/** if set, goal location will be projected on navigation data (navmesh) before using */
 	UPROPERTY(Category = Node, EditAnywhere, AdvancedDisplay)
 	uint32 bProjectGoalLocation : 1;
diff --git a/Engine/Source/Runtime/AIModule/Private/BehaviorTree/Tasks/BTTask_MoveTo.cpp b/Engine/Source/Runtime/AIModule/Private/BehaviorTree/Tasks/BTTask_MoveTo.cpp
--- a/Engine/Source/Runtime/AIModule/Private/BehaviorTree/Tasks/BTTask_MoveTo.cpp
+++ b/Engine/Source/Runtime/AIModule/Private/BehaviorTree/Tasks/BTTask_MoveTo.cpp
@@ -23,6 +23,7 @@ UBTTask_MoveTo::UBTTask_MoveTo(const FObjectInitializer& ObjectInitializer) : Su
 	bAllowStrafe = GET_AI_CONFIG_VAR(bAllowStrafing);
 	bAllowPartialPath = GET_AI_CONFIG_VAR(bAcceptPartialPaths);
 	bTrackMovingGoal = true;
+	bRequireNavigableEndLocation = true;
 	bProjectGoalLocation = true;
 	bUsePathfinding = true;
 
@@ -87,6 +88,7 @@ EBTNodeResult::Type UBTTask_MoveTo::PerformMoveTask(UBehaviorTreeComponent& Owne
 		MoveReq.SetCanStrafe(bAllowStrafe);
 		MoveReq.SetReachTestIncludesAgentRadius(bReachTestIncludesAgentRadius);
 		MoveReq.SetReachTestIncludesGoalRadius(bReachTestIncludesGoalRadius);
+		MoveReq.SetRequireNavigableEndLocation(bRequireNavigableEndLocation);
 		MoveReq.SetProjectGoalLocation(bProjectGoalLocation);
 		MoveReq.SetUsePathfinding(bUsePathfinding);
 
