Fixed mismatched copy constructor/assignment operator in FIndexedHandleManagerBase.

diff --git a/Engine/Source/Runtime/AIModule/Public/IndexedHandle.h b/Engine/Source/Runtime/AIModule/Public/IndexedHandle.h
--- a/Engine/Source/Runtime/AIModule/Public/IndexedHandle.h
+++ b/Engine/Source/Runtime/AIModule/Public/IndexedHandle.h
@@ -269,6 +269,9 @@ public:
 	/** returns the number of used handles, will be greater than or equal to zero */
 	int32 CalcNumUsedHandles() const { return Handles.Num() - FreeHandleIndices.Num(); };
 
+	// Defaulted copy constructor, but it acts differently from the copy assignment operator - notably, doesn't allocate a new ManagerID.
+	FIndexedHandleManagerBase(const FIndexedHandleManagerBase& Other) = default;
+
 	FIndexedHandleManagerBase& operator=(const FIndexedHandleManagerBase& Other)
 	{
 		Handles = Other.Handles;
