Added checks to EQS item types to ensure the correct amount of memory gets allocated.

diff --git a/Engine/Source/Runtime/AIModule/Private/EnvironmentQuery/Items/EnvQueryItemType_Actor.cpp b/Engine/Source/Runtime/AIModule/Private/EnvironmentQuery/Items/EnvQueryItemType_Actor.cpp
--- a/Engine/Source/Runtime/AIModule/Private/EnvironmentQuery/Items/EnvQueryItemType_Actor.cpp
+++ b/Engine/Source/Runtime/AIModule/Private/EnvironmentQuery/Items/EnvQueryItemType_Actor.cpp
@@ -57,6 +57,8 @@ void UEnvQueryItemType_Actor::SetContextHelper(FEnvQueryContextData& ContextData
 	ContextData.NumValues = MultipleActors.Num();
 	ContextData.RawData.SetNumUninitialized(sizeof(FWeakObjectPtr) * MultipleActors.Num());
 
+	checkf(ContextData.RawData.Num() == sizeof(FWeakObjectPtr) * MultipleActors.Num(), TEXT("Failed to allocate the appropriate amount of memory"));
+
 	uint8* RawData = (uint8*)ContextData.RawData.GetData();
 	for (int32 ActorIndex = 0; ActorIndex < MultipleActors.Num(); ActorIndex++)
 	{
@@ -69,7 +71,9 @@ void UEnvQueryItemType_Actor::SetContextHelper(FEnvQueryContextData& ContextData
 {
 	ContextData.ValueType = UEnvQueryItemType_Actor::StaticClass();
 	ContextData.NumValues = MultipleActors.Num();
-	ContextData.RawData.SetNumUninitialized(sizeof(FWeakObjectPtr)* MultipleActors.Num());
+	ContextData.RawData.SetNumUninitialized(sizeof(FWeakObjectPtr) * MultipleActors.Num());
+
+	checkf(ContextData.RawData.Num() == sizeof(FWeakObjectPtr) * MultipleActors.Num(), TEXT("Failed to allocate the appropriate amount of memory"));
 
 	uint8* RawData = (uint8*)ContextData.RawData.GetData();
 	for (int32 ActorIndex = 0; ActorIndex < MultipleActors.Num(); ActorIndex++)
diff --git a/Engine/Source/Runtime/AIModule/Private/EnvironmentQuery/Items/EnvQueryItemType_Direction.cpp b/Engine/Source/Runtime/AIModule/Private/EnvironmentQuery/Items/EnvQueryItemType_Direction.cpp
--- a/Engine/Source/Runtime/AIModule/Private/EnvironmentQuery/Items/EnvQueryItemType_Direction.cpp
+++ b/Engine/Source/Runtime/AIModule/Private/EnvironmentQuery/Items/EnvQueryItemType_Direction.cpp
@@ -65,6 +65,8 @@ void UEnvQueryItemType_Direction::SetContextHelper(FEnvQueryContextData& Context
 	ContextData.NumValues = MultipleDirections.Num();
 	ContextData.RawData.SetNumUninitialized(sizeof(FVector) * MultipleDirections.Num());
 
+	checkf(ContextData.RawData.Num() == sizeof(FVector) * MultipleDirections.Num(), TEXT("Failed to allocate the appropriate amount of memory"));
+
 	uint8* RawData = (uint8*)ContextData.RawData.GetData();
 	for (int32 DirectionIndex = 0; DirectionIndex < MultipleDirections.Num(); DirectionIndex++)
 	{
@@ -79,6 +81,8 @@ void UEnvQueryItemType_Direction::SetContextHelper(FEnvQueryContextData& Context
 	ContextData.NumValues = MultipleRotations.Num();
 	ContextData.RawData.SetNumUninitialized(sizeof(FVector) * MultipleRotations.Num());
 
+	checkf(ContextData.RawData.Num() == sizeof(FVector) * MultipleRotations.Num(), TEXT("Failed to allocate the appropriate amount of memory"));
+
 	uint8* RawData = (uint8*)ContextData.RawData.GetData();
 	for (int32 RotationIndex = 0; RotationIndex < MultipleRotations.Num(); RotationIndex++)
 	{
diff --git a/Engine/Source/Runtime/AIModule/Private/EnvironmentQuery/Items/EnvQueryItemType_Point.cpp b/Engine/Source/Runtime/AIModule/Private/EnvironmentQuery/Items/EnvQueryItemType_Point.cpp
--- a/Engine/Source/Runtime/AIModule/Private/EnvironmentQuery/Items/EnvQueryItemType_Point.cpp
+++ b/Engine/Source/Runtime/AIModule/Private/EnvironmentQuery/Items/EnvQueryItemType_Point.cpp
@@ -56,6 +56,8 @@ void UEnvQueryItemType_Point::SetContextHelper(FEnvQueryContextData& ContextData
 	ContextData.NumValues = MultiplePoints.Num();
 	ContextData.RawData.SetNumUninitialized(sizeof(FNavLocation)* MultiplePoints.Num());
 
+	checkf(ContextData.RawData.Num() == sizeof(FNavLocation) * MultiplePoints.Num(), TEXT("Failed to allocate the appropriate amount of memory"));
+
 	uint8* RawData = (uint8*)ContextData.RawData.GetData();
 	for (int32 PointIndex = 0; PointIndex < MultiplePoints.Num(); PointIndex++)
 	{
