disabled race detection for FOnBlackboardChangeNotification delegate because it requires to be trivially relocatable. https://horde.devtools.epicgames.com/job/647fae66d20bd4233b5b7d9b?step=0549&issue=315742

diff --git a/Engine/Source/Runtime/AIModule/Classes/BehaviorTree/BehaviorTreeTypes.h b/Engine/Source/Runtime/AIModule/Classes/BehaviorTree/BehaviorTreeTypes.h
--- a/Engine/Source/Runtime/AIModule/Classes/BehaviorTree/BehaviorTreeTypes.h
+++ b/Engine/Source/Runtime/AIModule/Classes/BehaviorTree/BehaviorTreeTypes.h
@@ -59,7 +59,10 @@ enum class EBlackboardNotificationResult : uint8
 
 // delegate defines
 DECLARE_DELEGATE_TwoParams(FOnBlackboardChange, const UBlackboardComponent&, FBlackboard::FKey /*key ID*/);
-DECLARE_DELEGATE_RetVal_TwoParams(EBlackboardNotificationResult, FOnBlackboardChangeNotification, const UBlackboardComponent&, FBlackboard::FKey /*key ID*/);
+
+// using "not checked" user policy (means race detection is disabled) because this delegate is stored in a TMultiMap and causes its reallocation
+// from inside delegate's execution. This is incompatible with race detection that needs to access the delegate instance after its execution
+using FOnBlackboardChangeNotification = TDelegate<EBlackboardNotificationResult(const UBlackboardComponent&, FBlackboard::FKey keyID), FNotThreadSafeNotCheckedDelegateUserPolicy>;
 
 namespace BTSpecialChild
 {
