Made EQS debug's refresh interval tweakable
#rnx

diff --git a/Engine/Source/Runtime/AIModule/Private/GameplayDebugger/GameplayDebuggerCategory_EQS.cpp b/Engine/Source/Runtime/AIModule/Private/GameplayDebugger/GameplayDebuggerCategory_EQS.cpp
--- a/Engine/Source/Runtime/AIModule/Private/GameplayDebugger/GameplayDebuggerCategory_EQS.cpp
+++ b/Engine/Source/Runtime/AIModule/Private/GameplayDebugger/GameplayDebuggerCategory_EQS.cpp
@@ -12,12 +12,23 @@
 #include "EnvironmentQuery/EQSRenderingComponent.h"
 #include "DrawDebugHelpers.h"
 
+namespace UE::GameplayDebuggerCategory_EQS::Tweakables
+{
+	float RefreshInterval = 2.0f;
+
+	FAutoConsoleVariableRef CVars[] =
+	{
+		FAutoConsoleVariableRef(TEXT("ai.debug.EQS.RefreshInterval"), Tweakables::RefreshInterval,
+			TEXT("Interval (in seconds) at which data will be collected.")),
+	};
+}
+
 FGameplayDebuggerCategory_EQS::FGameplayDebuggerCategory_EQS()
 {
 	MaxQueries = 5;
 	MaxItemTableRows = 10;
 	ShownQueryIndex = 0;
-	CollectDataInterval = 2.0f;
+	CollectDataInterval = UE::GameplayDebuggerCategory_EQS::Tweakables::RefreshInterval;
 
 	const FGameplayDebuggerInputHandlerConfig CycleConfig(TEXT("CycleQueries"), TEXT("Multiply"));
 	const FGameplayDebuggerInputHandlerConfig DetailsConfig(TEXT("ToggleDetails"), TEXT("Divide"));
