PR #10528: Export bound functions for UAIAsyncTaskBlueprintProxy to allow other ?

mz> the change lets users extent BP "AI Move To" node

diff --git a/Engine/Source/Runtime/AIModule/Classes/Blueprint/AIAsyncTaskBlueprintProxy.h b/Engine/Source/Runtime/AIModule/Classes/Blueprint/AIAsyncTaskBlueprintProxy.h
--- a/Engine/Source/Runtime/AIModule/Classes/Blueprint/AIAsyncTaskBlueprintProxy.h
+++ b/Engine/Source/Runtime/AIModule/Classes/Blueprint/AIAsyncTaskBlueprintProxy.h
@@ -26,10 +26,10 @@ class UAIAsyncTaskBlueprintProxy : public UObject
 
 public:
 	UFUNCTION()
-	void OnMoveCompleted(FAIRequestID RequestID, EPathFollowingResult::Type MovementResult);
+	AIMODULE_API void OnMoveCompleted(FAIRequestID RequestID, EPathFollowingResult::Type MovementResult);
 
-	void OnNoPath();
-	void OnAtGoal();
+	AIMODULE_API void OnNoPath();
+	AIMODULE_API void OnAtGoal();
 
 	//~ Begin UObject Interface
 	virtual void BeginDestroy() override;
