[Engine]
* Moved dllexport from type to methods/staticvar in all Engine runtime code. This improves compile times, memory and performance in dll builds

diff --git a/Engine/Source/Runtime/AIModule/Classes/AIController.h b/Engine/Source/Runtime/AIModule/Classes/AIController.h
--- a/Engine/Source/Runtime/AIModule/Classes/AIController.h
+++ b/Engine/Source/Runtime/AIModule/Classes/AIController.h
@@ -89,8 +89,8 @@ struct FFocusKnowledge
  * @see https://docs.unrealengine.com/latest/INT/Gameplay/Framework/Controller/
  */
 
-UCLASS(ClassGroup = AI, BlueprintType, Blueprintable)
-class AIMODULE_API AAIController : public AController, public IAIPerceptionListenerInterface, public IGameplayTaskOwnerInterface, public IGenericTeamAgentInterface, public IVisualLoggerDebugSnapshotInterface
+UCLASS(ClassGroup = AI, BlueprintType, Blueprintable, MinimalAPI)
+class AAIController : public AController, public IAIPerceptionListenerInterface, public IGameplayTaskOwnerInterface, public IGenericTeamAgentInterface, public IVisualLoggerDebugSnapshotInterface
 {
 	GENERATED_BODY()
 
@@ -161,9 +161,9 @@ protected:
 
 public:
 
-	AAIController(const FObjectInitializer& ObjectInitializer = FObjectInitializer::Get());
+	AIMODULE_API AAIController(const FObjectInitializer& ObjectInitializer = FObjectInitializer::Get());
 
-	virtual void SetPawn(APawn* InPawn) override;
+	AIMODULE_API virtual void SetPawn(APawn* InPawn) override;
 
 	/** Makes AI go toward specified Goal actor (destination will be continuously updated), aborts any active path following
 	 *  @param AcceptanceRadius - finish move if pawn gets close enough
@@ -175,7 +175,7 @@ public:
 	 *	@note AcceptanceRadius has default value or -1 due to Header Parser not being able to recognize UPathFollowingComponent::DefaultAcceptanceRadius
 	 */
 	UFUNCTION(BlueprintCallable, Category = "AI|Navigation", Meta = (AdvancedDisplay = "bStopOnOverlap,bCanStrafe,bAllowPartialPath"))
-	EPathFollowingRequestResult::Type MoveToActor(AActor* Goal, float AcceptanceRadius = -1, bool bStopOnOverlap = true,
+	AIMODULE_API EPathFollowingRequestResult::Type MoveToActor(AActor* Goal, float AcceptanceRadius = -1, bool bStopOnOverlap = true,
 		bool bUsePathfinding = true, bool bCanStrafe = true,
 		TSubclassOf<UNavigationQueryFilter> FilterClass = NULL, bool bAllowPartialPath = true);
 
@@ -190,7 +190,7 @@ public:
 	 *	@note AcceptanceRadius has default value or -1 due to Header Parser not being able to recognize UPathFollowingComponent::DefaultAcceptanceRadius
 	 */
 	UFUNCTION(BlueprintCallable, Category = "AI|Navigation", Meta = (AdvancedDisplay = "bStopOnOverlap,bCanStrafe,bAllowPartialPath"))
-	EPathFollowingRequestResult::Type MoveToLocation(const FVector& Dest, float AcceptanceRadius = -1, bool bStopOnOverlap = true,
+	AIMODULE_API EPathFollowingRequestResult::Type MoveToLocation(const FVector& Dest, float AcceptanceRadius = -1, bool bStopOnOverlap = true,
 		bool bUsePathfinding = true, bool bProjectDestinationToNavigation = false, bool bCanStrafe = true,
 		TSubclassOf<UNavigationQueryFilter> FilterClass = NULL, bool bAllowPartialPath = true);
 
@@ -199,44 +199,44 @@ public:
 	 *  @param OutPath - optional output param, filled in with assigned path
 	 *  @return struct holding MoveId and enum code
 	 */
-	virtual FPathFollowingRequestResult MoveTo(const FAIMoveRequest& MoveRequest, FNavPathSharedPtr* OutPath = nullptr);
+	AIMODULE_API virtual FPathFollowingRequestResult MoveTo(const FAIMoveRequest& MoveRequest, FNavPathSharedPtr* OutPath = nullptr);
 
 	/** Passes move request and path object to path following */
-	virtual FAIRequestID RequestMove(const FAIMoveRequest& MoveRequest, FNavPathSharedPtr Path);
+	AIMODULE_API virtual FAIRequestID RequestMove(const FAIMoveRequest& MoveRequest, FNavPathSharedPtr Path);
 
 	/** Finds path for given move request
  	 *  @param MoveRequest - details about move
 	 *  @param Query - pathfinding query for navigation system
 	 *  @param OutPath - generated path
 	 */
-	virtual void FindPathForMoveRequest(const FAIMoveRequest& MoveRequest, FPathFindingQuery& Query, FNavPathSharedPtr& OutPath) const;
+	AIMODULE_API virtual void FindPathForMoveRequest(const FAIMoveRequest& MoveRequest, FPathFindingQuery& Query, FNavPathSharedPtr& OutPath) const;
 
 	/** Helper function for creating pathfinding query for this agent from move request data */
-	bool BuildPathfindingQuery(const FAIMoveRequest& MoveRequest, FPathFindingQuery& Query) const;
+	AIMODULE_API bool BuildPathfindingQuery(const FAIMoveRequest& MoveRequest, FPathFindingQuery& Query) const;
 
 	UE_DEPRECATED_FORGAME(4.13, "This function is now deprecated, please use FindPathForMoveRequest() for adjusting Query or BuildPathfindingQuery() for getting one.")
-	virtual bool PreparePathfinding(const FAIMoveRequest& MoveRequest, FPathFindingQuery& Query);
+	AIMODULE_API virtual bool PreparePathfinding(const FAIMoveRequest& MoveRequest, FPathFindingQuery& Query);
 
 	UE_DEPRECATED_FORGAME(4.13, "This function is now deprecated, please use FindPathForMoveRequest() for adjusting pathfinding or path postprocess.")
-	virtual FAIRequestID RequestPathAndMove(const FAIMoveRequest& MoveRequest, FPathFindingQuery& Query);
+	AIMODULE_API virtual FAIRequestID RequestPathAndMove(const FAIMoveRequest& MoveRequest, FPathFindingQuery& Query);
 
 	/** if AI is currently moving due to request given by RequestToPause, then the move will be paused */
-	bool PauseMove(FAIRequestID RequestToPause);
+	AIMODULE_API bool PauseMove(FAIRequestID RequestToPause);
 
 	/** resumes last AI-performed, paused request provided it's ID was equivalent to RequestToResume */
-	bool ResumeMove(FAIRequestID RequestToResume);
+	AIMODULE_API bool ResumeMove(FAIRequestID RequestToResume);
 
 	/** Aborts the move the controller is currently performing */
-	virtual void StopMovement() override;
+	AIMODULE_API virtual void StopMovement() override;
 
 	/** Called on completing current movement request */
-	virtual void OnMoveCompleted(FAIRequestID RequestID, const FPathFollowingResult& Result);
+	AIMODULE_API virtual void OnMoveCompleted(FAIRequestID RequestID, const FPathFollowingResult& Result);
 
 	UE_DEPRECATED_FORGAME(4.13, "This function is now deprecated, please use version with EPathFollowingResultDetails parameter.")
-	virtual void OnMoveCompleted(FAIRequestID RequestID, EPathFollowingResult::Type Result);
+	AIMODULE_API virtual void OnMoveCompleted(FAIRequestID RequestID, EPathFollowingResult::Type Result);
 
 	/** Returns the Move Request ID for the current move */
-	FAIRequestID GetCurrentMoveRequestID() const;
+	AIMODULE_API FAIRequestID GetCurrentMoveRequestID() const;
 
 	/** Blueprint notification that we've completed the current movement request */
 	UPROPERTY(BlueprintAssignable, meta = (DisplayName = "MoveCompleted"))
@@ -246,26 +246,26 @@ public:
 
 	/** Returns status of path following */
 	UFUNCTION(BlueprintCallable, Category = "AI|Navigation")
-	EPathFollowingStatus::Type GetMoveStatus() const;
+	AIMODULE_API EPathFollowingStatus::Type GetMoveStatus() const;
 
 	/** Returns true if the current PathFollowingComponent's path is partial (does not reach desired destination). */
 	UFUNCTION(BlueprintCallable, Category = "AI|Navigation")
-	bool HasPartialPath() const;
+	AIMODULE_API bool HasPartialPath() const;
 
 	/** Returns position of current path segment's end. */
 	UFUNCTION(BlueprintCallable, Category = "AI|Navigation")
-	FVector GetImmediateMoveDestination() const;
+	AIMODULE_API FVector GetImmediateMoveDestination() const;
 
 	/** Updates state of movement block detection. */
 	UFUNCTION(BlueprintCallable, Category = "AI|Navigation")
-	void SetMoveBlockDetection(bool bEnable);
+	AIMODULE_API void SetMoveBlockDetection(bool bEnable);
 	
 	/** Starts executing behavior tree. */
 	UFUNCTION(BlueprintCallable, Category = "AI")
-	virtual bool RunBehaviorTree(UBehaviorTree* BTAsset);
+	AIMODULE_API virtual bool RunBehaviorTree(UBehaviorTree* BTAsset);
 
 protected:
-	virtual void CleanupBrainComponent();
+	AIMODULE_API virtual void CleanupBrainComponent();
 
 public:
 	/**
@@ -275,51 +275,51 @@ public:
 	 * @return true if we successfully linked the blackboard asset to the blackboard component.
 	 */
 	UFUNCTION(BlueprintCallable, Category = "AI")
-	bool UseBlackboard(UBlackboardData* BlackboardAsset, UBlackboardComponent*& BlackboardComponent);
+	AIMODULE_API bool UseBlackboard(UBlackboardData* BlackboardAsset, UBlackboardComponent*& BlackboardComponent);
 
 	/** does this AIController allow given UBlackboardComponent sync data with it */
-	virtual bool ShouldSyncBlackboardWith(const UBlackboardComponent& OtherBlackboardComponent) const;
+	AIMODULE_API virtual bool ShouldSyncBlackboardWith(const UBlackboardComponent& OtherBlackboardComponent) const;
 
 	UFUNCTION(BlueprintCallable, Category = "AI|Tasks")
-	void ClaimTaskResource(TSubclassOf<UGameplayTaskResource> ResourceClass);
+	AIMODULE_API void ClaimTaskResource(TSubclassOf<UGameplayTaskResource> ResourceClass);
 	
 	UFUNCTION(BlueprintCallable, Category = "AI|Tasks")
-	void UnclaimTaskResource(TSubclassOf<UGameplayTaskResource> ResourceClass);
+	AIMODULE_API void UnclaimTaskResource(TSubclassOf<UGameplayTaskResource> ResourceClass);
 
 protected:
 	UFUNCTION(BlueprintImplementableEvent)
-	void OnUsingBlackBoard(UBlackboardComponent* BlackboardComp, UBlackboardData* BlackboardAsset);
+	AIMODULE_API void OnUsingBlackBoard(UBlackboardComponent* BlackboardComp, UBlackboardData* BlackboardAsset);
 
-	virtual bool InitializeBlackboard(UBlackboardComponent& BlackboardComp, UBlackboardData& BlackboardAsset);
+	AIMODULE_API virtual bool InitializeBlackboard(UBlackboardComponent& BlackboardComp, UBlackboardData& BlackboardAsset);
 
 public:
 	/** Retrieve the final position that controller should be looking at. */
 	UFUNCTION(BlueprintCallable, Category = "AI")
-	FVector GetFocalPoint() const;
+	AIMODULE_API FVector GetFocalPoint() const;
 
-	FVector GetFocalPointForPriority(EAIFocusPriority::Type InPriority) const;
+	AIMODULE_API FVector GetFocalPointForPriority(EAIFocusPriority::Type InPriority) const;
 
 	/** Retrieve the focal point this controller should focus to on given actor. */
 	UFUNCTION(BlueprintCallable, Category = "AI")
-	virtual FVector GetFocalPointOnActor(const AActor *Actor) const;
+	AIMODULE_API virtual FVector GetFocalPointOnActor(const AActor *Actor) const;
 
 	/** Set the position that controller should be looking at. */
 	UFUNCTION(BlueprintCallable, Category = "AI", meta = (DisplayName = "SetFocalPoint", ScriptName = "SetFocalPoint", Keywords = "focus"))
-	void K2_SetFocalPoint(FVector FP);
+	AIMODULE_API void K2_SetFocalPoint(FVector FP);
 
 	/** Set Focus for actor, will set FocalPoint as a result. */
 	UFUNCTION(BlueprintCallable, Category = "AI", meta = (DisplayName = "SetFocus", ScriptName = "SetFocus"))
-	void K2_SetFocus(AActor* NewFocus);
+	AIMODULE_API void K2_SetFocus(AActor* NewFocus);
 
 	/** Get the focused actor. */
 	UFUNCTION(BlueprintCallable, Category = "AI")
-	AActor* GetFocusActor() const;
+	AIMODULE_API AActor* GetFocusActor() const;
 
 	FORCEINLINE AActor* GetFocusActorForPriority(EAIFocusPriority::Type InPriority) const { return FocusInformation.Priorities.IsValidIndex(InPriority) ? FocusInformation.Priorities[InPriority].Actor.Get() : nullptr; }
 
 	/** Clears Focus, will also clear FocalPoint as a result */
 	UFUNCTION(BlueprintCallable, Category = "AI", meta = (DisplayName = "ClearFocus", ScriptName = "ClearFocus"))
-	void K2_ClearFocus();
+	AIMODULE_API void K2_ClearFocus();
 
 
 	/**
@@ -334,28 +334,28 @@ public:
 	 * @param bOnlyTraceUp  (optional) - when true collision checks verifying the arc will only be done along the upward portion of the arc
 	 * @return - true if a valid arc was computed, false if no valid solution could be found
 	 */
-	bool SuggestTossVelocity(FVector& OutTossVelocity, FVector Start, FVector End, float TossSpeed, bool bPreferHighArc, float CollisionRadius = 0, bool bOnlyTraceUp = false);
+	AIMODULE_API bool SuggestTossVelocity(FVector& OutTossVelocity, FVector Start, FVector End, float TossSpeed, bool bPreferHighArc, float CollisionRadius = 0, bool bOnlyTraceUp = false);
 
 	//~ Begin AActor Interface
-	virtual void Tick(float DeltaTime) override;
-	virtual void PostInitializeComponents() override;
-	virtual void PostRegisterAllComponents() override;
+	AIMODULE_API virtual void Tick(float DeltaTime) override;
+	AIMODULE_API virtual void PostInitializeComponents() override;
+	AIMODULE_API virtual void PostRegisterAllComponents() override;
 	//~ End AActor Interface
 
 	//~ Begin AController Interface
 protected:
-	virtual void OnPossess(APawn* InPawn) override;
-	virtual void OnUnPossess() override;
+	AIMODULE_API virtual void OnPossess(APawn* InPawn) override;
+	AIMODULE_API virtual void OnUnPossess() override;
 
 public:
-	virtual bool ShouldPostponePathUpdates() const override;
-	virtual void DisplayDebug(UCanvas* Canvas, const FDebugDisplayInfo& DebugDisplay, float& YL, float& YPos) override;
+	AIMODULE_API virtual bool ShouldPostponePathUpdates() const override;
+	AIMODULE_API virtual void DisplayDebug(UCanvas* Canvas, const FDebugDisplayInfo& DebugDisplay, float& YL, float& YPos) override;
 
 #if ENABLE_VISUAL_LOG
-	virtual void GrabDebugSnapshot(FVisualLogEntry* Snapshot) const override;
+	AIMODULE_API virtual void GrabDebugSnapshot(FVisualLogEntry* Snapshot) const override;
 #endif
 
-	virtual void Reset() override;
+	AIMODULE_API virtual void Reset() override;
 
 	/**
 	 * Checks line to center and top of other actor
@@ -364,26 +364,26 @@ public:
 	 * @param bAlternateChecks used only in AIController implementation
 	 * @return true if controller's pawn can see Other actor.
 	 */
-	virtual bool LineOfSightTo(const AActor* Other, FVector ViewPoint = FVector(ForceInit), bool bAlternateChecks = false) const override;
+	AIMODULE_API virtual bool LineOfSightTo(const AActor* Other, FVector ViewPoint = FVector(ForceInit), bool bAlternateChecks = false) const override;
 	//~ End AController Interface
 
 	/** Notifies AIController of changes in given actors' perception */
-	virtual void ActorsPerceptionUpdated(const TArray<AActor*>& UpdatedActors);
+	AIMODULE_API virtual void ActorsPerceptionUpdated(const TArray<AActor*>& UpdatedActors);
 
 	/** Update direction AI is looking based on FocalPoint */
-	virtual void UpdateControlRotation(float DeltaTime, bool bUpdatePawn = true);
+	AIMODULE_API virtual void UpdateControlRotation(float DeltaTime, bool bUpdatePawn = true);
 
 	/** Set FocalPoint for given priority as absolute position or offset from base. */
-	virtual void SetFocalPoint(FVector NewFocus, EAIFocusPriority::Type InPriority = EAIFocusPriority::Gameplay);
+	AIMODULE_API virtual void SetFocalPoint(FVector NewFocus, EAIFocusPriority::Type InPriority = EAIFocusPriority::Gameplay);
 
 	/* Set Focus actor for given priority, will set FocalPoint as a result. */
-	virtual void SetFocus(AActor* NewFocus, EAIFocusPriority::Type InPriority = EAIFocusPriority::Gameplay);
+	AIMODULE_API virtual void SetFocus(AActor* NewFocus, EAIFocusPriority::Type InPriority = EAIFocusPriority::Gameplay);
 
 	/** Clears Focus for given priority, will also clear FocalPoint as a result
 	 *	@param InPriority focus priority to clear. If you don't know what to use you probably mean EAIFocusPriority::Gameplay*/
-	virtual void ClearFocus(EAIFocusPriority::Type InPriority);
+	AIMODULE_API virtual void ClearFocus(EAIFocusPriority::Type InPriority);
 
-	void SetPerceptionComponent(UAIPerceptionComponent& InPerceptionComponent);
+	AIMODULE_API void SetPerceptionComponent(UAIPerceptionComponent& InPerceptionComponent);
 	//----------------------------------------------------------------------//
 	// IAIPerceptionListenerInterface
 	//----------------------------------------------------------------------//
@@ -392,8 +392,8 @@ public:
 	//----------------------------------------------------------------------//
 	// INavAgentInterface
 	//----------------------------------------------------------------------//
-	virtual bool IsFollowingAPath() const override;
-	virtual IPathFollowingAgentInterface* GetPathFollowingAgent() const override;
+	AIMODULE_API virtual bool IsFollowingAPath() const override;
+	AIMODULE_API virtual IPathFollowingAgentInterface* GetPathFollowingAgent() const override;
 
 	//----------------------------------------------------------------------//
 	// IGenericTeamAgentInterface
@@ -401,7 +401,7 @@ public:
 private:
 	FGenericTeamId TeamID;
 public:
-	virtual void SetGenericTeamId(const FGenericTeamId& NewTeamID) override;
+	AIMODULE_API virtual void SetGenericTeamId(const FGenericTeamId& NewTeamID) override;
 	virtual FGenericTeamId GetGenericTeamId() const override { return TeamID; }
 
 	//----------------------------------------------------------------------//
@@ -420,19 +420,19 @@ public:
 	virtual void OnGameplayTaskDeactivated(UGameplayTask& Task) override {}
 
 	UFUNCTION()
-	virtual void OnGameplayTaskResourcesClaimed(FGameplayResourceSet NewlyClaimed, FGameplayResourceSet FreshlyReleased);
+	AIMODULE_API virtual void OnGameplayTaskResourcesClaimed(FGameplayResourceSet NewlyClaimed, FGameplayResourceSet FreshlyReleased);
 
 	//----------------------------------------------------------------------//
 	// debug/dev-time 
 	//----------------------------------------------------------------------//
-	virtual FString GetDebugIcon() const;
+	AIMODULE_API virtual FString GetDebugIcon() const;
 	
 	// Cheat/debugging functions
 	static void ToggleAIIgnorePlayers() { bAIIgnorePlayers = !bAIIgnorePlayers; }
 	static bool AreAIIgnoringPlayers() { return bAIIgnorePlayers; }
 
 	/** If true, AI controllers will ignore players. */
-	static bool bAIIgnorePlayers;
+	static AIMODULE_API bool bAIIgnorePlayers;
 
 public:
 	/** Returns PathFollowingComponent subobject **/
@@ -450,7 +450,7 @@ public:
 	/** Note that this function does not do any pathfollowing state transfer. 
 	 *	Intended to be called as part of initialization/setup process */
 	UFUNCTION(BlueprintCallable, Category = "AI|Navigation")
-	void SetPathFollowingComponent(UPathFollowingComponent* NewPFComponent);
+	AIMODULE_API void SetPathFollowingComponent(UPathFollowingComponent* NewPFComponent);
 
 	//----------------------------------------------------------------------//
 	// DEPRECATED
@@ -459,7 +459,7 @@ public:
 	UDEPRECATED_PawnActionsComponent* GetActionsComp() const { return ActionsComp_DEPRECATED; }
 	
 	UE_DEPRECATED(5.2, "PawnActions have been deprecated and are no longer being supported. It will get removed in following UE5 releases. Use GameplayTasks or AITasks instead.")
-	bool PerformAction(UDEPRECATED_PawnAction& Action, EAIRequestPriority::Type Priority, UObject* const Instigator = NULL);
+	AIMODULE_API bool PerformAction(UDEPRECATED_PawnAction& Action, EAIRequestPriority::Type Priority, UObject* const Instigator = NULL);
 
 	UFUNCTION(BlueprintGetter, meta = (DeprecatedFunction, DeprecationMessage = "PawnActions have been deprecated and are no longer being supported. It will get removed in following UE5 releases. Use GameplayTasks or AITasks instead."))
 	UDEPRECATED_PawnActionsComponent* GetDeprecatedActionsComponent() const { return ActionsComp_DEPRECATED; }
diff --git a/Engine/Source/Runtime/AIModule/Classes/AIResources.h b/Engine/Source/Runtime/AIModule/Classes/AIResources.h
--- a/Engine/Source/Runtime/AIModule/Classes/AIResources.h
+++ b/Engine/Source/Runtime/AIModule/Classes/AIResources.h
@@ -6,18 +6,18 @@
 #include "GameplayTaskResource.h"
 #include "AIResources.generated.h"
 
-UCLASS(meta = (DisplayName = "AI Movement"))
-class AIMODULE_API UAIResource_Movement : public UGameplayTaskResource
+UCLASS(meta = (DisplayName = "AI Movement"), MinimalAPI)
+class UAIResource_Movement : public UGameplayTaskResource
 {
 	GENERATED_BODY()
 
 #if !(UE_BUILD_SHIPPING || UE_BUILD_TEST)
-	virtual FString GenerateDebugDescription() const override;
+	AIMODULE_API virtual FString GenerateDebugDescription() const override;
 #endif // !(UE_BUILD_SHIPPING || UE_BUILD_TEST)
 };
 
-UCLASS(meta = (DisplayName = "AI Logic"))
-class AIMODULE_API UAIResource_Logic : public UGameplayTaskResource
+UCLASS(meta = (DisplayName = "AI Logic"), MinimalAPI)
+class UAIResource_Logic : public UGameplayTaskResource
 {
 	GENERATED_BODY()
 };
diff --git a/Engine/Source/Runtime/AIModule/Classes/AISubsystem.h b/Engine/Source/Runtime/AIModule/Classes/AISubsystem.h
--- a/Engine/Source/Runtime/AIModule/Classes/AISubsystem.h
+++ b/Engine/Source/Runtime/AIModule/Classes/AISubsystem.h
@@ -11,8 +11,8 @@
 class UAISystem;
 
 /** A class representing a common interface and behavior for AI subsystems */
-UCLASS(config = Engine, defaultconfig)
-class AIMODULE_API UAISubsystem : public UObject, public FTickableGameObject
+UCLASS(config = Engine, defaultconfig, MinimalAPI)
+class UAISubsystem : public UObject, public FTickableGameObject
 {
 	GENERATED_BODY()
 
@@ -21,15 +21,15 @@ private:
 	TObjectPtr<UAISystem> AISystem;
 			
 public:
-	UAISubsystem(const FObjectInitializer& ObjectInitializer = FObjectInitializer::Get());
+	AIMODULE_API UAISubsystem(const FObjectInitializer& ObjectInitializer = FObjectInitializer::Get());
 
-	virtual UWorld* GetWorld() const override;
+	AIMODULE_API virtual UWorld* GetWorld() const override;
 
 	// FTickableGameObject begin
 	virtual UWorld* GetTickableGameObjectWorld() const override { return GetWorldFast(); }
 	virtual void Tick(float DeltaTime) override {}
-	virtual ETickableTickType GetTickableTickType() const override;
-	virtual TStatId GetStatId() const override;
+	AIMODULE_API virtual ETickableTickType GetTickableTickType() const override;
+	AIMODULE_API virtual TStatId GetStatId() const override;
 	// FTickableGameObject end
 
 	UWorld* GetWorldFast() const { return AISystem ? AISystem->GetOuterWorld() : GetOuter()->GetWorld(); }
diff --git a/Engine/Source/Runtime/AIModule/Classes/AISystem.h b/Engine/Source/Runtime/AIModule/Classes/AISystem.h
--- a/Engine/Source/Runtime/AIModule/Classes/AISystem.h
+++ b/Engine/Source/Runtime/AIModule/Classes/AISystem.h
@@ -22,8 +22,8 @@ class UNavLocalGridManager;
 
 #define GET_AI_CONFIG_VAR(a) (GetDefault<UAISystem>()->a)
 
-UCLASS(config=Engine, defaultconfig)
-class AIMODULE_API UAISystem : public UAISystemBase
+UCLASS(config=Engine, defaultconfig, MinimalAPI)
+class UAISystem : public UAISystemBase
 {
 	GENERATED_BODY()
 
@@ -134,22 +134,22 @@ protected:
 	FDelegateHandle PawnBeginPlayDelegateHandle;
 
 	/** random number stream to be used by all things AI. WIP */
-	static FRandomStream RandomStream;
+	static AIMODULE_API FRandomStream RandomStream;
 	
 public:
-	UAISystem(const FObjectInitializer& ObjectInitializer);
+	AIMODULE_API UAISystem(const FObjectInitializer& ObjectInitializer);
 
-	virtual void BeginDestroy() override;
+	AIMODULE_API virtual void BeginDestroy() override;
 	
-	virtual void PostInitProperties() override;
+	AIMODULE_API virtual void PostInitProperties() override;
 
 	// UAISystemBase begin		
-	virtual void InitializeActorsForPlay(bool bTimeGotReset) override;
-	virtual void WorldOriginLocationChanged(FIntVector OldOriginLocation, FIntVector NewOriginLocation) override;
-	virtual void CleanupWorld(bool bSessionEnded = true, bool bCleanupResources = true) override;
+	AIMODULE_API virtual void InitializeActorsForPlay(bool bTimeGotReset) override;
+	AIMODULE_API virtual void WorldOriginLocationChanged(FIntVector OldOriginLocation, FIntVector NewOriginLocation) override;
+	AIMODULE_API virtual void CleanupWorld(bool bSessionEnded = true, bool bCleanupResources = true) override;
 	UE_DEPRECATED(5.1, "NewWorld was unused and not always calculated correctly and we expect it is not needed; let us know on UDN if it is necessary.")
-	virtual void CleanupWorld(bool bSessionEnded, bool bCleanupResources, UWorld* NewWorld) override;
-	virtual void StartPlay() override;
+	AIMODULE_API virtual void CleanupWorld(bool bSessionEnded, bool bCleanupResources, UWorld* NewWorld) override;
+	AIMODULE_API virtual void StartPlay() override;
 	// UAISystemBase end
 
 	/** Behavior tree manager getter */
@@ -193,13 +193,13 @@ public:
 	// cheats
 	//----------------------------------------------------------------------//
 	UFUNCTION(exec)
-	virtual void AIIgnorePlayers();
+	AIMODULE_API virtual void AIIgnorePlayers();
 
 	UFUNCTION(exec)
-	virtual void AILoggingVerbose();
+	AIMODULE_API virtual void AILoggingVerbose();
 
 	/** insta-runs EQS query for given Target */
-	void RunEQS(const FString& QueryName, UObject* Target);
+	AIMODULE_API void RunEQS(const FString& QueryName, UObject* Target);
 
 	/**
 	* Iterator for traversing all UBlackboardComponent instances associated
@@ -257,30 +257,30 @@ public:
 	* (i.e. InitializeComponent). The user is responsible for calling
 	* UnregisterBlackboardComponent (i.e. UninitializeComponent).
 	*/
-	void RegisterBlackboardComponent(class UBlackboardData& BlackboardAsset, class UBlackboardComponent& BlackboardComp);
+	AIMODULE_API void RegisterBlackboardComponent(class UBlackboardData& BlackboardAsset, class UBlackboardComponent& BlackboardComp);
 
 	/**
 	* Unregisters a UBlackboardComponent instance with this blackboard data
 	* asset. This should be called before the component has been uninitialized
 	* (i.e. UninitializeComponent).
 	*/
-	void UnregisterBlackboardComponent(class UBlackboardData& BlackboardAsset, class UBlackboardComponent& BlackboardComp);
+	AIMODULE_API void UnregisterBlackboardComponent(class UBlackboardData& BlackboardAsset, class UBlackboardComponent& BlackboardComp);
 
 	/**
 	* Creates a forward only iterator for that will iterate all
 	* UBlackboardComponent instances that reference the specified
 	* BlackboardAsset and it's parents.
 	*/
-	FBlackboardDataToComponentsIterator CreateBlackboardDataToComponentsIterator(class UBlackboardData& BlackboardAsset);
+	AIMODULE_API FBlackboardDataToComponentsIterator CreateBlackboardDataToComponentsIterator(class UBlackboardData& BlackboardAsset);
 
-	virtual void ConditionalLoadDebuggerPlugin();
+	AIMODULE_API virtual void ConditionalLoadDebuggerPlugin();
 
 	static const FRandomStream& GetRandomStream() { return RandomStream; }
 	static void SeedRandomStream(const int32 Seed) { return RandomStream.Initialize(Seed); }
 
 protected:
-	virtual void OnActorSpawned(AActor* SpawnedActor);
-	virtual void OnPawnBeginPlay(APawn* Pawn);
+	AIMODULE_API virtual void OnActorSpawned(AActor* SpawnedActor);
+	AIMODULE_API virtual void OnPawnBeginPlay(APawn* Pawn);
 
-	void LoadDebuggerPlugin();
+	AIMODULE_API void LoadDebuggerPlugin();
 };
diff --git a/Engine/Source/Runtime/AIModule/Classes/AITypes.h b/Engine/Source/Runtime/AIModule/Classes/AITypes.h
--- a/Engine/Source/Runtime/AIModule/Classes/AITypes.h
+++ b/Engine/Source/Runtime/AIModule/Classes/AITypes.h
@@ -325,12 +325,12 @@ public:
 };
 
 //////////////////////////////////////////////////////////////////////////
-struct AIMODULE_API FAIResCounter : FAIBasicCounter<uint8>
+struct FAIResCounter : FAIBasicCounter<uint8>
 {};
 typedef FAINamedID<FAIResCounter> FAIResourceID;
 
 //////////////////////////////////////////////////////////////////////////
-struct AIMODULE_API FAIResourcesSet
+struct FAIResourcesSet
 {
 	static const uint32 NoResources = 0;
 	static const uint32 AllResources = uint32(-1);
@@ -357,21 +357,21 @@ public:
 };
 
 /** structure used to define which subsystem requested locking of a specific AI resource (like movement, logic, etc.) */
-struct AIMODULE_API FAIResourceLock
+struct FAIResourceLock
 {
 	/** @note feel free to change the type if you need to support more then 16 lock sources */
 	typedef uint16 FLockFlags;
 	
-	FAIResourceLock();
+	AIMODULE_API FAIResourceLock();
 	
-	void SetLock(EAIRequestPriority::Type LockPriority);
-	void ClearLock(EAIRequestPriority::Type LockPriority);
+	AIMODULE_API void SetLock(EAIRequestPriority::Type LockPriority);
+	AIMODULE_API void ClearLock(EAIRequestPriority::Type LockPriority);
 
 	/** set whether we should use resource lock count.  clears all existing locks. */
-	void SetUseResourceLockCount(bool inUseResourceLockCount);
+	AIMODULE_API void SetUseResourceLockCount(bool inUseResourceLockCount);
 
 	/** force-clears all locks */
-	void ForceClearAllLocks();
+	AIMODULE_API void ForceClearAllLocks();
 
 	FORCEINLINE bool IsLocked() const
 	{
@@ -398,7 +398,7 @@ struct AIMODULE_API FAIResourceLock
 		return true;
 	}
 
-	FString GetLockPriorityName() const;
+	AIMODULE_API FString GetLockPriorityName() const;
 
 	void operator+=(const FAIResourceLock& Other)
 	{
@@ -431,7 +431,7 @@ namespace FAIResources
 
 
 USTRUCT(BlueprintType)
-struct AIMODULE_API FAIRequestID
+struct FAIRequestID
 {
 	GENERATED_USTRUCT_BODY()
 		
@@ -480,9 +480,9 @@ public:
 		return FString::FromInt(int32(RequestID));
 	}
 
-	static const FAIRequestID AnyRequest;
-	static const FAIRequestID CurrentRequest;
-	static const FAIRequestID InvalidRequest;
+	static AIMODULE_API const FAIRequestID AnyRequest;
+	static AIMODULE_API const FAIRequestID CurrentRequest;
+	static AIMODULE_API const FAIRequestID InvalidRequest;
 };
 
 //////////////////////////////////////////////////////////////////////////
@@ -490,13 +490,13 @@ public:
 class UNavigationQueryFilter;
 
 USTRUCT()
-struct AIMODULE_API FAIMoveRequest
+struct FAIMoveRequest
 {
 	GENERATED_USTRUCT_BODY()
 
-	FAIMoveRequest();
-	FAIMoveRequest(const AActor* InGoalActor);
-	FAIMoveRequest(const FVector& InGoalLocation);
+	AIMODULE_API FAIMoveRequest();
+	AIMODULE_API FAIMoveRequest(const AActor* InGoalActor);
+	AIMODULE_API FAIMoveRequest(const FVector& InGoalLocation);
 
 	FAIMoveRequest& SetNavigationFilter(TSubclassOf<UNavigationQueryFilter> Filter) { FilterClass = Filter; return *this; }
 	FAIMoveRequest& SetUsePathfinding(bool bPathfinding) { bUsePathfinding = bPathfinding; return *this; }
@@ -543,11 +543,11 @@ struct AIMODULE_API FAIMoveRequest
 	const FCustomMoveSharedPtr& GetUserData() const { return UserData; }
 	int32 GetUserFlags() const { return UserFlags; }
 
-	void SetGoalActor(const AActor* InGoalActor);
-	void SetGoalLocation(const FVector& InGoalLocation);
+	AIMODULE_API void SetGoalActor(const AActor* InGoalActor);
+	AIMODULE_API void SetGoalLocation(const FVector& InGoalLocation);
 
-	bool UpdateGoalLocation(const FVector& NewLocation) const;
-	FString ToString() const;
+	AIMODULE_API bool UpdateGoalLocation(const FVector& NewLocation) const;
+	AIMODULE_API FString ToString() const;
 
 protected:
 
diff --git a/Engine/Source/Runtime/AIModule/Classes/Actions/PawnAction.h b/Engine/Source/Runtime/AIModule/Classes/Actions/PawnAction.h
--- a/Engine/Source/Runtime/AIModule/Classes/Actions/PawnAction.h
+++ b/Engine/Source/Runtime/AIModule/Classes/Actions/PawnAction.h
@@ -43,8 +43,8 @@ namespace EPawnActionFailHandling
  *	Things to remember:
  *	* Actions are created paused
  */
-UCLASS(abstract, EditInlineNew, deprecated, meta = (DeprecationMessage = "PawnActions have been deprecated and are no longer being supported. It will get removed in following UE5 releases. Use GameplayTasks or AITasks instead."))
-class AIMODULE_API UDEPRECATED_PawnAction : public UObject
+UCLASS(abstract, EditInlineNew, deprecated, meta = (DeprecationMessage = "PawnActions have been deprecated and are no longer being supported. It will get removed in following UE5 releases. Use GameplayTasks or AITasks instead."), MinimalAPI)
+class UDEPRECATED_PawnAction : public UObject
 {
 	GENERATED_UCLASS_BODY()
 
@@ -134,7 +134,7 @@ private:
 public:
 
 	// Begin UObject
-	virtual UWorld* GetWorld() const override;
+	AIMODULE_API virtual UWorld* GetWorld() const override;
 	// End UObject
 
 	FORCEINLINE const UDEPRECATED_PawnAction* GetParentAction() const { return ParentAction_DEPRECATED; }
@@ -167,7 +167,7 @@ protected:
 	 *	@param bForce
 	 *	@return current state of task abort
 	 *	@NOTE do not make this virtual! Contains some essential logic. */
-	EPawnActionAbortState::Type Abort(EAIForceParam::Type ShouldForce = EAIForceParam::DoNotForce);
+	AIMODULE_API EPawnActionAbortState::Type Abort(EAIForceParam::Type ShouldForce = EAIForceParam::DoNotForce);
 	
 	FORCEINLINE UDEPRECATED_PawnActionsComponent* GetOwnerComponent() { return OwnerComponent_DEPRECATED; }
 public:
@@ -176,8 +176,8 @@ public:
 	FORCEINLINE EPawnActionAbortState::Type GetAbortState() const { return AbortState; }
 	FORCEINLINE UDEPRECATED_PawnActionsComponent* GetOwnerComponent() const { return OwnerComponent_DEPRECATED; }
 	FORCEINLINE UObject* GetInstigator() const { return Instigator; }
-	APawn* GetPawn() const;
-	AController* GetController() const;
+	AIMODULE_API APawn* GetPawn() const;
+	AIMODULE_API AController* GetController() const;
 
 	template<class TActionClass>
 	static TActionClass* CreateActionInstance(UWorld& World)
@@ -189,7 +189,7 @@ public:
 	//----------------------------------------------------------------------//
 	// messaging
 	//----------------------------------------------------------------------//
-	void WaitForMessage(FName MessageType, FAIRequestID RequestID = FAIRequestID::AnyRequest);
+	AIMODULE_API void WaitForMessage(FName MessageType, FAIRequestID RequestID = FAIRequestID::AnyRequest);
 	// @note this function will change its signature once AI messaging is rewritten @todo
 	virtual void HandleAIMessage(UBrainComponent*, const FAIMessage&){};
 
@@ -200,56 +200,56 @@ public:
 	// Blueprint interface
 	//----------------------------------------------------------------------//
 	UFUNCTION(BlueprintPure, Category = "AI|PawnActions")
-	TEnumAsByte<EAIRequestPriority::Type> GetActionPriority();
+	AIMODULE_API TEnumAsByte<EAIRequestPriority::Type> GetActionPriority();
 
 	UE_DEPRECATED(5.2, "PawnActions have been deprecated and are no longer being supported. It will get removed in following UE5 releases. Use GameplayTasks or AITasks instead.")
 	UFUNCTION(BlueprintCallable, Category = "AI|PawnActions", meta = (WorldContext="WorldContextObject", DeprecatedFunction, DeprecationMessage = "PawnActions have been deprecated and are no longer being supported. It will get removed in following UE5 releases. Use GameplayTasks or AITasks instead."))
-	static UDEPRECATED_PawnAction* CreateActionInstance(UObject* WorldContextObject, TSubclassOf<UDEPRECATED_PawnAction> ActionClass);
+	static AIMODULE_API UDEPRECATED_PawnAction* CreateActionInstance(UObject* WorldContextObject, TSubclassOf<UDEPRECATED_PawnAction> ActionClass);
 
 	//----------------------------------------------------------------------//
 	// debug
 	//----------------------------------------------------------------------//
-	FString GetStateDescription() const;
-	FString GetPriorityName() const;
-	virtual FString GetDisplayName() const;
+	AIMODULE_API FString GetStateDescription() const;
+	AIMODULE_API FString GetPriorityName() const;
+	AIMODULE_API virtual FString GetDisplayName() const;
 
 protected:
 
 	/** starts or resumes action, depending on internal state */
-	bool Activate();
-	void OnPopped();
+	AIMODULE_API bool Activate();
+	AIMODULE_API void OnPopped();
 
 	UFUNCTION(BlueprintCallable, Category = "AI|PawnActions")
-	virtual void Finish(TEnumAsByte<EPawnActionResult::Type> WithResult);
+	AIMODULE_API virtual void Finish(TEnumAsByte<EPawnActionResult::Type> WithResult);
 
-	void SendEvent(EPawnActionEventType::Type Event);
+	AIMODULE_API void SendEvent(EPawnActionEventType::Type Event);
 
-	void StopWaitingForMessages();
+	AIMODULE_API void StopWaitingForMessages();
 
-	void SetOwnerComponent(UDEPRECATED_PawnActionsComponent* Component);
+	AIMODULE_API void SetOwnerComponent(UDEPRECATED_PawnActionsComponent* Component);
 
-	void SetInstigator(UObject* const InInstigator);
+	AIMODULE_API void SetInstigator(UObject* const InInstigator);
 
-	virtual void Tick(float DeltaTime);
+	AIMODULE_API virtual void Tick(float DeltaTime);
 
 	/** called to start off the Action
 	 *	@return 'true' if actions successfully started. 
 	 *	@NOTE if action fails to start no finishing or aborting mechanics will be triggered */
-	virtual bool Start();
+	AIMODULE_API virtual bool Start();
 	/** called to pause action when higher priority or child action kicks in */
-	virtual bool Pause(const UDEPRECATED_PawnAction* PausedBy);
+	AIMODULE_API virtual bool Pause(const UDEPRECATED_PawnAction* PausedBy);
 	/** called to resume action after being paused */
-	virtual bool Resume();
+	AIMODULE_API virtual bool Resume();
 	/** called when this action is being removed from action stacks */
-	virtual void OnFinished(EPawnActionResult::Type WithResult);
+	AIMODULE_API virtual void OnFinished(EPawnActionResult::Type WithResult);
 	/** called to give Action chance to react to child action finishing.
 	 *	@NOTE gets called _AFTER_ child's OnFinished to give child action chance 
 	 *		to prepare "finishing data" for parent to read. 
 	 *	@NOTE clears parent-child binding */
-	virtual void OnChildFinished(UDEPRECATED_PawnAction& Action, EPawnActionResult::Type WithResult);
+	AIMODULE_API virtual void OnChildFinished(UDEPRECATED_PawnAction& Action, EPawnActionResult::Type WithResult);
 
 	/** apart from doing regular push request copies additional values from Parent, like Priority and Instigator */
-	bool PushChildAction(UDEPRECATED_PawnAction& Action);
+	AIMODULE_API bool PushChildAction(UDEPRECATED_PawnAction& Action);
 	
 	/** performs actual work on aborting Action. Should be called exclusively by Abort function
 	 *	@return only valid return values here are LatendAbortInProgress and AbortDone */
@@ -260,13 +260,13 @@ protected:
 private:
 	/** called when this action is put on a stack. Does not indicate action will be started soon
 	 *	(it depends on other actions on other action stacks. Called before Start() call */
-	void OnPushed();
+	AIMODULE_API void OnPushed();
 
 	/** Sets final result for this Action. To be called only once upon Action's finish */
-	void SetFinishResult(EPawnActionResult::Type Result);
+	AIMODULE_API void SetFinishResult(EPawnActionResult::Type Result);
 
 	// do not un-private. Internal logic only!
-	void SetAbortState(EPawnActionAbortState::Type NewAbortState);
+	AIMODULE_API void SetAbortState(EPawnActionAbortState::Type NewAbortState);
 };
 
 
diff --git a/Engine/Source/Runtime/AIModule/Classes/Actions/PawnAction_BlueprintBase.h b/Engine/Source/Runtime/AIModule/Classes/Actions/PawnAction_BlueprintBase.h
--- a/Engine/Source/Runtime/AIModule/Classes/Actions/PawnAction_BlueprintBase.h
+++ b/Engine/Source/Runtime/AIModule/Classes/Actions/PawnAction_BlueprintBase.h
@@ -9,8 +9,8 @@
 
 class APawn;
 
-UCLASS(abstract, Blueprintable)
-class AIMODULE_API UDEPRECATED_PawnAction_BlueprintBase : public UDEPRECATED_PawnAction
+UCLASS(abstract, Blueprintable, MinimalAPI)
+class UDEPRECATED_PawnAction_BlueprintBase : public UDEPRECATED_PawnAction
 {
 	GENERATED_UCLASS_BODY()
 
@@ -21,20 +21,20 @@ public:
 	//----------------------------------------------------------------------//
 
 	UFUNCTION(BlueprintImplementableEvent, Category = "AI|PawnActions")
-	void ActionStart(APawn* ControlledPawn);
+	AIMODULE_API void ActionStart(APawn* ControlledPawn);
 	UFUNCTION(BlueprintImplementableEvent, Category = "AI|PawnActions")
-	void ActionTick(APawn* ControlledPawn, float DeltaSeconds);
+	AIMODULE_API void ActionTick(APawn* ControlledPawn, float DeltaSeconds);
 	UFUNCTION(BlueprintImplementableEvent, Category = "AI|PawnActions")
-	void ActionPause(APawn* ControlledPawn);
+	AIMODULE_API void ActionPause(APawn* ControlledPawn);
 	UFUNCTION(BlueprintImplementableEvent, Category = "AI|PawnActions")
-	void ActionResume(APawn* ControlledPawn);
+	AIMODULE_API void ActionResume(APawn* ControlledPawn);
 	UFUNCTION(BlueprintImplementableEvent, Category = "AI|PawnActions")
-	void ActionFinished(APawn* ControlledPawn, EPawnActionResult::Type WithResult);
+	AIMODULE_API void ActionFinished(APawn* ControlledPawn, EPawnActionResult::Type WithResult);
 
 protected:
-	virtual void Tick(float DeltaTime) override;
-	virtual bool Start() override;
-	virtual bool Pause(const UDEPRECATED_PawnAction* PausedBy) override;
-	virtual bool Resume() override;
-	virtual void OnFinished(EPawnActionResult::Type WithResult) override;
+	AIMODULE_API virtual void Tick(float DeltaTime) override;
+	AIMODULE_API virtual bool Start() override;
+	AIMODULE_API virtual bool Pause(const UDEPRECATED_PawnAction* PausedBy) override;
+	AIMODULE_API virtual bool Resume() override;
+	AIMODULE_API virtual void OnFinished(EPawnActionResult::Type WithResult) override;
 };
diff --git a/Engine/Source/Runtime/AIModule/Classes/Actions/PawnAction_Move.h b/Engine/Source/Runtime/AIModule/Classes/Actions/PawnAction_Move.h
--- a/Engine/Source/Runtime/AIModule/Classes/Actions/PawnAction_Move.h
+++ b/Engine/Source/Runtime/AIModule/Classes/Actions/PawnAction_Move.h
@@ -23,8 +23,8 @@ namespace EPawnActionMoveMode
 	};
 }
 
-UCLASS()
-class AIMODULE_API UDEPRECATED_PawnAction_Move : public UDEPRECATED_PawnAction
+UCLASS(MinimalAPI)
+class UDEPRECATED_PawnAction_Move : public UDEPRECATED_PawnAction
 {
 	GENERATED_UCLASS_BODY()
 protected:
@@ -69,18 +69,18 @@ protected:
 	uint32 bAbortSubActionOnPathChange : 1;
 
 public:
-	virtual void BeginDestroy() override;
+	AIMODULE_API virtual void BeginDestroy() override;
 
-	static UDEPRECATED_PawnAction_Move* CreateAction(UWorld& World, AActor* GoalActor, EPawnActionMoveMode::Type Mode);
-	static UDEPRECATED_PawnAction_Move* CreateAction(UWorld& World, const FVector& GoalLocation, EPawnActionMoveMode::Type Mode);
+	static AIMODULE_API UDEPRECATED_PawnAction_Move* CreateAction(UWorld& World, AActor* GoalActor, EPawnActionMoveMode::Type Mode);
+	static AIMODULE_API UDEPRECATED_PawnAction_Move* CreateAction(UWorld& World, const FVector& GoalLocation, EPawnActionMoveMode::Type Mode);
 
-	static bool CheckAlreadyAtGoal(AAIController& Controller, const FVector& TestLocation, float Radius);
-	static bool CheckAlreadyAtGoal(AAIController& Controller, const AActor& TestGoal, float Radius);
+	static AIMODULE_API bool CheckAlreadyAtGoal(AAIController& Controller, const FVector& TestLocation, float Radius);
+	static AIMODULE_API bool CheckAlreadyAtGoal(AAIController& Controller, const AActor& TestGoal, float Radius);
 
-	virtual void HandleAIMessage(UBrainComponent*, const FAIMessage&) override;
+	AIMODULE_API virtual void HandleAIMessage(UBrainComponent*, const FAIMessage&) override;
 
-	void SetPath(FNavPathSharedRef InPath);
-	virtual void OnPathUpdated(FNavigationPath* UpdatedPath, ENavPathEvent::Type Event);
+	AIMODULE_API void SetPath(FNavPathSharedRef InPath);
+	AIMODULE_API virtual void OnPathUpdated(FNavigationPath* UpdatedPath, ENavPathEvent::Type Event);
 
 	void SetAcceptableRadius(float NewAcceptableRadius) { AcceptableRadius = NewAcceptableRadius; }
 	void SetFinishOnOverlap(bool bNewFinishOnOverlap) { bFinishOnOverlap = bNewFinishOnOverlap; }
@@ -106,20 +106,20 @@ protected:
 	/** Handle for efficient management of TryToRepath timer */
 	FTimerHandle TimerHandle_TryToRepath;
 
-	void ClearPath();
-	virtual bool Start() override;
-	virtual bool Pause(const UDEPRECATED_PawnAction* PausedBy) override;
-	virtual bool Resume() override;
-	virtual void OnFinished(EPawnActionResult::Type WithResult) override;
-	virtual EPawnActionAbortState::Type PerformAbort(EAIForceParam::Type ShouldForce) override;
-	virtual bool IsPartialPathAllowed() const;
+	AIMODULE_API void ClearPath();
+	AIMODULE_API virtual bool Start() override;
+	AIMODULE_API virtual bool Pause(const UDEPRECATED_PawnAction* PausedBy) override;
+	AIMODULE_API virtual bool Resume() override;
+	AIMODULE_API virtual void OnFinished(EPawnActionResult::Type WithResult) override;
+	AIMODULE_API virtual EPawnActionAbortState::Type PerformAbort(EAIForceParam::Type ShouldForce) override;
+	AIMODULE_API virtual bool IsPartialPathAllowed() const;
 
-	virtual EPathFollowingRequestResult::Type RequestMove(AAIController& Controller);
+	AIMODULE_API virtual EPathFollowingRequestResult::Type RequestMove(AAIController& Controller);
 	
-	bool PerformMoveAction();
-	void DeferredPerformMoveAction();
+	AIMODULE_API bool PerformMoveAction();
+	AIMODULE_API void DeferredPerformMoveAction();
 
-	void TryToRepath();
-	void ClearPendingRepath();
-	void ClearTimers();
+	AIMODULE_API void TryToRepath();
+	AIMODULE_API void ClearPendingRepath();
+	AIMODULE_API void ClearTimers();
 };
diff --git a/Engine/Source/Runtime/AIModule/Classes/Actions/PawnAction_Repeat.h b/Engine/Source/Runtime/AIModule/Classes/Actions/PawnAction_Repeat.h
--- a/Engine/Source/Runtime/AIModule/Classes/Actions/PawnAction_Repeat.h
+++ b/Engine/Source/Runtime/AIModule/Classes/Actions/PawnAction_Repeat.h
@@ -7,8 +7,8 @@
 #include "Actions/PawnAction.h"
 #include "PawnAction_Repeat.generated.h"
 
-UCLASS()
-class AIMODULE_API UDEPRECATED_PawnAction_Repeat : public UDEPRECATED_PawnAction
+UCLASS(MinimalAPI)
+class UDEPRECATED_PawnAction_Repeat : public UDEPRECATED_PawnAction
 {
 	GENERATED_UCLASS_BODY()
 
@@ -32,12 +32,12 @@ class AIMODULE_API UDEPRECATED_PawnAction_Repeat : public UDEPRECATED_PawnAction
 	EPawnSubActionTriggeringPolicy::Type SubActionTriggeringPolicy;
 
 	/** @param NumberOfRepeats number of times to repeat action. UDEPRECATED_PawnAction_Repeat::LoopForever loops forever */
-	static UDEPRECATED_PawnAction_Repeat* CreateAction(UWorld& World, UDEPRECATED_PawnAction* ActionToRepeat, int32 NumberOfRepeats, EPawnSubActionTriggeringPolicy::Type InSubActionTriggeringPolicy = EPawnSubActionTriggeringPolicy::CopyBeforeTriggering);
+	static AIMODULE_API UDEPRECATED_PawnAction_Repeat* CreateAction(UWorld& World, UDEPRECATED_PawnAction* ActionToRepeat, int32 NumberOfRepeats, EPawnSubActionTriggeringPolicy::Type InSubActionTriggeringPolicy = EPawnSubActionTriggeringPolicy::CopyBeforeTriggering);
 
 protected:
-	virtual bool Start() override;
-	virtual bool Resume() override;
-	virtual void OnChildFinished(UDEPRECATED_PawnAction& Action, EPawnActionResult::Type WithResult) override;
+	AIMODULE_API virtual bool Start() override;
+	AIMODULE_API virtual bool Resume() override;
+	AIMODULE_API virtual void OnChildFinished(UDEPRECATED_PawnAction& Action, EPawnActionResult::Type WithResult) override;
 
-	bool PushSubAction();
+	AIMODULE_API bool PushSubAction();
 };
diff --git a/Engine/Source/Runtime/AIModule/Classes/Actions/PawnAction_Sequence.h b/Engine/Source/Runtime/AIModule/Classes/Actions/PawnAction_Sequence.h
--- a/Engine/Source/Runtime/AIModule/Classes/Actions/PawnAction_Sequence.h
+++ b/Engine/Source/Runtime/AIModule/Classes/Actions/PawnAction_Sequence.h
@@ -7,8 +7,8 @@
 #include "Actions/PawnAction.h"
 #include "PawnAction_Sequence.generated.h"
 
-UCLASS()
-class AIMODULE_API UDEPRECATED_PawnAction_Sequence : public UDEPRECATED_PawnAction
+UCLASS(MinimalAPI)
+class UDEPRECATED_PawnAction_Sequence : public UDEPRECATED_PawnAction
 {
 	GENERATED_UCLASS_BODY()
 
@@ -25,12 +25,12 @@ class AIMODULE_API UDEPRECATED_PawnAction_Sequence : public UDEPRECATED_PawnActi
 
 	EPawnSubActionTriggeringPolicy::Type SubActionTriggeringPolicy;
 
-	static UDEPRECATED_PawnAction_Sequence* CreateAction(UWorld& World, TArray<UDEPRECATED_PawnAction*>& ActionSequence, EPawnSubActionTriggeringPolicy::Type InSubActionTriggeringPolicy = EPawnSubActionTriggeringPolicy::CopyBeforeTriggering);
+	static AIMODULE_API UDEPRECATED_PawnAction_Sequence* CreateAction(UWorld& World, TArray<UDEPRECATED_PawnAction*>& ActionSequence, EPawnSubActionTriggeringPolicy::Type InSubActionTriggeringPolicy = EPawnSubActionTriggeringPolicy::CopyBeforeTriggering);
 
 protected:
-	virtual bool Start() override;
-	virtual bool Resume() override;
-	virtual void OnChildFinished(UDEPRECATED_PawnAction& Action, EPawnActionResult::Type WithResult) override;
+	AIMODULE_API virtual bool Start() override;
+	AIMODULE_API virtual bool Resume() override;
+	AIMODULE_API virtual void OnChildFinished(UDEPRECATED_PawnAction& Action, EPawnActionResult::Type WithResult) override;
 
-	bool PushNextActionCopy();
+	AIMODULE_API bool PushNextActionCopy();
 };
diff --git a/Engine/Source/Runtime/AIModule/Classes/Actions/PawnAction_Wait.h b/Engine/Source/Runtime/AIModule/Classes/Actions/PawnAction_Wait.h
--- a/Engine/Source/Runtime/AIModule/Classes/Actions/PawnAction_Wait.h
+++ b/Engine/Source/Runtime/AIModule/Classes/Actions/PawnAction_Wait.h
@@ -9,8 +9,8 @@
 #include "PawnAction_Wait.generated.h"
 
 /** uses system timers rather then ticking */
-UCLASS()
-class AIMODULE_API UDEPRECATED_PawnAction_Wait : public UDEPRECATED_PawnAction
+UCLASS(MinimalAPI)
+class UDEPRECATED_PawnAction_Wait : public UDEPRECATED_PawnAction
 {
 	GENERATED_UCLASS_BODY()
 		
@@ -22,12 +22,12 @@ class AIMODULE_API UDEPRECATED_PawnAction_Wait : public UDEPRECATED_PawnAction
 	FTimerHandle TimerHandle;
 
 	/** InTimeToWait < 0 (or just FAISystem::InfiniteInterval) will result in waiting forever */
-	static UDEPRECATED_PawnAction_Wait* CreateAction(UWorld& World, float InTimeToWait = FAISystem::InfiniteInterval);
+	static AIMODULE_API UDEPRECATED_PawnAction_Wait* CreateAction(UWorld& World, float InTimeToWait = FAISystem::InfiniteInterval);
 
-	virtual bool Start() override;
-	virtual bool Pause(const UDEPRECATED_PawnAction* PausedBy) override;
-	virtual bool Resume() override;
-	virtual EPawnActionAbortState::Type PerformAbort(EAIForceParam::Type ShouldForce) override;
+	AIMODULE_API virtual bool Start() override;
+	AIMODULE_API virtual bool Pause(const UDEPRECATED_PawnAction* PausedBy) override;
+	AIMODULE_API virtual bool Resume() override;
+	AIMODULE_API virtual EPawnActionAbortState::Type PerformAbort(EAIForceParam::Type ShouldForce) override;
 
-	void TimerDone();
+	AIMODULE_API void TimerDone();
 };
diff --git a/Engine/Source/Runtime/AIModule/Classes/Actions/PawnActionsComponent.h b/Engine/Source/Runtime/AIModule/Classes/Actions/PawnActionsComponent.h
--- a/Engine/Source/Runtime/AIModule/Classes/Actions/PawnActionsComponent.h
+++ b/Engine/Source/Runtime/AIModule/Classes/Actions/PawnActionsComponent.h
@@ -14,7 +14,7 @@
 class AController;
 
 USTRUCT()
-struct AIMODULE_API FPawnActionEvent
+struct FPawnActionEvent
 {
 	GENERATED_USTRUCT_BODY()
 
@@ -34,13 +34,13 @@ struct AIMODULE_API FPawnActionEvent
 	FPawnActionEvent() : Action_DEPRECATED(NULL), EventType(EPawnActionEventType::Invalid), Priority(EAIRequestPriority::MAX), Index(uint32(-1))
 	{}
 
-	FPawnActionEvent(UDEPRECATED_PawnAction& Action, EPawnActionEventType::Type EventType, uint32 Index);
+	AIMODULE_API FPawnActionEvent(UDEPRECATED_PawnAction& Action, EPawnActionEventType::Type EventType, uint32 Index);
 
 	bool operator==(const FPawnActionEvent& Other) const { return (Action_DEPRECATED == Other.Action_DEPRECATED) && (EventType == Other.EventType) && (Priority == Other.Priority); }
 };
 
 USTRUCT()
-struct AIMODULE_API FPawnActionStack
+struct FPawnActionStack
 {
 	GENERATED_USTRUCT_BODY()
 
@@ -53,16 +53,16 @@ private:
 	TObjectPtr<UDEPRECATED_PawnAction> TopAction_DEPRECATED;
 
 public:
-	void Pause();
-	void Resume();
+	AIMODULE_API void Pause();
+	AIMODULE_API void Resume();
 
 	/** All it does is tie actions into a double-linked list making NewTopAction
 	 *	new stack's top */
-	void PushAction(UDEPRECATED_PawnAction& NewTopAction);
+	AIMODULE_API void PushAction(UDEPRECATED_PawnAction& NewTopAction);
 
 	/** Looks through the double-linked action list looking for specified action
 	 *	and if found action will be popped along with all it's siblings */
-	void PopAction(UDEPRECATED_PawnAction& ActionToPop);
+	AIMODULE_API void PopAction(UDEPRECATED_PawnAction& ActionToPop);
 	
 	FORCEINLINE UDEPRECATED_PawnAction* GetTop() const { return TopAction_DEPRECATED; }
 
@@ -71,11 +71,11 @@ public:
 	//----------------------------------------------------------------------//
 	// Debugging-testing purposes 
 	//----------------------------------------------------------------------//
-	int32 GetStackSize() const;
+	AIMODULE_API int32 GetStackSize() const;
 };
 
-UCLASS(deprecated, meta = (DeprecationMessage = "PawnActions have been deprecated and are no longer being supported. It will get removed in following UE5 releases. Use GameplayTasks or AITasks instead."))
-class AIMODULE_API UDEPRECATED_PawnActionsComponent : public UActorComponent
+UCLASS(deprecated, meta = (DeprecationMessage = "PawnActions have been deprecated and are no longer being supported. It will get removed in following UE5 releases. Use GameplayTasks or AITasks instead."), MinimalAPI)
+class UDEPRECATED_PawnActionsComponent : public UActorComponent
 {
 	GENERATED_UCLASS_BODY()
 
@@ -102,17 +102,17 @@ public:
 	//----------------------------------------------------------------------//
 	// UActorComponent
 	//----------------------------------------------------------------------//
-	virtual void OnUnregister() override;
+	AIMODULE_API virtual void OnUnregister() override;
 
 	//----------------------------------------------------------------------//
 	// blueprint interface
 	//----------------------------------------------------------------------//
 
 	UFUNCTION(BlueprintCallable, Category = "AI|PawnActions", meta = (DisplayName = "PerformAction_DEPRECATED", ScriptName = "PerformAction", DeprecatedFunction, DeprecationMessage = "PawnActions have been deprecated and are no longer being supported. It will get removed in following UE5 releases. Use GameplayTasks or AITasks instead."))
-	static bool K2_PerformAction(APawn* Pawn, UDEPRECATED_PawnAction* Action, TEnumAsByte<EAIRequestPriority::Type> Priority = EAIRequestPriority::HardScript);
+	static AIMODULE_API bool K2_PerformAction(APawn* Pawn, UDEPRECATED_PawnAction* Action, TEnumAsByte<EAIRequestPriority::Type> Priority = EAIRequestPriority::HardScript);
 
 	UE_DEPRECATED(5.2, "PawnActions have been deprecated and are no longer being supported. It will get removed in following UE5 releases. Use GameplayTasks or AITasks instead.")
-	static bool PerformAction(APawn& Pawn, UDEPRECATED_PawnAction& Action, TEnumAsByte<EAIRequestPriority::Type> Priority = EAIRequestPriority::HardScript);
+	static AIMODULE_API bool PerformAction(APawn& Pawn, UDEPRECATED_PawnAction& Action, TEnumAsByte<EAIRequestPriority::Type> Priority = EAIRequestPriority::HardScript);
 
 	//----------------------------------------------------------------------//
 	// 
@@ -120,50 +120,50 @@ public:
 	/** Use it to save component work to figure out what it's controlling
 	 *	or if component can't/won't be able to figure it out properly
 	 *	@NOTE will throw a log warning if trying to set ControlledPawn if it's already set */
-	void SetControlledPawn(APawn* NewPawn);
+	AIMODULE_API void SetControlledPawn(APawn* NewPawn);
 	FORCEINLINE APawn* GetControlledPawn() { return ControlledPawn; }
 	FORCEINLINE const APawn* GetControlledPawn() const { return ControlledPawn; }
 	FORCEINLINE AController* GetController() { return ControlledPawn ? ControlledPawn->GetController() : NULL; }
 	FORCEINLINE UDEPRECATED_PawnAction* GetCurrentAction() { return CurrentAction_DEPRECATED; }
 
-	bool OnEvent(UDEPRECATED_PawnAction& Action, EPawnActionEventType::Type Event);
+	AIMODULE_API bool OnEvent(UDEPRECATED_PawnAction& Action, EPawnActionEventType::Type Event);
 
 	UFUNCTION(BlueprintCallable, Category = PawnAction, meta = (DisplayName = "PushAction_DEPRECATED", ScriptName = "PushAction", DeprecatedFunction, DeprecationMessage = "PawnActions have been deprecated and are no longer being supported. It will get removed in following UE5 releases. Use GameplayTasks or AITasks instead."))
-	bool K2_PushAction(UDEPRECATED_PawnAction* NewAction, EAIRequestPriority::Type Priority, UObject* Instigator = NULL);
+	AIMODULE_API bool K2_PushAction(UDEPRECATED_PawnAction* NewAction, EAIRequestPriority::Type Priority, UObject* Instigator = NULL);
 
 	UE_DEPRECATED(5.2, "PawnActions have been deprecated and are no longer being supported. It will get removed in following UE5 releases. Use GameplayTasks or AITasks instead.")
-	bool PushAction(UDEPRECATED_PawnAction& NewAction, EAIRequestPriority::Type Priority, UObject* Instigator = NULL);	
+	AIMODULE_API bool PushAction(UDEPRECATED_PawnAction& NewAction, EAIRequestPriority::Type Priority, UObject* Instigator = NULL);	
 
 	/** Aborts given action instance */
 	UFUNCTION(BlueprintCallable, Category = PawnAction, meta = (DisplayName = "AbortAction_DEPRECATED", ScriptName = "AbortAction", DeprecatedFunction, DeprecationMessage = "PawnActions have been deprecated and are no longer being supported. It will get removed in following UE5 releases. Use GameplayTasks or AITasks instead."))
-	EPawnActionAbortState::Type K2_AbortAction(UDEPRECATED_PawnAction* ActionToAbort);
+	AIMODULE_API EPawnActionAbortState::Type K2_AbortAction(UDEPRECATED_PawnAction* ActionToAbort);
 
 	UE_DEPRECATED(5.2, "PawnActions have been deprecated and are no longer being supported. It will get removed in following UE5 releases. Use GameplayTasks or AITasks instead.")
-	EPawnActionAbortState::Type AbortAction(UDEPRECATED_PawnAction& ActionToAbort);
+	AIMODULE_API EPawnActionAbortState::Type AbortAction(UDEPRECATED_PawnAction& ActionToAbort);
 
 	/** Aborts given action instance */
 	UFUNCTION(BlueprintCallable, Category = PawnAction, meta = (DisplayName = "ForceAbortAction_DEPRECATED", ScriptName = "ForceAbortAction", DeprecatedFunction, DeprecationMessage = "PawnActions have been deprecated and are no longer being supported. It will get removed in following UE5 releases. Use GameplayTasks or AITasks instead."))
-	EPawnActionAbortState::Type K2_ForceAbortAction(UDEPRECATED_PawnAction* ActionToAbort);
+	AIMODULE_API EPawnActionAbortState::Type K2_ForceAbortAction(UDEPRECATED_PawnAction* ActionToAbort);
 
 	UE_DEPRECATED(5.2, "PawnActions have been deprecated and are no longer being supported. It will get removed in following UE5 releases. Use GameplayTasks or AITasks instead.")
-	EPawnActionAbortState::Type ForceAbortAction(UDEPRECATED_PawnAction& ActionToAbort);
+	AIMODULE_API EPawnActionAbortState::Type ForceAbortAction(UDEPRECATED_PawnAction& ActionToAbort);
 
 	/** removes all actions instigated with Priority by Instigator
 	 *	@param Priority if equal to EAIRequestPriority::MAX then all priority queues will be searched. 
 	 *		This is less efficient so use with caution 
 	 *	@return number of action abortions requested (performed asyncronously) */
-	uint32 AbortActionsInstigatedBy(UObject* const Instigator, EAIRequestPriority::Type Priority);
+	AIMODULE_API uint32 AbortActionsInstigatedBy(UObject* const Instigator, EAIRequestPriority::Type Priority);
 	
-	void TickComponent(float DeltaTime, enum ELevelTick TickType, FActorComponentTickFunction *ThisTickFunction) override;
+	AIMODULE_API void TickComponent(float DeltaTime, enum ELevelTick TickType, FActorComponentTickFunction *ThisTickFunction) override;
 
 	FORCEINLINE UDEPRECATED_PawnAction* GetActiveAction(EAIRequestPriority::Type Priority) const { return ActionStacks[Priority].GetTop(); }
-	bool HasActiveActionOfType(EAIRequestPriority::Type Priority, TSubclassOf<UDEPRECATED_PawnAction> PawnActionClass) const;
+	AIMODULE_API bool HasActiveActionOfType(EAIRequestPriority::Type Priority, TSubclassOf<UDEPRECATED_PawnAction> PawnActionClass) const;
 
 #if ENABLE_VISUAL_LOG
-	void DescribeSelfToVisLog(struct FVisualLogEntry* Snapshot) const;
+	AIMODULE_API void DescribeSelfToVisLog(struct FVisualLogEntry* Snapshot) const;
 #endif // ENABLE_VISUAL_LOG
 
-	static FString DescribeEventType(EPawnActionEventType::Type EventType);
+	static AIMODULE_API FString DescribeEventType(EPawnActionEventType::Type EventType);
 
 	//----------------------------------------------------------------------//
 	// Debugging-testing purposes 
@@ -174,13 +174,13 @@ public:
 protected:
 	/** Finds the action that should be running. If it's different from CurrentAction
 	 *	then CurrentAction gets paused and newly selected action gets started up */
-	void UpdateCurrentAction();
+	AIMODULE_API void UpdateCurrentAction();
 
-	APawn* CacheControlledPawn();
+	AIMODULE_API APawn* CacheControlledPawn();
 
-	void UpdateAILogicLock();
+	AIMODULE_API void UpdateAILogicLock();
 
 private:
 	/** Removed all pending action events associated with PawnAction. Private to make sure it's called only in special cases */
-	void RemoveEventsForAction(UDEPRECATED_PawnAction& PawnAction);
+	AIMODULE_API void RemoveEventsForAction(UDEPRECATED_PawnAction& PawnAction);
 };
diff --git a/Engine/Source/Runtime/AIModule/Classes/BehaviorTree/BTAuxiliaryNode.h b/Engine/Source/Runtime/AIModule/Classes/BehaviorTree/BTAuxiliaryNode.h
--- a/Engine/Source/Runtime/AIModule/Classes/BehaviorTree/BTAuxiliaryNode.h
+++ b/Engine/Source/Runtime/AIModule/Classes/BehaviorTree/BTAuxiliaryNode.h
@@ -26,16 +26,16 @@ struct FBTAuxiliaryMemory : public FBTInstancedNodeMemory
  *
  */
 
-UCLASS(Abstract)
-class AIMODULE_API UBTAuxiliaryNode : public UBTNode
+UCLASS(Abstract, MinimalAPI)
+class UBTAuxiliaryNode : public UBTNode
 {
 	GENERATED_UCLASS_BODY()
 
 	/** wrapper for node instancing: OnBecomeRelevant */
-	void WrappedOnBecomeRelevant(UBehaviorTreeComponent& OwnerComp, uint8* NodeMemory) const;
+	AIMODULE_API void WrappedOnBecomeRelevant(UBehaviorTreeComponent& OwnerComp, uint8* NodeMemory) const;
 
 	/** wrapper for node instancing: OnCeaseRelevant */
-	void WrappedOnCeaseRelevant(UBehaviorTreeComponent& OwnerComp, uint8* NodeMemory) const;
+	AIMODULE_API void WrappedOnCeaseRelevant(UBehaviorTreeComponent& OwnerComp, uint8* NodeMemory) const;
 
 	/** wrapper for node instancing: TickNode
 	  * @param OwnerComp	The behavior tree owner of this node
@@ -43,25 +43,25 @@ class AIMODULE_API UBTAuxiliaryNode : public UBTNode
 	  * @param DeltaSeconds		DeltaTime since last call
 	  * @param NextNeededDeltaTime		In out parameter, if this node needs a smaller DeltaTime it is the node's responsibility to change it
 	  * @returns	True if it actually done some processing or false if it was skipped because of not ticking or in between time interval */
-	bool WrappedTickNode(UBehaviorTreeComponent& OwnerComp, uint8* NodeMemory, float DeltaSeconds, float& NextNeededDeltaTime) const;
+	AIMODULE_API bool WrappedTickNode(UBehaviorTreeComponent& OwnerComp, uint8* NodeMemory, float DeltaSeconds, float& NextNeededDeltaTime) const;
 
-	virtual void DescribeRuntimeValues(const UBehaviorTreeComponent& OwnerComp, uint8* NodeMemory, EBTDescriptionVerbosity::Type Verbosity, TArray<FString>& Values) const override;
-	virtual uint16 GetSpecialMemorySize() const override;
+	AIMODULE_API virtual void DescribeRuntimeValues(const UBehaviorTreeComponent& OwnerComp, uint8* NodeMemory, EBTDescriptionVerbosity::Type Verbosity, TArray<FString>& Values) const override;
+	AIMODULE_API virtual uint16 GetSpecialMemorySize() const override;
 
 	/** fill in data about tree structure */
-	void InitializeParentLink(uint8 InChildIndex);
+	AIMODULE_API void InitializeParentLink(uint8 InChildIndex);
 
 	/** @return parent task node */
-	const UBTNode* GetMyNode() const;
+	AIMODULE_API const UBTNode* GetMyNode() const;
 
 	/** @return index of child in parent's array or MAX_uint8 */
-	uint8 GetChildIndex() const;
+	AIMODULE_API uint8 GetChildIndex() const;
 
 	/** Get The next needed deltatime for this node
 	  * @param OwnerComp	The behavior tree owner of this node
 	  * @param NodeMemory	The instance memory of the current node
 	  * @return The next needed DeltaTime */
-	float GetNextNeededDeltaTime(UBehaviorTreeComponent& OwnerComp, uint8* NodeMemory) const;
+	AIMODULE_API float GetNextNeededDeltaTime(UBehaviorTreeComponent& OwnerComp, uint8* NodeMemory) const;
 
 protected:
 
@@ -84,25 +84,25 @@ protected:
 	 * this function should be considered as const (don't modify state of object) if node is not instanced!  
 	 * bNotifyBecomeRelevant must be set to true for this function to be called 
 	 * Calling INIT_AUXILIARY_NODE_NOTIFY_FLAGS in the constructor of the node will set this flag automatically */
-	virtual void OnBecomeRelevant(UBehaviorTreeComponent& OwnerComp, uint8* NodeMemory);
+	AIMODULE_API virtual void OnBecomeRelevant(UBehaviorTreeComponent& OwnerComp, uint8* NodeMemory);
 
 	/** called when auxiliary node becomes inactive
 	 * this function should be considered as const (don't modify state of object) if node is not instanced!  
 	 * bNotifyCeaseRelevant must be set to true for this function to be called 
 	 * Calling INIT_AUXILIARY_NODE_NOTIFY_FLAGS in the constructor of the node will set this flag automatically */
-	virtual void OnCeaseRelevant(UBehaviorTreeComponent& OwnerComp, uint8* NodeMemory);
+	AIMODULE_API virtual void OnCeaseRelevant(UBehaviorTreeComponent& OwnerComp, uint8* NodeMemory);
 
 	/** tick function
 	 * this function should be considered as const (don't modify state of object) if node is not instanced!   
 	 * bNotifyTick must be set to true for this function to be called 
 	 * Calling INIT_AUXILIARY_NODE_NOTIFY_FLAGS in the constructor of the node will set this flag automatically */
-	virtual void TickNode(UBehaviorTreeComponent& OwnerComp, uint8* NodeMemory, float DeltaSeconds);
+	AIMODULE_API virtual void TickNode(UBehaviorTreeComponent& OwnerComp, uint8* NodeMemory, float DeltaSeconds);
 
 	/** sets next tick time */
-	void SetNextTickTime(uint8* NodeMemory, float RemainingTime) const;
+	AIMODULE_API void SetNextTickTime(uint8* NodeMemory, float RemainingTime) const;
 
 	/** gets remaining time for next tick */
-	float GetNextTickRemainingTime(uint8* NodeMemory) const;
+	AIMODULE_API float GetNextTickRemainingTime(uint8* NodeMemory) const;
 	
 	template<typename TickNode,	typename OnBecomeRelevant, typename OnCeaseRelevant>
 	void InitNotifyFlags(TickNode, OnBecomeRelevant, OnCeaseRelevant)
diff --git a/Engine/Source/Runtime/AIModule/Classes/BehaviorTree/BTCompositeNode.h b/Engine/Source/Runtime/AIModule/Classes/BehaviorTree/BTCompositeNode.h
--- a/Engine/Source/Runtime/AIModule/Classes/BehaviorTree/BTCompositeNode.h
+++ b/Engine/Source/Runtime/AIModule/Classes/BehaviorTree/BTCompositeNode.h
@@ -83,8 +83,8 @@ struct FBTCompositeChild
 	TArray<FBTDecoratorLogic> DecoratorOps;
 };
 
-UCLASS(Abstract)
-class AIMODULE_API UBTCompositeNode : public UBTNode
+UCLASS(Abstract, MinimalAPI)
+class UBTCompositeNode : public UBTNode
 {
 	GENERATED_UCLASS_BODY()
 
@@ -96,22 +96,22 @@ class AIMODULE_API UBTCompositeNode : public UBTNode
 	UPROPERTY()
 	TArray<TObjectPtr<UBTService>> Services;
 
-	~UBTCompositeNode();
+	AIMODULE_API ~UBTCompositeNode();
 
 	/** fill in data about tree structure */
-	void InitializeComposite(uint16 InLastExecutionIndex);
+	AIMODULE_API void InitializeComposite(uint16 InLastExecutionIndex);
 
 	/** find next child branch to execute */
-	int32 FindChildToExecute(FBehaviorTreeSearchData& SearchData, EBTNodeResult::Type& LastResult) const;
+	AIMODULE_API int32 FindChildToExecute(FBehaviorTreeSearchData& SearchData, EBTNodeResult::Type& LastResult) const;
 
 	/** get index of child node (handle subtrees) */
-	int32 GetChildIndex(FBehaviorTreeSearchData& SearchData, const UBTNode& ChildNode) const;
+	AIMODULE_API int32 GetChildIndex(FBehaviorTreeSearchData& SearchData, const UBTNode& ChildNode) const;
 	/** get index of child node */
-	int32 GetChildIndex(const UBTNode& ChildNode) const;
+	AIMODULE_API int32 GetChildIndex(const UBTNode& ChildNode) const;
 
 	/** called before passing search to child node */
-	void OnChildActivation(FBehaviorTreeSearchData& SearchData, const UBTNode& ChildNode) const;
-	void OnChildActivation(FBehaviorTreeSearchData& SearchData, int32 ChildIndex) const;
+	AIMODULE_API void OnChildActivation(FBehaviorTreeSearchData& SearchData, const UBTNode& ChildNode) const;
+	AIMODULE_API void OnChildActivation(FBehaviorTreeSearchData& SearchData, int32 ChildIndex) const;
 
 	/**
 	 * Notification called after child has finished search
@@ -120,7 +120,7 @@ class AIMODULE_API UBTCompositeNode : public UBTNode
 	 * @param NodeResult the raison of the deactivation
 	 * @param bRequestedFromValidInstance the new requested search start is within the current active instances
 	 */
-	void OnChildDeactivation(FBehaviorTreeSearchData& SearchData, const UBTNode& ChildNode, EBTNodeResult::Type& NodeResult, const bool bRequestedFromValidInstance) const;
+	AIMODULE_API void OnChildDeactivation(FBehaviorTreeSearchData& SearchData, const UBTNode& ChildNode, EBTNodeResult::Type& NodeResult, const bool bRequestedFromValidInstance) const;
 	/**
 	 * Notification called after child has finished search
 	 * @param SearchData for any new addition or removal of extra aux nodes/ descriptor
@@ -128,60 +128,60 @@ class AIMODULE_API UBTCompositeNode : public UBTNode
 	 * @param NodeResult the raison of the deactivation
 	 * @param bRequestedFromValidInstance the new request search start is within the current active instances
 	 */
-	void OnChildDeactivation(FBehaviorTreeSearchData& SearchData, int32 ChildIndex, EBTNodeResult::Type& NodeResult, const bool bRequestedFromValidInstance) const;
+	AIMODULE_API void OnChildDeactivation(FBehaviorTreeSearchData& SearchData, int32 ChildIndex, EBTNodeResult::Type& NodeResult, const bool bRequestedFromValidInstance) const;
 
 	/** called when start enters this node */
-	void OnNodeActivation(FBehaviorTreeSearchData& SearchData) const;
+	AIMODULE_API void OnNodeActivation(FBehaviorTreeSearchData& SearchData) const;
 
 	/** called when search leaves this node */
-	void OnNodeDeactivation(FBehaviorTreeSearchData& SearchData, EBTNodeResult::Type& NodeResult) const;
+	AIMODULE_API void OnNodeDeactivation(FBehaviorTreeSearchData& SearchData, EBTNodeResult::Type& NodeResult) const;
 
 	/** called when search needs to reactivate this node */
-	void OnNodeRestart(FBehaviorTreeSearchData& SearchData) const;
+	AIMODULE_API void OnNodeRestart(FBehaviorTreeSearchData& SearchData) const;
 
 	/** notify about task execution start */
-	void ConditionalNotifyChildExecution(UBehaviorTreeComponent& OwnerComp, uint8* NodeMemory, const UBTNode& ChildNode, EBTNodeResult::Type& NodeResult) const;
+	AIMODULE_API void ConditionalNotifyChildExecution(UBehaviorTreeComponent& OwnerComp, uint8* NodeMemory, const UBTNode& ChildNode, EBTNodeResult::Type& NodeResult) const;
 
 	/** size of instance memory */
-	virtual uint16 GetInstanceMemorySize() const override;
+	AIMODULE_API virtual uint16 GetInstanceMemorySize() const override;
 
 	/** @return child node at given index */
-	UBTNode* GetChildNode(int32 Index) const;
+	AIMODULE_API UBTNode* GetChildNode(int32 Index) const;
 
 	/** @return children count */
-	int32 GetChildrenNum() const;
+	AIMODULE_API int32 GetChildrenNum() const;
 
 	/** @return execution index of child node */
-	uint16 GetChildExecutionIndex(int32 Index, EBTChildIndex ChildMode = EBTChildIndex::TaskNode) const;
+	AIMODULE_API uint16 GetChildExecutionIndex(int32 Index, EBTChildIndex ChildMode = EBTChildIndex::TaskNode) const;
 
 	/** @return execution index of last node in child branches */
-	uint16 GetLastExecutionIndex() const;
+	AIMODULE_API uint16 GetLastExecutionIndex() const;
 
 	/** set override for next child index */
-	virtual void SetChildOverride(FBehaviorTreeSearchData& SearchData, int8 Index) const;
+	AIMODULE_API virtual void SetChildOverride(FBehaviorTreeSearchData& SearchData, int8 Index) const;
 
 	/** gathers description of all runtime parameters */
-	virtual void DescribeRuntimeValues(const UBehaviorTreeComponent& OwnerComp, uint8* NodeMemory, EBTDescriptionVerbosity::Type Verbosity, TArray<FString>& Values) const override;
+	AIMODULE_API virtual void DescribeRuntimeValues(const UBehaviorTreeComponent& OwnerComp, uint8* NodeMemory, EBTDescriptionVerbosity::Type Verbosity, TArray<FString>& Values) const override;
 
 	/** check if child node can execute new subtree */
-	virtual bool CanPushSubtree(UBehaviorTreeComponent& OwnerComp, uint8* NodeMemory, int32 ChildIdx) const;
+	AIMODULE_API virtual bool CanPushSubtree(UBehaviorTreeComponent& OwnerComp, uint8* NodeMemory, int32 ChildIdx) const;
 
 #if WITH_EDITOR
 	/** @return allowed flow abort modes for decorators */
-	virtual bool CanAbortLowerPriority() const;
-	virtual bool CanAbortSelf() const;
+	AIMODULE_API virtual bool CanAbortLowerPriority() const;
+	AIMODULE_API virtual bool CanAbortSelf() const;
 #endif // WITH_EDITOR
 
 	/** find branch containing specified node index */
-	int32 GetMatchingChildIndex(int32 ActiveInstanceIdx, FBTNodeIndex& NodeIdx) const;
+	AIMODULE_API int32 GetMatchingChildIndex(int32 ActiveInstanceIdx, FBTNodeIndex& NodeIdx) const;
 
 	/** get first execution index of given branch */
-	uint16 GetBranchExecutionIndex(uint16 NodeInBranchIdx) const;
+	AIMODULE_API uint16 GetBranchExecutionIndex(uint16 NodeInBranchIdx) const;
 
 	/** is child execution allowed by decorators? */
-	bool DoDecoratorsAllowExecution(UBehaviorTreeComponent& OwnerComp, const int32 InstanceIdx, const int32 ChildIdx) const;
+	AIMODULE_API bool DoDecoratorsAllowExecution(UBehaviorTreeComponent& OwnerComp, const int32 InstanceIdx, const int32 ChildIdx) const;
 
-	bool IsApplyingDecoratorScope() const;
+	AIMODULE_API bool IsApplyingDecoratorScope() const;
 
 	// Deprecated methods
 	UE_DEPRECATED(5.0, "This function is deprecated. Please use RequestBranchDeactivation instead.")
@@ -219,47 +219,47 @@ protected:
 	/** called just after child execution, allows to modify result  
 	 * bUseChildExecutionNotify must be set to true for this function to be called 
 	 * Calling INIT_COMPOSITE_NODE_NOTIFY_FLAGS in the constructor of the node will set this flag automatically */
-	virtual void NotifyChildExecution(UBehaviorTreeComponent& OwnerComp, uint8* NodeMemory, int32 ChildIdx, EBTNodeResult::Type& NodeResult) const;
+	AIMODULE_API virtual void NotifyChildExecution(UBehaviorTreeComponent& OwnerComp, uint8* NodeMemory, int32 ChildIdx, EBTNodeResult::Type& NodeResult) const;
 
 	/** called when start enters this node  
 	 * bUseNodeActivationNotify must be set to true for this function to be called 
 	 * Calling INIT_COMPOSITE_NODE_NOTIFY_FLAGS in the constructor of the node will set this flag automatically */
-	virtual void NotifyNodeActivation(FBehaviorTreeSearchData& SearchData) const;
+	AIMODULE_API virtual void NotifyNodeActivation(FBehaviorTreeSearchData& SearchData) const;
 
 	/** called when start leaves this node 
 	 * bUseNodeDeactivationNotify must be set to true for this function to be called  
 	 * Calling INIT_COMPOSITE_NODE_NOTIFY_FLAGS in the constructor of the node will set this flag automatically */
-	virtual void NotifyNodeDeactivation(FBehaviorTreeSearchData& SearchData, EBTNodeResult::Type& NodeResult) const;
+	AIMODULE_API virtual void NotifyNodeDeactivation(FBehaviorTreeSearchData& SearchData, EBTNodeResult::Type& NodeResult) const;
 
 	/** check if NotifyDecoratorsOnActivation is allowed, requires bUseDecoratorsActivationCheck flag 
 	 * bUseDecoratorsActivationCheck must be set to true for this function to be called  
 	 * Calling INIT_COMPOSITE_NODE_NOTIFY_FLAGS in the constructor of the node will set this flag automatically */
-	virtual bool CanNotifyDecoratorsOnActivation(FBehaviorTreeSearchData& SearchData, int32 ChildIdx) const;
+	AIMODULE_API virtual bool CanNotifyDecoratorsOnActivation(FBehaviorTreeSearchData& SearchData, int32 ChildIdx) const;
 
 	/** check if NotifyDecoratorsOnDeactivation is allowed, requires bUseDecoratorsDeactivationCheck flag   
 	* bUseDecoratorsDeactivationCheck must be set to true for this function to be called 
 	 * Calling INIT_COMPOSITE_NODE_NOTIFY_FLAGS in the constructor of the node will set this flag automatically */
-	virtual bool CanNotifyDecoratorsOnDeactivation(FBehaviorTreeSearchData& SearchData, int32 ChildIdx, EBTNodeResult::Type& NodeResult) const;
+	AIMODULE_API virtual bool CanNotifyDecoratorsOnDeactivation(FBehaviorTreeSearchData& SearchData, int32 ChildIdx, EBTNodeResult::Type& NodeResult) const;
 
 	/** check if NotifyDecoratorsOnFailedActivation is allowed, requires bUseDecoratorsActivationCheck flag  
 	 * bUseDecoratorsFailedActivationCheck must be set to true for this function to be called 
 	 * Calling INIT_COMPOSITE_NODE_NOTIFY_FLAGS in the constructor of the node will set this flag automatically */
-	virtual bool CanNotifyDecoratorsOnFailedActivation(FBehaviorTreeSearchData& SearchData, int32 ChildIdx, EBTNodeResult::Type& NodeResult) const;
+	AIMODULE_API virtual bool CanNotifyDecoratorsOnFailedActivation(FBehaviorTreeSearchData& SearchData, int32 ChildIdx, EBTNodeResult::Type& NodeResult) const;
 
 	/** runs through decorators on given child node and notify them about activation */
-	void NotifyDecoratorsOnActivation(FBehaviorTreeSearchData& SearchData, int32 ChildIdx) const;
+	AIMODULE_API void NotifyDecoratorsOnActivation(FBehaviorTreeSearchData& SearchData, int32 ChildIdx) const;
 
 	/** runs through decorators on given child node and notify them about deactivation */
-	void NotifyDecoratorsOnDeactivation(FBehaviorTreeSearchData& SearchData, int32 ChildIdx, EBTNodeResult::Type& NodeResult, const bool bIsInSameActiveInstance) const;
+	AIMODULE_API void NotifyDecoratorsOnDeactivation(FBehaviorTreeSearchData& SearchData, int32 ChildIdx, EBTNodeResult::Type& NodeResult, const bool bIsInSameActiveInstance) const;
 
 	/** runs through decorators on given child node and notify them about failed activation */
-	void NotifyDecoratorsOnFailedActivation(FBehaviorTreeSearchData& SearchData, int32 ChildIdx, EBTNodeResult::Type& NodeResult) const;
+	AIMODULE_API void NotifyDecoratorsOnFailedActivation(FBehaviorTreeSearchData& SearchData, int32 ChildIdx, EBTNodeResult::Type& NodeResult) const;
 
 	/** get next child to process and store it in CurrentChild */
-	int32 GetNextChild(FBehaviorTreeSearchData& SearchData, int32 LastChildIdx, EBTNodeResult::Type LastResult) const;
+	AIMODULE_API int32 GetNextChild(FBehaviorTreeSearchData& SearchData, int32 LastChildIdx, EBTNodeResult::Type LastResult) const;
 
 	/** store delayed execution request */
-	void RequestDelayedExecution(UBehaviorTreeComponent& OwnerComp, EBTNodeResult::Type LastResult) const;
+	AIMODULE_API void RequestDelayedExecution(UBehaviorTreeComponent& OwnerComp, EBTNodeResult::Type LastResult) const;
 
 protected:
 	virtual int32 GetNextChildHandler(struct FBehaviorTreeSearchData& SearchData, int32 PrevChild, EBTNodeResult::Type LastResult) const { return BTSpecialChild::ReturnToParent; }
diff --git a/Engine/Source/Runtime/AIModule/Classes/BehaviorTree/BTDecorator.h b/Engine/Source/Runtime/AIModule/Classes/BehaviorTree/BTDecorator.h
--- a/Engine/Source/Runtime/AIModule/Classes/BehaviorTree/BTDecorator.h
+++ b/Engine/Source/Runtime/AIModule/Classes/BehaviorTree/BTDecorator.h
@@ -33,36 +33,36 @@ enum class EBTDecoratorAbortRequest : uint8
  *
  */
 
-UCLASS(Abstract)
-class AIMODULE_API UBTDecorator : public UBTAuxiliaryNode
+UCLASS(Abstract, MinimalAPI)
+class UBTDecorator : public UBTAuxiliaryNode
 {
 	GENERATED_UCLASS_BODY()
 
 	/** wrapper for node instancing: CalculateRawConditionValue */
-	bool WrappedCanExecute(UBehaviorTreeComponent& OwnerComp, uint8* NodeMemory) const;
+	AIMODULE_API bool WrappedCanExecute(UBehaviorTreeComponent& OwnerComp, uint8* NodeMemory) const;
 
 	/** wrapper for node instancing: OnNodeActivation  */
-	void WrappedOnNodeActivation(FBehaviorTreeSearchData& SearchData) const;
+	AIMODULE_API void WrappedOnNodeActivation(FBehaviorTreeSearchData& SearchData) const;
 	
 	/** wrapper for node instancing: OnNodeDeactivation */
-	void WrappedOnNodeDeactivation(FBehaviorTreeSearchData& SearchData, EBTNodeResult::Type NodeResult) const;
+	AIMODULE_API void WrappedOnNodeDeactivation(FBehaviorTreeSearchData& SearchData, EBTNodeResult::Type NodeResult) const;
 
 	/** wrapper for node instancing: OnNodeProcessed */
-	void WrappedOnNodeProcessed(FBehaviorTreeSearchData& SearchData, EBTNodeResult::Type& NodeResult) const;
+	AIMODULE_API void WrappedOnNodeProcessed(FBehaviorTreeSearchData& SearchData, EBTNodeResult::Type& NodeResult) const;
 
 	/** @return flow controller's abort mode */
-	EBTFlowAbortMode::Type GetFlowAbortMode() const;
+	AIMODULE_API EBTFlowAbortMode::Type GetFlowAbortMode() const;
 
 	/** @return true if condition should be inversed */
-	bool IsInversed() const;
+	AIMODULE_API bool IsInversed() const;
 
-	virtual FString GetStaticDescription() const override;
+	AIMODULE_API virtual FString GetStaticDescription() const override;
 
 	/** modify current flow abort mode, so it can be used with parent composite */
-	void UpdateFlowAbortMode();
+	AIMODULE_API void UpdateFlowAbortMode();
 
 	/** @return true if current abort mode can be used with parent composite */
-	bool IsFlowAbortModeValid() const;
+	AIMODULE_API bool IsFlowAbortModeValid() const;
 
 protected:
 
@@ -97,32 +97,32 @@ protected:
 	UPROPERTY(Category=FlowControl, EditAnywhere)
 	TEnumAsByte<EBTFlowAbortMode::Type> FlowAbortMode;
 
-	void SetIsInversed(bool bShouldBeInversed);
+	AIMODULE_API void SetIsInversed(bool bShouldBeInversed);
 
 	/** called when underlying node is activated
 	 * this function should be considered as const (don't modify state of object) if node is not instanced! 
 	 * bNotifyActivation must be set to true for this function to be called
 	 * Calling INIT_DECORATOR_NODE_NOTIFY_FLAGS in the constructor of the decorator will set this flag automatically */
-	virtual void OnNodeActivation(FBehaviorTreeSearchData& SearchData);
+	AIMODULE_API virtual void OnNodeActivation(FBehaviorTreeSearchData& SearchData);
 
 	/** called when underlying node has finished
 	 * this function should be considered as const (don't modify state of object) if node is not instanced! 
 	 * bNotifyDeactivation must be set to true for this function to be called
 	 * Calling INIT_DECORATOR_NODE_NOTIFY_FLAGS in the constructor of the decorator will set this flag automatically */
-	virtual void OnNodeDeactivation(FBehaviorTreeSearchData& SearchData, EBTNodeResult::Type NodeResult);
+	AIMODULE_API virtual void OnNodeDeactivation(FBehaviorTreeSearchData& SearchData, EBTNodeResult::Type NodeResult);
 
 	/** called when underlying node was processed (deactivated or failed to activate)
 	 * this function should be considered as const (don't modify state of object) if node is not instanced! 
 	 * bNotifyProcessed must be set to true for this function to be called 
 	 * Calling INIT_DECORATOR_NODE_NOTIFY_FLAGS in the constructor of the decorator will set this flag automatically */
-	virtual void OnNodeProcessed(FBehaviorTreeSearchData& SearchData, EBTNodeResult::Type& NodeResult);
+	AIMODULE_API virtual void OnNodeProcessed(FBehaviorTreeSearchData& SearchData, EBTNodeResult::Type& NodeResult);
 
 	/** calculates raw, core value of decorator's condition. Should not include calling IsInversed */
-	virtual bool CalculateRawConditionValue(UBehaviorTreeComponent& OwnerComp, uint8* NodeMemory) const;
+	AIMODULE_API virtual bool CalculateRawConditionValue(UBehaviorTreeComponent& OwnerComp, uint8* NodeMemory) const;
 
 	/** more "flow aware" version of calling RequestExecution(this) on owning behavior tree component
 	 *  should be used in external events that may change result of CalculateRawConditionValue */
-	void ConditionalFlowAbort(UBehaviorTreeComponent& OwnerComp, EBTDecoratorAbortRequest RequestMode) const;
+	AIMODULE_API void ConditionalFlowAbort(UBehaviorTreeComponent& OwnerComp, EBTDecoratorAbortRequest RequestMode) const;
 
 	friend FBehaviorDecoratorDetails;
 
diff --git a/Engine/Source/Runtime/AIModule/Classes/BehaviorTree/BTFunctionLibrary.h b/Engine/Source/Runtime/AIModule/Classes/BehaviorTree/BTFunctionLibrary.h
--- a/Engine/Source/Runtime/AIModule/Classes/BehaviorTree/BTFunctionLibrary.h
+++ b/Engine/Source/Runtime/AIModule/Classes/BehaviorTree/BTFunctionLibrary.h
@@ -25,94 +25,94 @@ namespace FBTNodeBPImplementationHelper
 	AIMODULE_API int32 CheckEventImplementationVersion(FName GenericEventName, FName AIEventName, const UObject& Object, const UClass& StopAtClass);
 }
 
-UCLASS(meta=(RestrictedToClasses="BTNode"))
-class AIMODULE_API UBTFunctionLibrary : public UBlueprintFunctionLibrary
+UCLASS(meta=(RestrictedToClasses="BTNode"), MinimalAPI)
+class UBTFunctionLibrary : public UBlueprintFunctionLibrary
 {
 	GENERATED_BODY()
 
 public:
 	UFUNCTION(BlueprintPure, Category="AI|BehaviorTree", Meta=(HidePin="NodeOwner", DefaultToSelf="NodeOwner"))
-	static UBlackboardComponent* GetOwnersBlackboard(UBTNode* NodeOwner);
+	static AIMODULE_API UBlackboardComponent* GetOwnersBlackboard(UBTNode* NodeOwner);
 
 	UFUNCTION(BlueprintPure, Category = "AI|BehaviorTree", Meta=(HidePin="NodeOwner", DefaultToSelf="NodeOwner"))
-	static UBehaviorTreeComponent* GetOwnerComponent(UBTNode* NodeOwner);
+	static AIMODULE_API UBehaviorTreeComponent* GetOwnerComponent(UBTNode* NodeOwner);
 
 	UFUNCTION(BlueprintPure, Category="AI|BehaviorTree", Meta=(HidePin="NodeOwner", DefaultToSelf="NodeOwner"))
-	static UObject* GetBlackboardValueAsObject(UBTNode* NodeOwner, const FBlackboardKeySelector& Key);
+	static AIMODULE_API UObject* GetBlackboardValueAsObject(UBTNode* NodeOwner, const FBlackboardKeySelector& Key);
 
 	UFUNCTION(BlueprintPure, Category="AI|BehaviorTree", Meta=(HidePin="NodeOwner", DefaultToSelf="NodeOwner"))
-	static AActor* GetBlackboardValueAsActor(UBTNode* NodeOwner, const FBlackboardKeySelector& Key);
+	static AIMODULE_API AActor* GetBlackboardValueAsActor(UBTNode* NodeOwner, const FBlackboardKeySelector& Key);
 
 	UFUNCTION(BlueprintPure, Category="AI|BehaviorTree", Meta=(HidePin="NodeOwner", DefaultToSelf="NodeOwner"))
-	static UClass* GetBlackboardValueAsClass(UBTNode* NodeOwner, const FBlackboardKeySelector& Key);
+	static AIMODULE_API UClass* GetBlackboardValueAsClass(UBTNode* NodeOwner, const FBlackboardKeySelector& Key);
 
 	UFUNCTION(BlueprintPure, Category="AI|BehaviorTree", Meta=(HidePin="NodeOwner", DefaultToSelf="NodeOwner"))
-	static uint8 GetBlackboardValueAsEnum(UBTNode* NodeOwner, const FBlackboardKeySelector& Key);
+	static AIMODULE_API uint8 GetBlackboardValueAsEnum(UBTNode* NodeOwner, const FBlackboardKeySelector& Key);
 
 	UFUNCTION(BlueprintPure, Category="AI|BehaviorTree", Meta=(HidePin="NodeOwner", DefaultToSelf="NodeOwner"))
-	static int32 GetBlackboardValueAsInt(UBTNode* NodeOwner, const FBlackboardKeySelector& Key);
+	static AIMODULE_API int32 GetBlackboardValueAsInt(UBTNode* NodeOwner, const FBlackboardKeySelector& Key);
 
 	UFUNCTION(BlueprintPure, Category="AI|BehaviorTree", Meta=(HidePin="NodeOwner", DefaultToSelf="NodeOwner"))
-	static float GetBlackboardValueAsFloat(UBTNode* NodeOwner, const FBlackboardKeySelector& Key);
+	static AIMODULE_API float GetBlackboardValueAsFloat(UBTNode* NodeOwner, const FBlackboardKeySelector& Key);
 
 	UFUNCTION(BlueprintPure, Category="AI|BehaviorTree", Meta=(HidePin="NodeOwner", DefaultToSelf="NodeOwner"))
-	static bool GetBlackboardValueAsBool(UBTNode* NodeOwner, const FBlackboardKeySelector& Key);
+	static AIMODULE_API bool GetBlackboardValueAsBool(UBTNode* NodeOwner, const FBlackboardKeySelector& Key);
 
 	UFUNCTION(BlueprintPure, Category="AI|BehaviorTree", Meta=(HidePin="NodeOwner", DefaultToSelf="NodeOwner"))
-	static FString GetBlackboardValueAsString(UBTNode* NodeOwner, const FBlackboardKeySelector& Key);
+	static AIMODULE_API FString GetBlackboardValueAsString(UBTNode* NodeOwner, const FBlackboardKeySelector& Key);
 
 	UFUNCTION(BlueprintPure, Category="AI|BehaviorTree", Meta=(HidePin="NodeOwner", DefaultToSelf="NodeOwner"))
-	static FName GetBlackboardValueAsName(UBTNode* NodeOwner, const FBlackboardKeySelector& Key);
+	static AIMODULE_API FName GetBlackboardValueAsName(UBTNode* NodeOwner, const FBlackboardKeySelector& Key);
 
 	UFUNCTION(BlueprintPure, Category="AI|BehaviorTree", Meta=(HidePin="NodeOwner", DefaultToSelf="NodeOwner"))
-	static FVector GetBlackboardValueAsVector(UBTNode* NodeOwner, const FBlackboardKeySelector& Key);
+	static AIMODULE_API FVector GetBlackboardValueAsVector(UBTNode* NodeOwner, const FBlackboardKeySelector& Key);
 
 	UFUNCTION(BlueprintPure, Category ="AI|BehaviorTree", Meta=(HidePin="NodeOwner", DefaultToSelf="NodeOwner"))
-	static FRotator GetBlackboardValueAsRotator(UBTNode* NodeOwner, const FBlackboardKeySelector& Key);
+	static AIMODULE_API FRotator GetBlackboardValueAsRotator(UBTNode* NodeOwner, const FBlackboardKeySelector& Key);
 
 	UFUNCTION(BlueprintCallable, Category="AI|BehaviorTree", Meta=(HidePin="NodeOwner", DefaultToSelf="NodeOwner"))
-	static void SetBlackboardValueAsObject(UBTNode* NodeOwner, const FBlackboardKeySelector& Key, UObject* Value);
+	static AIMODULE_API void SetBlackboardValueAsObject(UBTNode* NodeOwner, const FBlackboardKeySelector& Key, UObject* Value);
 
 	UFUNCTION(BlueprintCallable, Category="AI|BehaviorTree", Meta=(HidePin="NodeOwner", DefaultToSelf="NodeOwner"))
-	static void SetBlackboardValueAsClass(UBTNode* NodeOwner, const FBlackboardKeySelector& Key, UClass* Value);
+	static AIMODULE_API void SetBlackboardValueAsClass(UBTNode* NodeOwner, const FBlackboardKeySelector& Key, UClass* Value);
 
 	UFUNCTION(BlueprintCallable, Category="AI|BehaviorTree", Meta=(HidePin="NodeOwner", DefaultToSelf="NodeOwner"))
-	static void SetBlackboardValueAsEnum(UBTNode* NodeOwner, const FBlackboardKeySelector& Key, uint8 Value);
+	static AIMODULE_API void SetBlackboardValueAsEnum(UBTNode* NodeOwner, const FBlackboardKeySelector& Key, uint8 Value);
 
 	UFUNCTION(BlueprintCallable, Category="AI|BehaviorTree", Meta=(HidePin="NodeOwner", DefaultToSelf="NodeOwner"))
-	static void SetBlackboardValueAsInt(UBTNode* NodeOwner, const FBlackboardKeySelector& Key, int32 Value);
+	static AIMODULE_API void SetBlackboardValueAsInt(UBTNode* NodeOwner, const FBlackboardKeySelector& Key, int32 Value);
 
 	UFUNCTION(BlueprintCallable, Category="AI|BehaviorTree", Meta=(HidePin="NodeOwner", DefaultToSelf="NodeOwner"))
-	static void SetBlackboardValueAsFloat(UBTNode* NodeOwner, const FBlackboardKeySelector& Key, float Value);
+	static AIMODULE_API void SetBlackboardValueAsFloat(UBTNode* NodeOwner, const FBlackboardKeySelector& Key, float Value);
 
 	UFUNCTION(BlueprintCallable, Category="AI|BehaviorTree", Meta=(HidePin="NodeOwner", DefaultToSelf="NodeOwner"))
-	static void SetBlackboardValueAsBool(UBTNode* NodeOwner, const FBlackboardKeySelector& Key, bool Value);
+	static AIMODULE_API void SetBlackboardValueAsBool(UBTNode* NodeOwner, const FBlackboardKeySelector& Key, bool Value);
 
 	UFUNCTION(BlueprintCallable, Category="AI|BehaviorTree", Meta=(HidePin="NodeOwner", DefaultToSelf="NodeOwner"))
-	static void SetBlackboardValueAsString(UBTNode* NodeOwner, const FBlackboardKeySelector& Key, FString Value);
+	static AIMODULE_API void SetBlackboardValueAsString(UBTNode* NodeOwner, const FBlackboardKeySelector& Key, FString Value);
 
 	UFUNCTION(BlueprintCallable, Category="AI|BehaviorTree", Meta=(HidePin="NodeOwner", DefaultToSelf="NodeOwner"))
-	static void SetBlackboardValueAsName(UBTNode* NodeOwner, const FBlackboardKeySelector& Key, FName Value);
+	static AIMODULE_API void SetBlackboardValueAsName(UBTNode* NodeOwner, const FBlackboardKeySelector& Key, FName Value);
 
 	UFUNCTION(BlueprintCallable, Category="AI|BehaviorTree", Meta=(HidePin="NodeOwner", DefaultToSelf="NodeOwner"))
-	static void SetBlackboardValueAsVector(UBTNode* NodeOwner, const FBlackboardKeySelector& Key, FVector Value);
+	static AIMODULE_API void SetBlackboardValueAsVector(UBTNode* NodeOwner, const FBlackboardKeySelector& Key, FVector Value);
 
 	/** (DEPRECATED) Use ClearBlackboardValue instead */
 	UFUNCTION(BlueprintCallable, Category="AI|BehaviorTree", Meta=(HidePin="NodeOwner", DefaultToSelf="NodeOwner", DeprecatedFunction, DeprecationMessage="Use ClearBlackboardValue instead."))
-	static void ClearBlackboardValueAsVector(UBTNode* NodeOwner, const FBlackboardKeySelector& Key);
+	static AIMODULE_API void ClearBlackboardValueAsVector(UBTNode* NodeOwner, const FBlackboardKeySelector& Key);
 
 	UFUNCTION(BlueprintCallable, Category="AI|BehaviorTree", Meta = (HidePin = "NodeOwner", DefaultToSelf = "NodeOwner"))
-	static void SetBlackboardValueAsRotator(UBTNode* NodeOwner, const FBlackboardKeySelector& Key, FRotator Value);
+	static AIMODULE_API void SetBlackboardValueAsRotator(UBTNode* NodeOwner, const FBlackboardKeySelector& Key, FRotator Value);
 
 	/** Resets indicated value to "not set" value, based on values type */
 	UFUNCTION(BlueprintCallable, Category="AI|BehaviorTree", Meta = (HidePin = "NodeOwner", DefaultToSelf = "NodeOwner"))
-	static void ClearBlackboardValue(UBTNode* NodeOwner, const FBlackboardKeySelector& Key);
+	static AIMODULE_API void ClearBlackboardValue(UBTNode* NodeOwner, const FBlackboardKeySelector& Key);
 
 	/** Initialize variables marked as "instance memory" and set owning actor for blackboard operations */
 	UFUNCTION(BlueprintCallable, Category="AI|BehaviorTree", Meta=(HidePin="NodeOwner", DefaultToSelf="NodeOwner", DeprecatedFunction, DeprecationMessage="No longer needed"))
-	static void StartUsingExternalEvent(UBTNode* NodeOwner, AActor* OwningActor);
+	static AIMODULE_API void StartUsingExternalEvent(UBTNode* NodeOwner, AActor* OwningActor);
 
 	/** Save variables marked as "instance memory" and clear owning actor */
 	UFUNCTION(BlueprintCallable, Category="AI|BehaviorTree", Meta=(HidePin="NodeOwner", DefaultToSelf="NodeOwner", DeprecatedFunction, DeprecationMessage="No longer needed"))
-	static void StopUsingExternalEvent(UBTNode* NodeOwner);
+	static AIMODULE_API void StopUsingExternalEvent(UBTNode* NodeOwner);
 };
