[Engine]

diff --git a/Engine/Source/Runtime/AIModule/Classes/BehaviorTree/BTNode.h b/Engine/Source/Runtime/AIModule/Classes/BehaviorTree/BTNode.h
--- a/Engine/Source/Runtime/AIModule/Classes/BehaviorTree/BTNode.h
+++ b/Engine/Source/Runtime/AIModule/Classes/BehaviorTree/BTNode.h
@@ -34,52 +34,52 @@ struct FBTInstancedNodeMemory
 	int32 NodeIdx;
 };
 
-UCLASS(Abstract,config=Game)
-class AIMODULE_API UBTNode : public UObject, public IGameplayTaskOwnerInterface
+UCLASS(Abstract,config=Game, MinimalAPI)
+class UBTNode : public UObject, public IGameplayTaskOwnerInterface
 {
 	GENERATED_UCLASS_BODY()
 
-	virtual UWorld* GetWorld() const override;
+	AIMODULE_API virtual UWorld* GetWorld() const override;
 
 	/** fill in data about tree structure */
-	void InitializeNode(UBTCompositeNode* InParentNode, uint16 InExecutionIndex, uint16 InMemoryOffset, uint8 InTreeDepth);
+	AIMODULE_API void InitializeNode(UBTCompositeNode* InParentNode, uint16 InExecutionIndex, uint16 InMemoryOffset, uint8 InTreeDepth);
 
 	/** initialize any asset related data */
-	virtual void InitializeFromAsset(UBehaviorTree& Asset);
+	AIMODULE_API virtual void InitializeFromAsset(UBehaviorTree& Asset);
 	
 	/** initialize memory block */
-	virtual void InitializeMemory(UBehaviorTreeComponent& OwnerComp, uint8* NodeMemory, EBTMemoryInit::Type InitType) const;
+	AIMODULE_API virtual void InitializeMemory(UBehaviorTreeComponent& OwnerComp, uint8* NodeMemory, EBTMemoryInit::Type InitType) const;
 
 	/** cleanup memory block */
-	virtual void CleanupMemory(UBehaviorTreeComponent& OwnerComp, uint8* NodeMemory, EBTMemoryClear::Type CleanupType) const;
+	AIMODULE_API virtual void CleanupMemory(UBehaviorTreeComponent& OwnerComp, uint8* NodeMemory, EBTMemoryClear::Type CleanupType) const;
 
 	/** gathers description of all runtime parameters */
-	virtual void DescribeRuntimeValues(const UBehaviorTreeComponent& OwnerComp, uint8* NodeMemory, EBTDescriptionVerbosity::Type Verbosity, TArray<FString>& Values) const;
+	AIMODULE_API virtual void DescribeRuntimeValues(const UBehaviorTreeComponent& OwnerComp, uint8* NodeMemory, EBTDescriptionVerbosity::Type Verbosity, TArray<FString>& Values) const;
 
 	/** size of instance memory */
-	virtual uint16 GetInstanceMemorySize() const;
+	AIMODULE_API virtual uint16 GetInstanceMemorySize() const;
 
 	/** called when node instance is added to tree */
-	virtual void OnInstanceCreated(UBehaviorTreeComponent& OwnerComp);
+	AIMODULE_API virtual void OnInstanceCreated(UBehaviorTreeComponent& OwnerComp);
 
 	/** called when node instance is removed from tree */
-	virtual void OnInstanceDestroyed(UBehaviorTreeComponent& OwnerComp);
+	AIMODULE_API virtual void OnInstanceDestroyed(UBehaviorTreeComponent& OwnerComp);
 
 	/** called on creating subtree to set up memory and instancing */
-	void InitializeInSubtree(UBehaviorTreeComponent& OwnerComp, uint8* NodeMemory, int32& NextInstancedIndex, EBTMemoryInit::Type InitType) const;
+	AIMODULE_API void InitializeInSubtree(UBehaviorTreeComponent& OwnerComp, uint8* NodeMemory, int32& NextInstancedIndex, EBTMemoryInit::Type InitType) const;
 
 	/** called on removing subtree to cleanup memory */
-	void CleanupInSubtree(UBehaviorTreeComponent& OwnerComp, uint8* NodeMemory, EBTMemoryClear::Type CleanupType) const;
+	AIMODULE_API void CleanupInSubtree(UBehaviorTreeComponent& OwnerComp, uint8* NodeMemory, EBTMemoryClear::Type CleanupType) const;
 
 	/** size of special, hidden memory block for internal mechanics */
-	virtual uint16 GetSpecialMemorySize() const;
+	AIMODULE_API virtual uint16 GetSpecialMemorySize() const;
 
 #if USE_BEHAVIORTREE_DEBUGGER
 	/** fill in data about execution order */
-	void InitializeExecutionOrder(UBTNode* NextNode);
+	AIMODULE_API void InitializeExecutionOrder(UBTNode* NextNode);
 
 	/** @return next node in execution order */
-	UBTNode* GetNextNode() const;
+	AIMODULE_API UBTNode* GetNextNode() const;
 #endif
 
 	template<typename T>
@@ -102,57 +102,57 @@ class AIMODULE_API UBTNode : public UObject, public IGameplayTaskOwnerInterface
 	T* GetSpecialNodeMemory(uint8* NodeMemory) const;
 
 	/** @return parent node */
-	UBTCompositeNode* GetParentNode() const;
+	AIMODULE_API UBTCompositeNode* GetParentNode() const;
 
 	/** @return name of node */
-	FString GetNodeName() const;
+	AIMODULE_API FString GetNodeName() const;
 
 	/** @return execution index */
-	uint16 GetExecutionIndex() const;
+	AIMODULE_API uint16 GetExecutionIndex() const;
 
 	/** @return memory offset */
-	uint16 GetMemoryOffset() const;
+	AIMODULE_API uint16 GetMemoryOffset() const;
 
 	/** @return depth in tree */
-	uint8 GetTreeDepth() const;
+	AIMODULE_API uint8 GetTreeDepth() const;
 
 	/** sets bIsInjected flag, do NOT call this function unless you really know what you are doing! */
-	void MarkInjectedNode();
+	AIMODULE_API void MarkInjectedNode();
 
 	/** @return true if node was injected by subtree */
-	bool IsInjected() const;
+	AIMODULE_API bool IsInjected() const;
 
 	/** sets bCreateNodeInstance flag, do NOT call this function on already pushed tree instance! */
-	void ForceInstancing(bool bEnable);
+	AIMODULE_API void ForceInstancing(bool bEnable);
 
 	/** @return true if node wants to be instanced */
-	bool HasInstance() const;
+	AIMODULE_API bool HasInstance() const;
 
 	/** @return true if this object is instanced node */
-	bool IsInstanced() const;
+	AIMODULE_API bool IsInstanced() const;
 
 	/** @return tree asset */
-	UBehaviorTree* GetTreeAsset() const;
+	AIMODULE_API UBehaviorTree* GetTreeAsset() const;
 
 	/** @return blackboard asset */
-	UBlackboardData* GetBlackboardAsset() const;
+	AIMODULE_API UBlackboardData* GetBlackboardAsset() const;
 
 	/** @return node instance if bCreateNodeInstance was set */
-	UBTNode* GetNodeInstance(const UBehaviorTreeComponent& OwnerComp, uint8* NodeMemory) const;
-	UBTNode* GetNodeInstance(FBehaviorTreeSearchData& SearchData) const;
+	AIMODULE_API UBTNode* GetNodeInstance(const UBehaviorTreeComponent& OwnerComp, uint8* NodeMemory) const;
+	AIMODULE_API UBTNode* GetNodeInstance(FBehaviorTreeSearchData& SearchData) const;
 
 	/** @return string containing description of this node instance with all relevant runtime values */
-	FString GetRuntimeDescription(const UBehaviorTreeComponent& OwnerComp, uint8* NodeMemory, EBTDescriptionVerbosity::Type Verbosity) const;
+	AIMODULE_API FString GetRuntimeDescription(const UBehaviorTreeComponent& OwnerComp, uint8* NodeMemory, EBTDescriptionVerbosity::Type Verbosity) const;
 
 	/** @return string containing description of this node with all setup values */
-	virtual FString GetStaticDescription() const;
+	AIMODULE_API virtual FString GetStaticDescription() const;
 
 #if WITH_EDITOR
 	/** Get the name of the icon used to display this node in the editor */
-	virtual FName GetNodeIconName() const;
+	AIMODULE_API virtual FName GetNodeIconName() const;
 
 	/** Get whether this node is using a blueprint for its logic */
-	virtual bool UsesBlueprint() const;
+	AIMODULE_API virtual bool UsesBlueprint() const;
 
 	/** Called after creating new node in behavior tree editor, use for versioning */
 	virtual void OnNodeCreated() {}
@@ -162,14 +162,14 @@ class AIMODULE_API UBTNode : public UObject, public IGameplayTaskOwnerInterface
 	virtual void SetOwner(AActor* ActorOwner) {}
 
 	// BEGIN IGameplayTaskOwnerInterface
-	virtual UGameplayTasksComponent* GetGameplayTasksComponent(const UGameplayTask& Task) const override;
-	virtual AActor* GetGameplayTaskOwner(const UGameplayTask* Task) const override;
-	virtual AActor* GetGameplayTaskAvatar(const UGameplayTask* Task) const override;
-	virtual uint8 GetGameplayTaskDefaultPriority() const override;
-	virtual void OnGameplayTaskInitialized(UGameplayTask& Task) override;
+	AIMODULE_API virtual UGameplayTasksComponent* GetGameplayTasksComponent(const UGameplayTask& Task) const override;
+	AIMODULE_API virtual AActor* GetGameplayTaskOwner(const UGameplayTask* Task) const override;
+	AIMODULE_API virtual AActor* GetGameplayTaskAvatar(const UGameplayTask* Task) const override;
+	AIMODULE_API virtual uint8 GetGameplayTaskDefaultPriority() const override;
+	AIMODULE_API virtual void OnGameplayTaskInitialized(UGameplayTask& Task) override;
 	// END IGameplayTaskOwnerInterface
 
-	UBehaviorTreeComponent* GetBTComponentForTask(UGameplayTask& Task) const;
+	AIMODULE_API UBehaviorTreeComponent* GetBTComponentForTask(UGameplayTask& Task) const;
 	
 	template <class T>
 	T* NewBTAITask(UBehaviorTreeComponent& BTComponent)
diff --git a/Engine/Source/Runtime/AIModule/Classes/BehaviorTree/BTService.h b/Engine/Source/Runtime/AIModule/Classes/BehaviorTree/BTService.h
--- a/Engine/Source/Runtime/AIModule/Classes/BehaviorTree/BTService.h
+++ b/Engine/Source/Runtime/AIModule/Classes/BehaviorTree/BTService.h
@@ -30,22 +30,22 @@
  *  their runtime properties in provided NodeMemory block (allocation size determined by GetInstanceMemorySize() )
  */
  
-UCLASS(Abstract)
-class AIMODULE_API UBTService : public UBTAuxiliaryNode
+UCLASS(Abstract, MinimalAPI)
+class UBTService : public UBTAuxiliaryNode
 {
 	GENERATED_UCLASS_BODY()
 
-	virtual FString GetStaticDescription() const override;
+	AIMODULE_API virtual FString GetStaticDescription() const override;
 
-	void NotifyParentActivation(FBehaviorTreeSearchData& SearchData);
+	AIMODULE_API void NotifyParentActivation(FBehaviorTreeSearchData& SearchData);
 
 protected:
 
 	// Gets the description of our tick interval
-	FString GetStaticTickIntervalDescription() const;
+	AIMODULE_API FString GetStaticTickIntervalDescription() const;
 
 	// Gets the description for our service
-	virtual FString GetStaticServiceDescription() const;
+	AIMODULE_API virtual FString GetStaticServiceDescription() const;
 
 	/** defines time span between subsequent ticks of the service */
 	UPROPERTY(Category=Service, EditAnywhere, meta=(ClampMin="0.001"))
@@ -70,20 +70,20 @@ protected:
 	 * this function should be considered as const (don't modify state of object) if node is not instanced!
 	 * bNotifyTick must be set to true for this function to be called 
 	 * Calling INIT_SERVICE_NODE_NOTIFY_FLAGS in the constructor of the service will set this flag automatically */
-	virtual void TickNode(UBehaviorTreeComponent& OwnerComp, uint8* NodeMemory, float DeltaSeconds) override;
+	AIMODULE_API virtual void TickNode(UBehaviorTreeComponent& OwnerComp, uint8* NodeMemory, float DeltaSeconds) override;
 
 	/** called when search enters underlying branch
 	 * this function should be considered as const (don't modify state of object) if node is not instanced! 
 	 * bNotifyOnSearch must be set to true for this function to be called  
 	 * Calling INIT_SERVICE_NODE_NOTIFY_FLAGS in the constructor of the service will set this flag automatically */
-	virtual void OnSearchStart(FBehaviorTreeSearchData& SearchData);
+	AIMODULE_API virtual void OnSearchStart(FBehaviorTreeSearchData& SearchData);
 
 #if WITH_EDITOR
-	virtual FName GetNodeIconName() const override;
+	AIMODULE_API virtual FName GetNodeIconName() const override;
 #endif // WITH_EDITOR
 
 	/** set next tick time */
-	virtual void ScheduleNextTick(UBehaviorTreeComponent& OwnerComp, uint8* NodeMemory);
+	AIMODULE_API virtual void ScheduleNextTick(UBehaviorTreeComponent& OwnerComp, uint8* NodeMemory);
 
 	template<typename TickNode, typename OnBecomeRelevant, typename OnCeaseRelevant, typename OnSearchStart>
 	void InitNotifyFlags(TickNode, OnBecomeRelevant, OnCeaseRelevant, OnSearchStart)
@@ -102,4 +102,4 @@ protected:
 					&NodeType::OnBecomeRelevant, \
 					&NodeType::OnCeaseRelevant, \
 					&NodeType::OnSearchStart); \
-	} while (false)
\ No newline at end of file
+	} while (false)
diff --git a/Engine/Source/Runtime/AIModule/Classes/BehaviorTree/BTTaskNode.h b/Engine/Source/Runtime/AIModule/Classes/BehaviorTree/BTTaskNode.h
--- a/Engine/Source/Runtime/AIModule/Classes/BehaviorTree/BTTaskNode.h
+++ b/Engine/Source/Runtime/AIModule/Classes/BehaviorTree/BTTaskNode.h
@@ -29,41 +29,41 @@ struct FBTTaskMemory : public FBTInstancedNodeMemory
  *
  */
 
-UCLASS(Abstract)
-class AIMODULE_API UBTTaskNode : public UBTNode
+UCLASS(Abstract, MinimalAPI)
+class UBTTaskNode : public UBTNode
 {
 	GENERATED_UCLASS_BODY()
 
 	/** starts this task, should return Succeeded, Failed or InProgress
 	 *  (use FinishLatentTask() when returning InProgress)
 	 * this function should be considered as const (don't modify state of object) if node is not instanced! */
-	virtual EBTNodeResult::Type ExecuteTask(UBehaviorTreeComponent& OwnerComp, uint8* NodeMemory);
+	AIMODULE_API virtual EBTNodeResult::Type ExecuteTask(UBehaviorTreeComponent& OwnerComp, uint8* NodeMemory);
 
-	virtual uint16 GetSpecialMemorySize() const override;
+	AIMODULE_API virtual uint16 GetSpecialMemorySize() const override;
 
 protected:
 	/** aborts this task, should return Aborted or InProgress
 	 *  (use FinishLatentAbort() when returning InProgress)
 	 * this function should be considered as const (don't modify state of object) if node is not instanced! */
-	virtual EBTNodeResult::Type AbortTask(UBehaviorTreeComponent& OwnerComp, uint8* NodeMemory);
+	AIMODULE_API virtual EBTNodeResult::Type AbortTask(UBehaviorTreeComponent& OwnerComp, uint8* NodeMemory);
 
 	/** sets next tick time */
-	void SetNextTickTime(uint8* NodeMemory, float RemainingTime) const;
+	AIMODULE_API void SetNextTickTime(uint8* NodeMemory, float RemainingTime) const;
 
 public:
 #if WITH_EDITOR
-	virtual FName GetNodeIconName() const override;
+	AIMODULE_API virtual FName GetNodeIconName() const override;
 #endif // WITH_EDITOR
-	virtual void OnGameplayTaskDeactivated(UGameplayTask& Task) override;
+	AIMODULE_API virtual void OnGameplayTaskDeactivated(UGameplayTask& Task) override;
 
 	/** message observer's hook */
-	void ReceivedMessage(UBrainComponent* BrainComp, const FAIMessage& Message);
+	AIMODULE_API void ReceivedMessage(UBrainComponent* BrainComp, const FAIMessage& Message);
 
 	/** wrapper for node instancing: ExecuteTask */
-	EBTNodeResult::Type WrappedExecuteTask(UBehaviorTreeComponent& OwnerComp, uint8* NodeMemory) const;
+	AIMODULE_API EBTNodeResult::Type WrappedExecuteTask(UBehaviorTreeComponent& OwnerComp, uint8* NodeMemory) const;
 
 	/** wrapper for node instancing: AbortTask */
-	EBTNodeResult::Type WrappedAbortTask(UBehaviorTreeComponent& OwnerComp, uint8* NodeMemory) const;
+	AIMODULE_API EBTNodeResult::Type WrappedAbortTask(UBehaviorTreeComponent& OwnerComp, uint8* NodeMemory) const;
 
 	/** wrapper for node instancing: TickTask
 	  * @param OwnerComp	The behavior tree owner of this node
@@ -71,19 +71,19 @@ public:
 	  * @param DeltaSeconds		DeltaTime since last call
 	  * @param NextNeededDeltaTime		In out parameter, if this node needs a smaller DeltaTime it is the node's responsibility to change it
 	  * @returns	True if it actually done some processing or false if it was skipped because of not ticking or in between time interval */
-	bool WrappedTickTask(UBehaviorTreeComponent& OwnerComp, uint8* NodeMemory, float DeltaSeconds, float& NextNeededDeltaTime) const;
+	AIMODULE_API bool WrappedTickTask(UBehaviorTreeComponent& OwnerComp, uint8* NodeMemory, float DeltaSeconds, float& NextNeededDeltaTime) const;
 
 	/** wrapper for node instancing: OnTaskFinished */
-	void WrappedOnTaskFinished(UBehaviorTreeComponent& OwnerComp, uint8* NodeMemory, EBTNodeResult::Type TaskResult) const;
+	AIMODULE_API void WrappedOnTaskFinished(UBehaviorTreeComponent& OwnerComp, uint8* NodeMemory, EBTNodeResult::Type TaskResult) const;
 
 	/** helper function: finish latent executing */
-	void FinishLatentTask(UBehaviorTreeComponent& OwnerComp, EBTNodeResult::Type TaskResult) const;
+	AIMODULE_API void FinishLatentTask(UBehaviorTreeComponent& OwnerComp, EBTNodeResult::Type TaskResult) const;
 
 	/** helper function: finishes latent aborting */
-	void FinishLatentAbort(UBehaviorTreeComponent& OwnerComp) const;
+	AIMODULE_API void FinishLatentAbort(UBehaviorTreeComponent& OwnerComp) const;
 
 	/** @return true if task search should be discarded when this task is selected to execute but is already running */
-	bool ShouldIgnoreRestartSelf() const;
+	AIMODULE_API bool ShouldIgnoreRestartSelf() const;
 
 	/** service nodes */
 	UPROPERTY()
@@ -105,24 +105,24 @@ protected:
 	 * this function should be considered as const (don't modify state of object) if node is not instanced! 
 	 * bNotifyTick must be set to true for this function to be called
 	 * Calling INIT_TASK_NODE_NOTIFY_FLAGS in the constructor of the task will set this flag automatically */
-	virtual void TickTask(UBehaviorTreeComponent& OwnerComp, uint8* NodeMemory, float DeltaSeconds);
+	AIMODULE_API virtual void TickTask(UBehaviorTreeComponent& OwnerComp, uint8* NodeMemory, float DeltaSeconds);
 
 	/** message handler, default implementation will finish latent execution/abortion
 	 * this function should be considered as const (don't modify state of object) if node is not instanced! */
-	virtual void OnMessage(UBehaviorTreeComponent& OwnerComp, uint8* NodeMemory, FName Message, int32 RequestID, bool bSuccess);
+	AIMODULE_API virtual void OnMessage(UBehaviorTreeComponent& OwnerComp, uint8* NodeMemory, FName Message, int32 RequestID, bool bSuccess);
 
 	/** called when task execution is finished
 	 * this function should be considered as const (don't modify state of object) if node is not instanced! 
 	 * bNotifyTaskFinished must be set to true for this function to be called 
 	 * Calling INIT_TASK_NODE_NOTIFY_FLAGS in the constructor of the task will set this flag automatically */
-	virtual void OnTaskFinished(UBehaviorTreeComponent& OwnerComp, uint8* NodeMemory, EBTNodeResult::Type TaskResult);
+	AIMODULE_API virtual void OnTaskFinished(UBehaviorTreeComponent& OwnerComp, uint8* NodeMemory, EBTNodeResult::Type TaskResult);
 
 	/** register message observer */
-	void WaitForMessage(UBehaviorTreeComponent& OwnerComp, FName MessageType) const;
-	void WaitForMessage(UBehaviorTreeComponent& OwnerComp, FName MessageType, int32 RequestID) const;
+	AIMODULE_API void WaitForMessage(UBehaviorTreeComponent& OwnerComp, FName MessageType) const;
+	AIMODULE_API void WaitForMessage(UBehaviorTreeComponent& OwnerComp, FName MessageType, int32 RequestID) const;
 	
 	/** unregister message observers */
-	void StopWaitingForMessages(UBehaviorTreeComponent& OwnerComp) const;
+	AIMODULE_API void StopWaitingForMessages(UBehaviorTreeComponent& OwnerComp) const;
 	
 	template<typename TickTask,	typename OnTaskFinished>
 	void InitNotifyFlags(TickTask, OnTaskFinished)
diff --git a/Engine/Source/Runtime/AIModule/Classes/BehaviorTree/BehaviorTree.h b/Engine/Source/Runtime/AIModule/Classes/BehaviorTree/BehaviorTree.h
--- a/Engine/Source/Runtime/AIModule/Classes/BehaviorTree/BehaviorTree.h
+++ b/Engine/Source/Runtime/AIModule/Classes/BehaviorTree/BehaviorTree.h
@@ -11,8 +11,8 @@
 class UBlackboardData;
 class UBTDecorator;
 
-UCLASS(BlueprintType)
-class AIMODULE_API UBehaviorTree : public UObject, public IBlackboardAssetProvider
+UCLASS(BlueprintType, MinimalAPI)
+class UBehaviorTree : public UObject, public IBlackboardAssetProvider
 {
 	GENERATED_UCLASS_BODY()
 
@@ -34,7 +34,7 @@ class AIMODULE_API UBehaviorTree : public UObject, public IBlackboardAssetProvid
 
 	// BEGIN IBlackboardAssetProvider
 	/** @return blackboard asset */
-	virtual UBlackboardData* GetBlackboardAsset() const override;
+	AIMODULE_API virtual UBlackboardData* GetBlackboardAsset() const override;
 	// END IBlackboardAssetProvider
 
 	/** blackboard asset for this tree */
diff --git a/Engine/Source/Runtime/AIModule/Classes/BehaviorTree/BehaviorTreeComponent.h b/Engine/Source/Runtime/AIModule/Classes/BehaviorTree/BehaviorTreeComponent.h
--- a/Engine/Source/Runtime/AIModule/Classes/BehaviorTree/BehaviorTreeComponent.h
+++ b/Engine/Source/Runtime/AIModule/Classes/BehaviorTree/BehaviorTreeComponent.h
@@ -100,53 +100,53 @@ enum class EBTBranchAction : uint16
 };
 ENUM_CLASS_FLAGS(EBTBranchAction);
 
-UCLASS(ClassGroup = AI, meta = (BlueprintSpawnableComponent))
-class AIMODULE_API UBehaviorTreeComponent : public UBrainComponent
+UCLASS(ClassGroup = AI, meta = (BlueprintSpawnableComponent), MinimalAPI)
+class UBehaviorTreeComponent : public UBrainComponent
 {
 	GENERATED_UCLASS_BODY()
 
 	// UActorComponent overrides
-	virtual void RegisterComponentTickFunctions(bool bRegister) override;
-	virtual void SetComponentTickEnabled(bool bEnabled) override;
+	AIMODULE_API virtual void RegisterComponentTickFunctions(bool bRegister) override;
+	AIMODULE_API virtual void SetComponentTickEnabled(bool bEnabled) override;
 
 	// Begin UBrainComponent overrides
-	virtual void StartLogic() override;
-	virtual void RestartLogic() override;
-	virtual void StopLogic(const FString& Reason) override;
-	virtual void PauseLogic(const FString& Reason) override;
-	virtual EAILogicResuming::Type ResumeLogic(const FString& Reason) override;
+	AIMODULE_API virtual void StartLogic() override;
+	AIMODULE_API virtual void RestartLogic() override;
+	AIMODULE_API virtual void StopLogic(const FString& Reason) override;
+	AIMODULE_API virtual void PauseLogic(const FString& Reason) override;
+	AIMODULE_API virtual EAILogicResuming::Type ResumeLogic(const FString& Reason) override;
 
 	/** indicates instance has been initialized to work with specific BT asset */
-	bool TreeHasBeenStarted() const;
+	AIMODULE_API bool TreeHasBeenStarted() const;
 
 public:
 	/** DO NOT USE. This constructor is for internal usage only for hot-reload purposes. */
-	UBehaviorTreeComponent(FVTableHelper& Helper);
+	AIMODULE_API UBehaviorTreeComponent(FVTableHelper& Helper);
 
-	virtual bool IsRunning() const override;
-	virtual bool IsPaused() const override;
-	virtual void Cleanup() override;
-	virtual void HandleMessage(const FAIMessage& Message) override;
+	AIMODULE_API virtual bool IsRunning() const override;
+	AIMODULE_API virtual bool IsPaused() const override;
+	AIMODULE_API virtual void Cleanup() override;
+	AIMODULE_API virtual void HandleMessage(const FAIMessage& Message) override;
 	// End UBrainComponent overrides
 
 	// Begin UActorComponent overrides
-	virtual void UninitializeComponent() override;
+	AIMODULE_API virtual void UninitializeComponent() override;
 	// End UActorComponent overrides
 
 	/** starts execution from root */
-	void StartTree(UBehaviorTree& Asset, EBTExecutionMode::Type ExecuteMode = EBTExecutionMode::Looped);
+	AIMODULE_API void StartTree(UBehaviorTree& Asset, EBTExecutionMode::Type ExecuteMode = EBTExecutionMode::Looped);
 
 	/** stops execution */
-	void StopTree(EBTStopMode::Type StopMode = EBTStopMode::Safe);
+	AIMODULE_API void StopTree(EBTStopMode::Type StopMode = EBTStopMode::Safe);
 
 	/** restarts execution from root
 	 * @param RestartMode to force the reevaluation of the root node, which could skip active nodes that are removed and then readded (default)
 	 *        or a complete restart from scratch of the tree (equivalent of StopTree then StartTree)
 	 */
-	void RestartTree(EBTRestartMode RestartMode = EBTRestartMode::ForceReevaluateRootNode);
+	AIMODULE_API void RestartTree(EBTRestartMode RestartMode = EBTRestartMode::ForceReevaluateRootNode);
 
 	/** request execution change */
-	void RequestExecution(const UBTCompositeNode* RequestedOn, const int32 InstanceIdx, 
+	AIMODULE_API void RequestExecution(const UBTCompositeNode* RequestedOn, const int32 InstanceIdx, 
 		const UBTNode* RequestedBy, const int32 RequestedByChildIndex,
 		const EBTNodeResult::Type ContinueWithResult, bool bStoreForDebugger = true);
 
@@ -158,122 +158,122 @@ public:
 
 	/** request unregistration of aux nodes in the specified branch */
 	UE_DEPRECATED(5.0, "This function is deprecated. Please use RequestBranchDeactivation instead.")
-	void RequestUnregisterAuxNodesInBranch(const UBTCompositeNode* Node);
+	AIMODULE_API void RequestUnregisterAuxNodesInBranch(const UBTCompositeNode* Node);
 
 	/** request branch evaluation: helper for active node (ex: tasks) */
-	void RequestBranchEvaluation(EBTNodeResult::Type ContinueWithResult);
+	AIMODULE_API void RequestBranchEvaluation(EBTNodeResult::Type ContinueWithResult);
 
 	/** request branch evaluation: helper for decorator */
-	void RequestBranchEvaluation(const UBTDecorator& RequestedBy);
+	AIMODULE_API void RequestBranchEvaluation(const UBTDecorator& RequestedBy);
 
 	/** request branch activation: helper for decorator */
-	void RequestBranchActivation(const UBTDecorator& RequestedBy, const bool bRequestEvenIfExecuting);
+	AIMODULE_API void RequestBranchActivation(const UBTDecorator& RequestedBy, const bool bRequestEvenIfExecuting);
 
 	/** request branch deactivation: helper for decorator */
-	void RequestBranchDeactivation(const UBTDecorator& RequestedBy);
+	AIMODULE_API void RequestBranchDeactivation(const UBTDecorator& RequestedBy);
 
 	/** finish latent execution or abort */
-	void OnTaskFinished(const UBTTaskNode* TaskNode, EBTNodeResult::Type TaskResult);
+	AIMODULE_API void OnTaskFinished(const UBTTaskNode* TaskNode, EBTNodeResult::Type TaskResult);
 
 	/** setup message observer for given task */
-	void RegisterMessageObserver(const UBTTaskNode* TaskNode, FName MessageType);
-	void RegisterMessageObserver(const UBTTaskNode* TaskNode, FName MessageType, FAIRequestID MessageID);
+	AIMODULE_API void RegisterMessageObserver(const UBTTaskNode* TaskNode, FName MessageType);
+	AIMODULE_API void RegisterMessageObserver(const UBTTaskNode* TaskNode, FName MessageType, FAIRequestID MessageID);
 	
 	/** remove message observers registered with task */
-	void UnregisterMessageObserversFrom(const UBTTaskNode* TaskNode);
-	void UnregisterMessageObserversFrom(const FBTNodeIndex& TaskIdx);
+	AIMODULE_API void UnregisterMessageObserversFrom(const UBTTaskNode* TaskNode);
+	AIMODULE_API void UnregisterMessageObserversFrom(const FBTNodeIndex& TaskIdx);
 
 	/** add active parallel task */
-	void RegisterParallelTask(const UBTTaskNode* TaskNode);
+	AIMODULE_API void RegisterParallelTask(const UBTTaskNode* TaskNode);
 
 	/** remove parallel task */
-	void UnregisterParallelTask(const UBTTaskNode* TaskNode, uint16 InstanceIdx);
+	AIMODULE_API void UnregisterParallelTask(const UBTTaskNode* TaskNode, uint16 InstanceIdx);
 
 	/** unregister all aux nodes less important than given index */
-	void UnregisterAuxNodesUpTo(const FBTNodeIndex& Index);
+	AIMODULE_API void UnregisterAuxNodesUpTo(const FBTNodeIndex& Index);
 
 	/** unregister all aux nodes between given execution index range: FromIndex < AuxIndex < ToIndex */
-	void UnregisterAuxNodesInRange(const FBTNodeIndex& FromIndex, const FBTNodeIndex& ToIndex);
+	AIMODULE_API void UnregisterAuxNodesInRange(const FBTNodeIndex& FromIndex, const FBTNodeIndex& ToIndex);
 
 	/** unregister all aux nodes in branch of tree */
-	void UnregisterAuxNodesInBranch(const UBTCompositeNode* Node, bool bApplyImmediately = true);
+	AIMODULE_API void UnregisterAuxNodesInBranch(const UBTCompositeNode* Node, bool bApplyImmediately = true);
 
 	/** BEGIN UActorComponent overrides */
-	virtual void TickComponent(float DeltaTime, enum ELevelTick TickType, FActorComponentTickFunction *ThisTickFunction) override;
+	AIMODULE_API virtual void TickComponent(float DeltaTime, enum ELevelTick TickType, FActorComponentTickFunction *ThisTickFunction) override;
 	/** END UActorComponent overrides */
 
 	/** Schedule when will be the next tick, 0.0f means next frame, FLT_MAX means never */
-	void ScheduleNextTick(float NextDeltaTime);
+	AIMODULE_API void ScheduleNextTick(float NextDeltaTime);
 
 	/** process execution flow */
-	void ProcessExecutionRequest();
+	AIMODULE_API void ProcessExecutionRequest();
 
 	/** schedule execution flow update in next tick */
-	void ScheduleExecutionUpdate();
+	AIMODULE_API void ScheduleExecutionUpdate();
 
 	/** tries to find behavior tree instance in context */
-	int32 FindInstanceContainingNode(const UBTNode* Node) const;
+	AIMODULE_API int32 FindInstanceContainingNode(const UBTNode* Node) const;
 
 	/** tries to find template node for given instanced node */
-	UBTNode* FindTemplateNode(const UBTNode* Node) const;
+	AIMODULE_API UBTNode* FindTemplateNode(const UBTNode* Node) const;
 
 	/** @return current tree */
-	UBehaviorTree* GetCurrentTree() const;
+	AIMODULE_API UBehaviorTree* GetCurrentTree() const;
 
 	/** @return tree from top of instance stack */
-	UBehaviorTree* GetRootTree() const;
+	AIMODULE_API UBehaviorTree* GetRootTree() const;
 
 	/** @return active node */
-	const UBTNode* GetActiveNode() const;
+	AIMODULE_API const UBTNode* GetActiveNode() const;
 	
 	/** get index of active instance on stack */
-	uint16 GetActiveInstanceIdx() const;
+	AIMODULE_API uint16 GetActiveInstanceIdx() const;
 
 	/** @return node memory */
-	uint8* GetNodeMemory(UBTNode* Node, int32 InstanceIdx) const;
+	AIMODULE_API uint8* GetNodeMemory(UBTNode* Node, int32 InstanceIdx) const;
 
 	/** @return true if ExecutionRequest is switching to higher priority node */
-	bool IsRestartPending() const;
+	AIMODULE_API bool IsRestartPending() const;
 
 	/** @return true if waiting for abort to finish */
-	bool IsAbortPending() const;
+	AIMODULE_API bool IsAbortPending() const;
 
 	/** @return true if active node is one of child nodes of given one */
-	bool IsExecutingBranch(const UBTNode* Node, int32 ChildIndex = -1) const;
+	AIMODULE_API bool IsExecutingBranch(const UBTNode* Node, int32 ChildIndex = -1) const;
 
 	/** @return true if aux node is currently active */
-	bool IsAuxNodeActive(const UBTAuxiliaryNode* AuxNode) const;
-	bool IsAuxNodeActive(const UBTAuxiliaryNode* AuxNodeTemplate, int32 InstanceIdx) const;
+	AIMODULE_API bool IsAuxNodeActive(const UBTAuxiliaryNode* AuxNode) const;
+	AIMODULE_API bool IsAuxNodeActive(const UBTAuxiliaryNode* AuxNodeTemplate, int32 InstanceIdx) const;
 
 	/** Returns true if InstanceStack contains any BT runtime instances */
 	bool IsInstanceStackEmpty() const { return (InstanceStack.Num() == 0); }
 
 	/** @return status of speficied task */
-	EBTTaskStatus::Type GetTaskStatus(const UBTTaskNode* TaskNode) const;
+	AIMODULE_API EBTTaskStatus::Type GetTaskStatus(const UBTTaskNode* TaskNode) const;
 
-	virtual FString GetDebugInfoString() const override;
-	virtual FString DescribeActiveTasks() const;
-	virtual FString DescribeActiveTrees() const;
+	AIMODULE_API virtual FString GetDebugInfoString() const override;
+	AIMODULE_API virtual FString DescribeActiveTasks() const;
+	AIMODULE_API virtual FString DescribeActiveTrees() const;
 
 	/** @return the cooldown tag end time, 0.0f if CooldownTag is not found */
 	UFUNCTION(BlueprintCallable, Category = "AI|Logic")
-	double GetTagCooldownEndTime(FGameplayTag CooldownTag) const;
+	AIMODULE_API double GetTagCooldownEndTime(FGameplayTag CooldownTag) const;
 
 	/** add to the cooldown tag's duration */
 	UFUNCTION(BlueprintCallable, Category = "AI|Logic")
-	void AddCooldownTagDuration(FGameplayTag CooldownTag, float CooldownDuration, bool bAddToExistingDuration);
+	AIMODULE_API void AddCooldownTagDuration(FGameplayTag CooldownTag, float CooldownDuration, bool bAddToExistingDuration);
 
 	/** assign subtree to RunBehaviorDynamic task specified by tag */
 	UFUNCTION(BlueprintCallable, Category="AI|Logic")
-	virtual void SetDynamicSubtree(FGameplayTag InjectTag, UBehaviorTree* BehaviorAsset);
+	AIMODULE_API virtual void SetDynamicSubtree(FGameplayTag InjectTag, UBehaviorTree* BehaviorAsset);
 
 // Code for timing BT Search for FramePro
 #if !UE_BUILD_SHIPPING
-	static void EndFrame();
+	static AIMODULE_API void EndFrame();
 #endif
 
 #if ENABLE_VISUAL_LOG
-	virtual void DescribeSelfToVisLog(struct FVisualLogEntry* Snapshot) const override;
+	AIMODULE_API virtual void DescribeSelfToVisLog(struct FVisualLogEntry* Snapshot) const override;
 #endif
 
 #if CSV_PROFILER
@@ -346,14 +346,14 @@ protected:
 	mutable TArray<FBehaviorTreeExecutionStep> DebuggerSteps;
 
 	/** set when at least one debugger window is opened */
-	static int32 ActiveDebuggerCounter;
+	static AIMODULE_API int32 ActiveDebuggerCounter;
 #endif
 
 	// Code for timing BT Search for FramePro
 #if !UE_BUILD_SHIPPING
-	static bool bAddedEndFrameCallback;
-	static double FrameSearchTime;
-	static int32 NumSearchTimeCalls;
+	static AIMODULE_API bool bAddedEndFrameCallback;
+	static AIMODULE_API double FrameSearchTime;
+	static AIMODULE_API int32 NumSearchTimeCalls;
 #endif
 
 	/** index of last active instance on stack */
@@ -379,105 +379,105 @@ protected:
 
 	/** push behavior tree instance on execution stack
 	 *	@NOTE: should never be called out-side of BT execution, meaning only BT tasks can push another BT instance! */
-	bool PushInstance(UBehaviorTree& TreeAsset);
+	AIMODULE_API bool PushInstance(UBehaviorTree& TreeAsset);
 
 	/** add unique Id of newly created subtree to KnownInstances list and return its index */
-	uint8 UpdateInstanceId(UBehaviorTree* TreeAsset, const UBTNode* OriginNode, int32 OriginInstanceIdx);
+	AIMODULE_API uint8 UpdateInstanceId(UBehaviorTree* TreeAsset, const UBTNode* OriginNode, int32 OriginInstanceIdx);
 
 	/** remove instanced nodes, known subtree instances and safely clears their persistent memory */
-	void RemoveAllInstances();
+	AIMODULE_API void RemoveAllInstances();
 
 	/** copy memory block from running instances to persistent memory */
-	void CopyInstanceMemoryToPersistent();
+	AIMODULE_API void CopyInstanceMemoryToPersistent();
 
 	/** copy memory block from persistent memory to running instances (rollback) */
-	void CopyInstanceMemoryFromPersistent();
+	AIMODULE_API void CopyInstanceMemoryFromPersistent();
 
 	/** called when tree runs out of nodes to execute */
-	virtual void OnTreeFinished();
+	AIMODULE_API virtual void OnTreeFinished();
 
 	/** apply pending node updates from SearchData */
-	void ApplySearchData(UBTNode* NewActiveNode);
+	AIMODULE_API void ApplySearchData(UBTNode* NewActiveNode);
 
 	/** apply pending node updates required for discarded search */
-	void ApplyDiscardedSearch();
+	AIMODULE_API void ApplyDiscardedSearch();
 
 	/** apply updates from specific list */
-	void ApplySearchUpdates(const TArray<FBehaviorTreeSearchUpdate>& UpdateList, int32 NewNodeExecutionIndex, bool bPostUpdate = false);
+	AIMODULE_API void ApplySearchUpdates(const TArray<FBehaviorTreeSearchUpdate>& UpdateList, int32 NewNodeExecutionIndex, bool bPostUpdate = false);
 
 	/** abort currently executed task */
-	void AbortCurrentTask();
+	AIMODULE_API void AbortCurrentTask();
 
 	/** execute new task */
-	void ExecuteTask(UBTTaskNode* TaskNode);
+	AIMODULE_API void ExecuteTask(UBTTaskNode* TaskNode);
 
 	/** deactivate all nodes up to requested one */
-	bool DeactivateUpTo(const UBTCompositeNode* Node, uint16 NodeInstanceIdx, EBTNodeResult::Type& NodeResult, int32& OutLastDeactivatedChildIndex);
+	AIMODULE_API bool DeactivateUpTo(const UBTCompositeNode* Node, uint16 NodeInstanceIdx, EBTNodeResult::Type& NodeResult, int32& OutLastDeactivatedChildIndex);
 
 	/** returns true if execution was waiting on latent aborts and they are all finished;  */
-	bool TrackPendingLatentAborts();
+	AIMODULE_API bool TrackPendingLatentAborts();
 
 	/** tracks if there are new tasks using latent abort in progress */
-	void TrackNewLatentAborts();
+	AIMODULE_API void TrackNewLatentAborts();
 
 	/** return true if the current or any parallel task has a latent abort in progress */
-	bool HasActiveLatentAborts() const;
+	AIMODULE_API bool HasActiveLatentAborts() const;
 
 	/** apply pending execution from last task search */
-	void ProcessPendingExecution();
+	AIMODULE_API void ProcessPendingExecution();
 
 	/** apply pending tree initialization */
-	void ProcessPendingInitialize();
+	AIMODULE_API void ProcessPendingInitialize();
 
 	/** restore state of tree to state before search */
-	void RollbackSearchChanges();
+	AIMODULE_API void RollbackSearchChanges();
 
 	/** make a snapshot for debugger */
-	void StoreDebuggerExecutionStep(EBTExecutionSnap::Type SnapType);
+	AIMODULE_API void StoreDebuggerExecutionStep(EBTExecutionSnap::Type SnapType);
 
 	/** make a snapshot for debugger from given subtree instance */
-	void StoreDebuggerInstance(FBehaviorTreeDebuggerInstance& InstanceInfo, uint16 InstanceIdx, EBTExecutionSnap::Type SnapType) const;
-	void StoreDebuggerRemovedInstance(uint16 InstanceIdx) const;
+	AIMODULE_API void StoreDebuggerInstance(FBehaviorTreeDebuggerInstance& InstanceInfo, uint16 InstanceIdx, EBTExecutionSnap::Type SnapType) const;
+	AIMODULE_API void StoreDebuggerRemovedInstance(uint16 InstanceIdx) const;
 
 	/** store search step for debugger */
-	void StoreDebuggerSearchStep(const UBTNode* Node, uint16 InstanceIdx, EBTNodeResult::Type NodeResult) const;
-	void StoreDebuggerSearchStep(const UBTNode* Node, uint16 InstanceIdx, bool bPassed) const;
+	AIMODULE_API void StoreDebuggerSearchStep(const UBTNode* Node, uint16 InstanceIdx, EBTNodeResult::Type NodeResult) const;
+	AIMODULE_API void StoreDebuggerSearchStep(const UBTNode* Node, uint16 InstanceIdx, bool bPassed) const;
 
 	/** store restarting node for debugger */
-	void StoreDebuggerRestart(const UBTNode* Node, uint16 InstanceIdx, bool bAllowed);
+	AIMODULE_API void StoreDebuggerRestart(const UBTNode* Node, uint16 InstanceIdx, bool bAllowed);
 
 	/** describe blackboard's key values */
-	void StoreDebuggerBlackboard(TMap<FName, FString>& BlackboardValueDesc) const;
+	AIMODULE_API void StoreDebuggerBlackboard(TMap<FName, FString>& BlackboardValueDesc) const;
 
 	/** gather nodes runtime descriptions */
-	void StoreDebuggerRuntimeValues(TArray<FString>& RuntimeDescriptions, UBTNode* RootNode, uint16 InstanceIdx) const;
+	AIMODULE_API void StoreDebuggerRuntimeValues(TArray<FString>& RuntimeDescriptions, UBTNode* RootNode, uint16 InstanceIdx) const;
 
 	/** update runtime description of given task node in latest debugger's snapshot */
-	void UpdateDebuggerAfterExecution(const UBTTaskNode* TaskNode, uint16 InstanceIdx) const;
+	AIMODULE_API void UpdateDebuggerAfterExecution(const UBTTaskNode* TaskNode, uint16 InstanceIdx) const;
 
 	/** check if debugger is currently running and can gather data */
-	static bool IsDebuggerActive();
+	static AIMODULE_API bool IsDebuggerActive();
 
 	/** Return NodeA's relative priority in regards to NodeB */
-	EBTNodeRelativePriority CalculateRelativePriority(const UBTNode* NodeA, const UBTNode* NodeB) const;
+	AIMODULE_API EBTNodeRelativePriority CalculateRelativePriority(const UBTNode* NodeA, const UBTNode* NodeB) const;
 
 	/** Evaluate a branch as current active node is finished */
-	void EvaluateBranch(EBTNodeResult::Type LastResult);
+	AIMODULE_API void EvaluateBranch(EBTNodeResult::Type LastResult);
 
 	/** Evaluate a branch as the decorator conditions have changed */
-	void EvaluateBranch(const UBTDecorator& RequestedBy);
+	AIMODULE_API void EvaluateBranch(const UBTDecorator& RequestedBy);
 
 	/** Activate a branch as the decorator conditions are now passing */
-	void ActivateBranch(const UBTDecorator& RequestedBy, bool bRequestEvenIfNotExecuting);
+	AIMODULE_API void ActivateBranch(const UBTDecorator& RequestedBy, bool bRequestEvenIfNotExecuting);
 
 	/** Deactivate a branch as the decorator conditions are not passing anymore */
-	void DeactivateBranch(const UBTDecorator& RequestedBy);
+	AIMODULE_API void DeactivateBranch(const UBTDecorator& RequestedBy);
 
 	/** Suspend any branch actions and queue them to be processed later by ResumeBranchActions() */
-	void SuspendBranchActions(EBTBranchAction BranchActions = EBTBranchAction::All);
+	AIMODULE_API void SuspendBranchActions(EBTBranchAction BranchActions = EBTBranchAction::All);
 
 	/** Resume branch actions and execute all the queued up ones */
-	void ResumeBranchActions();
+	AIMODULE_API void ResumeBranchActions();
 
 	UE_DEPRECATED(5.1, "This function is deprecated. Please use SuspendBranchActions instead.")
 	void SuspendBranchDeactivation() { SuspendBranchActions(); }
diff --git a/Engine/Source/Runtime/AIModule/Classes/BehaviorTree/BehaviorTreeManager.h b/Engine/Source/Runtime/AIModule/Classes/BehaviorTree/BehaviorTreeManager.h
--- a/Engine/Source/Runtime/AIModule/Classes/BehaviorTree/BehaviorTreeManager.h
+++ b/Engine/Source/Runtime/AIModule/Classes/BehaviorTree/BehaviorTreeManager.h
@@ -29,8 +29,8 @@ struct FBehaviorTreeTemplateInfo
 	uint16 InstanceMemorySize;
 };
 
-UCLASS(config=Engine, Transient)
-class AIMODULE_API UBehaviorTreeManager : public UObject
+UCLASS(config=Engine, Transient, MinimalAPI)
+class UBehaviorTreeManager : public UObject
 {
 	GENERATED_UCLASS_BODY()
 
@@ -39,27 +39,27 @@ class AIMODULE_API UBehaviorTreeManager : public UObject
 	int32 MaxDebuggerSteps;
 
 	/** get behavior tree template for given blueprint */
-	bool LoadTree(UBehaviorTree& Asset, UBTCompositeNode*& Root, uint16& InstanceMemorySize);
+	AIMODULE_API bool LoadTree(UBehaviorTree& Asset, UBTCompositeNode*& Root, uint16& InstanceMemorySize);
 
 	/** get aligned memory size */
-	static uint16 GetAlignedDataSize(uint16 Size);
+	static AIMODULE_API uint16 GetAlignedDataSize(uint16 Size);
 
 	/** helper function for sorting and aligning node memory */
-	static void InitializeMemoryHelper(const TArray<UBTDecorator*>& Nodes, TArray<uint16>& MemoryOffsets, int32& MemorySize, bool bForceInstancing = false);
+	static AIMODULE_API void InitializeMemoryHelper(const TArray<UBTDecorator*>& Nodes, TArray<uint16>& MemoryOffsets, int32& MemorySize, bool bForceInstancing = false);
 
 	/** cleanup hooks for map loading */
-	virtual void FinishDestroy() override;
+	AIMODULE_API virtual void FinishDestroy() override;
 
-	void DumpUsageStats() const;
+	AIMODULE_API void DumpUsageStats() const;
 
 	/** register new behavior tree component for tracking */
-	void AddActiveComponent(UBehaviorTreeComponent& Component);
+	AIMODULE_API void AddActiveComponent(UBehaviorTreeComponent& Component);
 
 	/** unregister behavior tree component from tracking */
-	void RemoveActiveComponent(UBehaviorTreeComponent& Component);
+	AIMODULE_API void RemoveActiveComponent(UBehaviorTreeComponent& Component);
 
-	static UBehaviorTreeManager* GetCurrent(UWorld* World);
-	static UBehaviorTreeManager* GetCurrent(UObject* WorldContextObject);
+	static AIMODULE_API UBehaviorTreeManager* GetCurrent(UWorld* World);
+	static AIMODULE_API UBehaviorTreeManager* GetCurrent(UObject* WorldContextObject);
 
 protected:
 
diff --git a/Engine/Source/Runtime/AIModule/Classes/BehaviorTree/BehaviorTreeTypes.h b/Engine/Source/Runtime/AIModule/Classes/BehaviorTree/BehaviorTreeTypes.h
--- a/Engine/Source/Runtime/AIModule/Classes/BehaviorTree/BehaviorTreeTypes.h
+++ b/Engine/Source/Runtime/AIModule/Classes/BehaviorTree/BehaviorTreeTypes.h
@@ -338,7 +338,7 @@ struct FBehaviorTreeInstance
 	AIMODULE_API FBehaviorTreeInstance(const FBehaviorTreeInstance& Other);
 
 	UE_DEPRECATED(5.2, "Copying FBehaviorTreeInstance assignement operator has been deprecated in favor of move assignement operator")
-	AIMODULE_API FBehaviorTreeInstance& operator=(const FBehaviorTreeInstance& Other) = default;
+	FBehaviorTreeInstance& operator=(const FBehaviorTreeInstance& Other) = default;
 
 	AIMODULE_API FBehaviorTreeInstance(FBehaviorTreeInstance&& Other);
 	AIMODULE_API FBehaviorTreeInstance& operator=(FBehaviorTreeInstance&& Other);
@@ -600,7 +600,7 @@ struct FBehaviorTreePropertyMemory
 /** helper struct for defining types of allowed blackboard entries
  *  (e.g. only entries holding points and objects derived form actor class) */
 USTRUCT(BlueprintType)
-struct AIMODULE_API FBlackboardKeySelector
+struct FBlackboardKeySelector
 {
 	GENERATED_USTRUCT_BODY()
 
@@ -630,28 +630,28 @@ protected:
 	uint32 bNoneIsAllowedValue:1;
 
 	/** find initial selection. Called when None is not a valid option for this key selector */
-	void InitSelection(const UBlackboardData& BlackboardAsset);
+	AIMODULE_API void InitSelection(const UBlackboardData& BlackboardAsset);
 
 public:
 	/** find ID and class of selected key */
-	void ResolveSelectedKey(const UBlackboardData& BlackboardAsset);
+	AIMODULE_API void ResolveSelectedKey(const UBlackboardData& BlackboardAsset);
 		
 	void AllowNoneAsValue(bool bAllow) { bNoneIsAllowedValue = bAllow; }
 
 	FORCEINLINE FBlackboard::FKey GetSelectedKeyID() const { return FBlackboard::FKey(IntCastChecked<uint16>(SelectedKeyID)); }
 
 	/** helper functions for setting basic filters */
-	void AddObjectFilter(UObject* Owner, FName PropertyName, TSubclassOf<UObject> AllowedClass);
-	void AddClassFilter(UObject* Owner, FName PropertyName, TSubclassOf<UObject> AllowedClass);
-	void AddEnumFilter(UObject* Owner, FName PropertyName, UEnum* AllowedEnum);
-	void AddNativeEnumFilter(UObject* Owner, FName PropertyName, const FString& AllowedEnumName);
-	void AddIntFilter(UObject* Owner, FName PropertyName);
-	void AddFloatFilter(UObject* Owner, FName PropertyName);
-	void AddBoolFilter(UObject* Owner, FName PropertyName);
-	void AddVectorFilter(UObject* Owner, FName PropertyName);
-	void AddRotatorFilter(UObject* Owner, FName PropertyName);
-	void AddStringFilter(UObject* Owner, FName PropertyName);
-	void AddNameFilter(UObject* Owner, FName PropertyName);
+	AIMODULE_API void AddObjectFilter(UObject* Owner, FName PropertyName, TSubclassOf<UObject> AllowedClass);
+	AIMODULE_API void AddClassFilter(UObject* Owner, FName PropertyName, TSubclassOf<UObject> AllowedClass);
+	AIMODULE_API void AddEnumFilter(UObject* Owner, FName PropertyName, UEnum* AllowedEnum);
+	AIMODULE_API void AddNativeEnumFilter(UObject* Owner, FName PropertyName, const FString& AllowedEnumName);
+	AIMODULE_API void AddIntFilter(UObject* Owner, FName PropertyName);
+	AIMODULE_API void AddFloatFilter(UObject* Owner, FName PropertyName);
+	AIMODULE_API void AddBoolFilter(UObject* Owner, FName PropertyName);
+	AIMODULE_API void AddVectorFilter(UObject* Owner, FName PropertyName);
+	AIMODULE_API void AddRotatorFilter(UObject* Owner, FName PropertyName);
+	AIMODULE_API void AddStringFilter(UObject* Owner, FName PropertyName);
+	AIMODULE_API void AddNameFilter(UObject* Owner, FName PropertyName);
 
 	FORCEINLINE bool IsNone() const { return bNoneIsAllowedValue && GetSelectedKeyID() == FBlackboard::InvalidKey; }
 	FORCEINLINE bool IsSet() const { return GetSelectedKeyID() != FBlackboard::InvalidKey; }
@@ -661,31 +661,31 @@ public:
 	friend FBlackboardDecoratorDetails;
 };
 
-UCLASS(Abstract)
-class AIMODULE_API UBehaviorTreeTypes : public UObject
+UCLASS(Abstract, MinimalAPI)
+class UBehaviorTreeTypes : public UObject
 {
 	GENERATED_BODY()
 
-	static FString BTLoggingContext;
+	static AIMODULE_API FString BTLoggingContext;
 
 public:
 
-	static FString DescribeNodeHelper(const UBTNode* Node);
+	static AIMODULE_API FString DescribeNodeHelper(const UBTNode* Node);
 
-	static FString DescribeNodeResult(EBTNodeResult::Type NodeResult);
-	static FString DescribeFlowAbortMode(EBTFlowAbortMode::Type FlowAbortMode);
-	static FString DescribeActiveNode(EBTActiveNode::Type ActiveNodeType);
-	static FString DescribeTaskStatus(EBTTaskStatus::Type TaskStatus);
-	static FString DescribeNodeUpdateMode(EBTNodeUpdateMode::Type UpdateMode);
+	static AIMODULE_API FString DescribeNodeResult(EBTNodeResult::Type NodeResult);
+	static AIMODULE_API FString DescribeFlowAbortMode(EBTFlowAbortMode::Type FlowAbortMode);
+	static AIMODULE_API FString DescribeActiveNode(EBTActiveNode::Type ActiveNodeType);
+	static AIMODULE_API FString DescribeTaskStatus(EBTTaskStatus::Type TaskStatus);
+	static AIMODULE_API FString DescribeNodeUpdateMode(EBTNodeUpdateMode::Type UpdateMode);
 
 	/** returns short name of object's class (BTTaskNode_Wait -> Wait) */
-	static FString GetShortTypeName(const UObject* Ob);
+	static AIMODULE_API FString GetShortTypeName(const UObject* Ob);
 	
 	static FString GetBTLoggingContext() { return BTLoggingContext; }
 	
 	// @param NewBTLoggingContext the object which name's will be added to some of the BT logging
 	// 	pass nullptr to clear
-	static void SetBTLoggingContext(const UBTNode* NewBTLoggingContext);
+	static AIMODULE_API void SetBTLoggingContext(const UBTNode* NewBTLoggingContext);
 };
 
 /** Helper struct to push a node as the new logging context and automatically reset the context on destruction. */
@@ -701,4 +701,4 @@ struct FScopedBTLoggingContext
 	{
 		UBehaviorTreeTypes::SetBTLoggingContext(nullptr);
 	}
-};
\ No newline at end of file
+};
diff --git a/Engine/Source/Runtime/AIModule/Classes/BehaviorTree/Blackboard/BlackboardKeyType.h b/Engine/Source/Runtime/AIModule/Classes/BehaviorTree/Blackboard/BlackboardKeyType.h
--- a/Engine/Source/Runtime/AIModule/Classes/BehaviorTree/Blackboard/BlackboardKeyType.h
+++ b/Engine/Source/Runtime/AIModule/Classes/BehaviorTree/Blackboard/BlackboardKeyType.h
@@ -19,71 +19,71 @@ struct FBlackboardInstancedKeyMemory
 	int32 KeyIdx;
 };
 
-UCLASS(EditInlineNew, Abstract, CollapseCategories, AutoExpandCategories=(Blackboard))
-class AIMODULE_API UBlackboardKeyType : public UObject
+UCLASS(EditInlineNew, Abstract, CollapseCategories, AutoExpandCategories=(Blackboard), MinimalAPI)
+class UBlackboardKeyType : public UObject
 {
 	GENERATED_UCLASS_BODY()
 
 	/** handle dynamic data size */
-	virtual void PreInitialize(UBlackboardComponent& OwnerComp);
+	AIMODULE_API virtual void PreInitialize(UBlackboardComponent& OwnerComp);
 
 	/** handle instancing if needed */
-	void InitializeKey(UBlackboardComponent& OwnerComp, FBlackboard::FKey KeyID);
+	AIMODULE_API void InitializeKey(UBlackboardComponent& OwnerComp, FBlackboard::FKey KeyID);
 
 	/** does it match settings in filter? */
-	virtual bool IsAllowedByFilter(UBlackboardKeyType* FilterOb) const;
+	AIMODULE_API virtual bool IsAllowedByFilter(UBlackboardKeyType* FilterOb) const;
 
 	/** extract location from entry, supports instanced keys */
-	bool WrappedGetLocation(const UBlackboardComponent& OwnerComp, const uint8* MemoryBlock, FVector& Location) const;
+	AIMODULE_API bool WrappedGetLocation(const UBlackboardComponent& OwnerComp, const uint8* MemoryBlock, FVector& Location) const;
 
 	/** extract rotation from entry, supports instanced keys */
-	bool WrappedGetRotation(const UBlackboardComponent& OwnerComp, const uint8* MemoryBlock, FRotator& Rotation) const;
+	AIMODULE_API bool WrappedGetRotation(const UBlackboardComponent& OwnerComp, const uint8* MemoryBlock, FRotator& Rotation) const;
 
 	/** free value before removing from blackboard, supports instanced keys */
-	void WrappedFree(UBlackboardComponent& OwnerComp, uint8* MemoryBlock);
+	AIMODULE_API void WrappedFree(UBlackboardComponent& OwnerComp, uint8* MemoryBlock);
 
 	/** sets value to the default, supports instanced keys */
-	void WrappedClear(const UBlackboardComponent& OwnerComp, uint8* MemoryBlock) const;
+	AIMODULE_API void WrappedClear(const UBlackboardComponent& OwnerComp, uint8* MemoryBlock) const;
 
 	/** check if key has stored value, supports instanced keys */
-	bool WrappedIsEmpty(const UBlackboardComponent& OwnerComp, const uint8* MemoryBlock) const;
+	AIMODULE_API bool WrappedIsEmpty(const UBlackboardComponent& OwnerComp, const uint8* MemoryBlock) const;
 
 	/** various value testing, used by decorators, supports instanced keys */
-	bool WrappedTestBasicOperation(const UBlackboardComponent& OwnerComp, const uint8* MemoryBlock, EBasicKeyOperation::Type Op) const;
-	bool WrappedTestArithmeticOperation(const UBlackboardComponent& OwnerComp, const uint8* MemoryBlock, EArithmeticKeyOperation::Type Op, int32 OtherIntValue, float OtherFloatValue) const;
-	bool WrappedTestTextOperation(const UBlackboardComponent& OwnerComp, const uint8* MemoryBlock, ETextKeyOperation::Type Op, const FString& OtherString) const;
+	AIMODULE_API bool WrappedTestBasicOperation(const UBlackboardComponent& OwnerComp, const uint8* MemoryBlock, EBasicKeyOperation::Type Op) const;
+	AIMODULE_API bool WrappedTestArithmeticOperation(const UBlackboardComponent& OwnerComp, const uint8* MemoryBlock, EArithmeticKeyOperation::Type Op, int32 OtherIntValue, float OtherFloatValue) const;
+	AIMODULE_API bool WrappedTestTextOperation(const UBlackboardComponent& OwnerComp, const uint8* MemoryBlock, ETextKeyOperation::Type Op, const FString& OtherString) const;
 
 	/** describe params of arithmetic test */
-	virtual FString DescribeArithmeticParam(int32 IntValue, float FloatValue) const;
+	AIMODULE_API virtual FString DescribeArithmeticParam(int32 IntValue, float FloatValue) const;
 
 	/** convert value to text, supports instanced keys */
-	FString WrappedDescribeValue(const UBlackboardComponent& OwnerComp, const uint8* MemoryBlock) const;
+	AIMODULE_API FString WrappedDescribeValue(const UBlackboardComponent& OwnerComp, const uint8* MemoryBlock) const;
 
 	/** description of params for property view */
-	virtual FString DescribeSelf() const;
+	AIMODULE_API virtual FString DescribeSelf() const;
 
 	/** create replacement key for deprecated data */
-	virtual UBlackboardKeyType* UpdateDeprecatedKey();
+	AIMODULE_API virtual UBlackboardKeyType* UpdateDeprecatedKey();
 
 	/** @return key instance if bCreateKeyInstance was set */
-	const UBlackboardKeyType* GetKeyInstance(const UBlackboardComponent& OwnerComp, const uint8* MemoryBlock) const;
-	UBlackboardKeyType* GetKeyInstance(UBlackboardComponent& OwnerComp, const uint8* MemoryBlock) const;
+	AIMODULE_API const UBlackboardKeyType* GetKeyInstance(const UBlackboardComponent& OwnerComp, const uint8* MemoryBlock) const;
+	AIMODULE_API UBlackboardKeyType* GetKeyInstance(UBlackboardComponent& OwnerComp, const uint8* MemoryBlock) const;
 
 	/** compares two values */
-	virtual EBlackboardCompare::Type CompareValues(const UBlackboardComponent& OwnerComp, const uint8* MemoryBlock,
+	AIMODULE_API virtual EBlackboardCompare::Type CompareValues(const UBlackboardComponent& OwnerComp, const uint8* MemoryBlock,
 		const UBlackboardKeyType* OtherKeyOb, const uint8* OtherMemoryBlock) const;
 
 	/** @return true if key wants to be instanced */
-	bool HasInstance() const;
+	AIMODULE_API bool HasInstance() const;
 
 	/** @return true if this object is instanced key */
-	bool IsInstanced() const;
+	AIMODULE_API bool IsInstanced() const;
 
 	/** get ValueSize */
-	uint16 GetValueSize() const;
+	AIMODULE_API uint16 GetValueSize() const;
 
 	/** get test supported by this type */
-	EBlackboardKeyOperation::Type GetTestOperation() const;
+	AIMODULE_API EBlackboardKeyOperation::Type GetTestOperation() const;
 
 protected:
 
@@ -134,33 +134,33 @@ protected:
 	friend UBlackboardComponent;
 	
 	/** copy value from other key, works directly on provided memory/properties */
-	virtual void CopyValues(UBlackboardComponent& OwnerComp, uint8* MemoryBlock, const UBlackboardKeyType* SourceKeyOb, const uint8* SourceBlock);
+	AIMODULE_API virtual void CopyValues(UBlackboardComponent& OwnerComp, uint8* MemoryBlock, const UBlackboardKeyType* SourceKeyOb, const uint8* SourceBlock);
 
 	/** initialize memory, works directly on provided memory/properties */
-	virtual void InitializeMemory(UBlackboardComponent& OwnerComp, uint8* MemoryBlock);
+	AIMODULE_API virtual void InitializeMemory(UBlackboardComponent& OwnerComp, uint8* MemoryBlock);
 
 	/** free value before removing from blackboard, works directly on provided memory/properties */
-	virtual void FreeMemory(UBlackboardComponent& OwnerComp, uint8* MemoryBlock);
+	AIMODULE_API virtual void FreeMemory(UBlackboardComponent& OwnerComp, uint8* MemoryBlock);
 
 	/** extract location from entry, works directly on provided memory/properties */
-	virtual bool GetLocation(const UBlackboardComponent& OwnerComp, const uint8* MemoryBlock, FVector& Location) const;
+	AIMODULE_API virtual bool GetLocation(const UBlackboardComponent& OwnerComp, const uint8* MemoryBlock, FVector& Location) const;
 	
 	/** extract rotation from entry, works directly on provided memory/properties */
-	virtual bool GetRotation(const UBlackboardComponent& OwnerComp, const uint8* MemoryBlock, FRotator& Rotation) const;
+	AIMODULE_API virtual bool GetRotation(const UBlackboardComponent& OwnerComp, const uint8* MemoryBlock, FRotator& Rotation) const;
 
 	/** sets value to the default, works directly on provided memory/properties */
-	virtual void Clear(UBlackboardComponent& OwnerComp, uint8* MemoryBlock);
+	AIMODULE_API virtual void Clear(UBlackboardComponent& OwnerComp, uint8* MemoryBlock);
 
 	/** check if key has stored value, works directly on provided memory/properties */
-	virtual bool IsEmpty(const UBlackboardComponent& OwnerComp, const uint8* MemoryBlock) const;
+	AIMODULE_API virtual bool IsEmpty(const UBlackboardComponent& OwnerComp, const uint8* MemoryBlock) const;
 
 	/** various value testing, works directly on provided memory/properties */
-	virtual bool TestBasicOperation(const UBlackboardComponent& OwnerComp, const uint8* MemoryBlock, EBasicKeyOperation::Type Op) const;
-	virtual bool TestArithmeticOperation(const UBlackboardComponent& OwnerComp, const uint8* MemoryBlock, EArithmeticKeyOperation::Type Op, int32 OtherIntValue, float OtherFloatValue) const;
-	virtual bool TestTextOperation(const UBlackboardComponent& OwnerComp, const uint8* MemoryBlock, ETextKeyOperation::Type Op, const FString& OtherString) const;
+	AIMODULE_API virtual bool TestBasicOperation(const UBlackboardComponent& OwnerComp, const uint8* MemoryBlock, EBasicKeyOperation::Type Op) const;
+	AIMODULE_API virtual bool TestArithmeticOperation(const UBlackboardComponent& OwnerComp, const uint8* MemoryBlock, EArithmeticKeyOperation::Type Op, int32 OtherIntValue, float OtherFloatValue) const;
+	AIMODULE_API virtual bool TestTextOperation(const UBlackboardComponent& OwnerComp, const uint8* MemoryBlock, ETextKeyOperation::Type Op, const FString& OtherString) const;
 
 	/** convert value to text, works directly on provided memory/properties */
-	virtual FString DescribeValue(const UBlackboardComponent& OwnerComp, const uint8* MemoryBlock) const;
+	AIMODULE_API virtual FString DescribeValue(const UBlackboardComponent& OwnerComp, const uint8* MemoryBlock) const;
 };
 
 //////////////////////////////////////////////////////////////////////////
diff --git a/Engine/Source/Runtime/AIModule/Classes/BehaviorTree/Blackboard/BlackboardKeyType_Bool.h b/Engine/Source/Runtime/AIModule/Classes/BehaviorTree/Blackboard/BlackboardKeyType_Bool.h
--- a/Engine/Source/Runtime/AIModule/Classes/BehaviorTree/Blackboard/BlackboardKeyType_Bool.h
+++ b/Engine/Source/Runtime/AIModule/Classes/BehaviorTree/Blackboard/BlackboardKeyType_Bool.h
@@ -8,21 +8,21 @@
 
 class UBlackboardComponent;
 
-UCLASS(EditInlineNew, meta=(DisplayName="Bool"))
-class AIMODULE_API UBlackboardKeyType_Bool : public UBlackboardKeyType
+UCLASS(EditInlineNew, meta=(DisplayName="Bool"), MinimalAPI)
+class UBlackboardKeyType_Bool : public UBlackboardKeyType
 {
 	GENERATED_UCLASS_BODY()
 
 	typedef bool FDataType;
-	static const FDataType InvalidValue;
+	static AIMODULE_API const FDataType InvalidValue;
 
-	static bool GetValue(const UBlackboardKeyType_Bool* KeyOb, const uint8* RawData);
-	static bool SetValue(UBlackboardKeyType_Bool* KeyOb, uint8* RawData, bool bValue);
+	static AIMODULE_API bool GetValue(const UBlackboardKeyType_Bool* KeyOb, const uint8* RawData);
+	static AIMODULE_API bool SetValue(UBlackboardKeyType_Bool* KeyOb, uint8* RawData, bool bValue);
 
-	virtual EBlackboardCompare::Type CompareValues(const UBlackboardComponent& OwnerComp, const uint8* MemoryBlock,
+	AIMODULE_API virtual EBlackboardCompare::Type CompareValues(const UBlackboardComponent& OwnerComp, const uint8* MemoryBlock,
 		const UBlackboardKeyType* OtherKeyOb, const uint8* OtherMemoryBlock) const override;
 
 protected:
-	virtual FString DescribeValue(const UBlackboardComponent& OwnerComp, const uint8* RawData) const override;
-	virtual bool TestBasicOperation(const UBlackboardComponent& OwnerComp, const uint8* MemoryBlock, EBasicKeyOperation::Type Op) const override;
+	AIMODULE_API virtual FString DescribeValue(const UBlackboardComponent& OwnerComp, const uint8* RawData) const override;
+	AIMODULE_API virtual bool TestBasicOperation(const UBlackboardComponent& OwnerComp, const uint8* MemoryBlock, EBasicKeyOperation::Type Op) const override;
 };
diff --git a/Engine/Source/Runtime/AIModule/Classes/BehaviorTree/Blackboard/BlackboardKeyType_Class.h b/Engine/Source/Runtime/AIModule/Classes/BehaviorTree/Blackboard/BlackboardKeyType_Class.h
--- a/Engine/Source/Runtime/AIModule/Classes/BehaviorTree/Blackboard/BlackboardKeyType_Class.h
+++ b/Engine/Source/Runtime/AIModule/Classes/BehaviorTree/Blackboard/BlackboardKeyType_Class.h
@@ -8,27 +8,27 @@
 
 class UBlackboardComponent;
 
-UCLASS(EditInlineNew, meta=(DisplayName="Class"))
-class AIMODULE_API UBlackboardKeyType_Class : public UBlackboardKeyType
+UCLASS(EditInlineNew, meta=(DisplayName="Class"), MinimalAPI)
+class UBlackboardKeyType_Class : public UBlackboardKeyType
 {
 	GENERATED_UCLASS_BODY()
 
 	typedef UClass* FDataType;
-	static const FDataType InvalidValue;
+	static AIMODULE_API const FDataType InvalidValue;
 
 	UPROPERTY(Category=Blackboard, EditDefaultsOnly, meta=(AllowAbstract="1"))
 	TObjectPtr<UClass> BaseClass;
 	
-	static UClass* GetValue(const UBlackboardKeyType_Class* KeyOb, const uint8* RawData);
-	static bool SetValue(UBlackboardKeyType_Class* KeyOb, uint8* RawData, UClass* Value);
+	static AIMODULE_API UClass* GetValue(const UBlackboardKeyType_Class* KeyOb, const uint8* RawData);
+	static AIMODULE_API bool SetValue(UBlackboardKeyType_Class* KeyOb, uint8* RawData, UClass* Value);
 
-	virtual EBlackboardCompare::Type CompareValues(const UBlackboardComponent& OwnerComp, const uint8* MemoryBlock,
+	AIMODULE_API virtual EBlackboardCompare::Type CompareValues(const UBlackboardComponent& OwnerComp, const uint8* MemoryBlock,
 		const UBlackboardKeyType* OtherKeyOb, const uint8* OtherMemoryBlock) const override;
 
-	virtual FString DescribeSelf() const override;
-	virtual bool IsAllowedByFilter(UBlackboardKeyType* FilterOb) const override;
+	AIMODULE_API virtual FString DescribeSelf() const override;
+	AIMODULE_API virtual bool IsAllowedByFilter(UBlackboardKeyType* FilterOb) const override;
 
 protected:
-	virtual FString DescribeValue(const UBlackboardComponent& OwnerComp, const uint8* RawData) const override;
-	virtual bool TestBasicOperation(const UBlackboardComponent& OwnerComp, const uint8* MemoryBlock, EBasicKeyOperation::Type Op) const override;
+	AIMODULE_API virtual FString DescribeValue(const UBlackboardComponent& OwnerComp, const uint8* RawData) const override;
+	AIMODULE_API virtual bool TestBasicOperation(const UBlackboardComponent& OwnerComp, const uint8* MemoryBlock, EBasicKeyOperation::Type Op) const override;
 };
diff --git a/Engine/Source/Runtime/AIModule/Classes/BehaviorTree/Blackboard/BlackboardKeyType_Enum.h b/Engine/Source/Runtime/AIModule/Classes/BehaviorTree/Blackboard/BlackboardKeyType_Enum.h
--- a/Engine/Source/Runtime/AIModule/Classes/BehaviorTree/Blackboard/BlackboardKeyType_Enum.h
+++ b/Engine/Source/Runtime/AIModule/Classes/BehaviorTree/Blackboard/BlackboardKeyType_Enum.h
@@ -8,13 +8,13 @@
 
 class UBlackboardComponent;
 
-UCLASS(EditInlineNew, meta=(DisplayName="Enum"))
-class AIMODULE_API UBlackboardKeyType_Enum : public UBlackboardKeyType
+UCLASS(EditInlineNew, meta=(DisplayName="Enum"), MinimalAPI)
+class UBlackboardKeyType_Enum : public UBlackboardKeyType
 {
 	GENERATED_UCLASS_BODY()
 
 	typedef uint8 FDataType;
-	static const FDataType InvalidValue;
+	static AIMODULE_API const FDataType InvalidValue;
 
 	UPROPERTY(Category=Blackboard, EditDefaultsOnly)
 	TObjectPtr<UEnum> EnumType;
@@ -27,22 +27,22 @@ class AIMODULE_API UBlackboardKeyType_Enum : public UBlackboardKeyType
 	UPROPERTY(Category = Blackboard, VisibleDefaultsOnly)
 	uint32 bIsEnumNameValid : 1;
 
-	static FDataType GetValue(const UBlackboardKeyType_Enum* KeyOb, const uint8* RawData);
-	static bool SetValue(UBlackboardKeyType_Enum* KeyOb, uint8* RawData, FDataType Value);
+	static AIMODULE_API FDataType GetValue(const UBlackboardKeyType_Enum* KeyOb, const uint8* RawData);
+	static AIMODULE_API bool SetValue(UBlackboardKeyType_Enum* KeyOb, uint8* RawData, FDataType Value);
 
-	virtual EBlackboardCompare::Type CompareValues(const UBlackboardComponent& OwnerComp, const uint8* MemoryBlock,
+	AIMODULE_API virtual EBlackboardCompare::Type CompareValues(const UBlackboardComponent& OwnerComp, const uint8* MemoryBlock,
 		const UBlackboardKeyType* OtherKeyOb, const uint8* OtherMemoryBlock) const override;
 
-	virtual FString DescribeSelf() const override;
-	virtual FString DescribeArithmeticParam(int32 IntValue, float FloatValue) const override;
-	virtual bool IsAllowedByFilter(UBlackboardKeyType* FilterOb) const override;
+	AIMODULE_API virtual FString DescribeSelf() const override;
+	AIMODULE_API virtual FString DescribeArithmeticParam(int32 IntValue, float FloatValue) const override;
+	AIMODULE_API virtual bool IsAllowedByFilter(UBlackboardKeyType* FilterOb) const override;
 
 protected:
 #if WITH_EDITOR
-	virtual void PostEditChangeProperty(struct FPropertyChangedEvent& PropertyChangedEvent) override;
-	static bool ValidateEnum(const UEnum& EnumType);
+	AIMODULE_API virtual void PostEditChangeProperty(struct FPropertyChangedEvent& PropertyChangedEvent) override;
+	static AIMODULE_API bool ValidateEnum(const UEnum& EnumType);
 #endif
 
-	virtual FString DescribeValue(const UBlackboardComponent& OwnerComp, const uint8* RawData) const override;
-	virtual bool TestArithmeticOperation(const UBlackboardComponent& OwnerComp, const uint8* MemoryBlock, EArithmeticKeyOperation::Type Op, int32 OtherIntValue, float OtherFloatValue) const override;
+	AIMODULE_API virtual FString DescribeValue(const UBlackboardComponent& OwnerComp, const uint8* RawData) const override;
+	AIMODULE_API virtual bool TestArithmeticOperation(const UBlackboardComponent& OwnerComp, const uint8* MemoryBlock, EArithmeticKeyOperation::Type Op, int32 OtherIntValue, float OtherFloatValue) const override;
 };
diff --git a/Engine/Source/Runtime/AIModule/Classes/BehaviorTree/Blackboard/BlackboardKeyType_Float.h b/Engine/Source/Runtime/AIModule/Classes/BehaviorTree/Blackboard/BlackboardKeyType_Float.h
--- a/Engine/Source/Runtime/AIModule/Classes/BehaviorTree/Blackboard/BlackboardKeyType_Float.h
+++ b/Engine/Source/Runtime/AIModule/Classes/BehaviorTree/Blackboard/BlackboardKeyType_Float.h
@@ -8,23 +8,23 @@
 
 class UBlackboardComponent;
 
-UCLASS(EditInlineNew, meta=(DisplayName="Float"))
-class AIMODULE_API UBlackboardKeyType_Float : public UBlackboardKeyType
+UCLASS(EditInlineNew, meta=(DisplayName="Float"), MinimalAPI)
+class UBlackboardKeyType_Float : public UBlackboardKeyType
 {
 	GENERATED_UCLASS_BODY()
 
 	typedef float FDataType;
-	static const FDataType InvalidValue;
+	static AIMODULE_API const FDataType InvalidValue;
 
-	static float GetValue(const UBlackboardKeyType_Float* KeyOb, const uint8* RawData);
-	static bool SetValue(UBlackboardKeyType_Float* KeyOb, uint8* RawData, float Value);
+	static AIMODULE_API float GetValue(const UBlackboardKeyType_Float* KeyOb, const uint8* RawData);
+	static AIMODULE_API bool SetValue(UBlackboardKeyType_Float* KeyOb, uint8* RawData, float Value);
 
-	virtual EBlackboardCompare::Type CompareValues(const UBlackboardComponent& OwnerComp, const uint8* MemoryBlock,
+	AIMODULE_API virtual EBlackboardCompare::Type CompareValues(const UBlackboardComponent& OwnerComp, const uint8* MemoryBlock,
 		const UBlackboardKeyType* OtherKeyOb, const uint8* OtherMemoryBlock) const override;
 
-	virtual FString DescribeArithmeticParam(int32 IntValue, float FloatValue) const override;
+	AIMODULE_API virtual FString DescribeArithmeticParam(int32 IntValue, float FloatValue) const override;
 
 protected:
-	virtual FString DescribeValue(const UBlackboardComponent& OwnerComp, const uint8* RawData) const override;
-	virtual bool TestArithmeticOperation(const UBlackboardComponent& OwnerComp, const uint8* MemoryBlock, EArithmeticKeyOperation::Type Op, int32 OtherIntValue, float OtherFloatValue) const override;
+	AIMODULE_API virtual FString DescribeValue(const UBlackboardComponent& OwnerComp, const uint8* RawData) const override;
+	AIMODULE_API virtual bool TestArithmeticOperation(const UBlackboardComponent& OwnerComp, const uint8* MemoryBlock, EArithmeticKeyOperation::Type Op, int32 OtherIntValue, float OtherFloatValue) const override;
 };
diff --git a/Engine/Source/Runtime/AIModule/Classes/BehaviorTree/Blackboard/BlackboardKeyType_Int.h b/Engine/Source/Runtime/AIModule/Classes/BehaviorTree/Blackboard/BlackboardKeyType_Int.h
--- a/Engine/Source/Runtime/AIModule/Classes/BehaviorTree/Blackboard/BlackboardKeyType_Int.h
+++ b/Engine/Source/Runtime/AIModule/Classes/BehaviorTree/Blackboard/BlackboardKeyType_Int.h
@@ -8,22 +8,22 @@
 
 class UBlackboardComponent;
 
-UCLASS(EditInlineNew, meta=(DisplayName="Int"))
-class AIMODULE_API UBlackboardKeyType_Int : public UBlackboardKeyType
+UCLASS(EditInlineNew, meta=(DisplayName="Int"), MinimalAPI)
+class UBlackboardKeyType_Int : public UBlackboardKeyType
 {
 	GENERATED_UCLASS_BODY()
 
 	typedef int32 FDataType;
-	static const FDataType InvalidValue;
+	static AIMODULE_API const FDataType InvalidValue;
 
-	static int32 GetValue(const UBlackboardKeyType_Int* KeyOb, const uint8* RawData);
-	static bool SetValue(UBlackboardKeyType_Int* KeyOb, uint8* RawData, int32 Value);
+	static AIMODULE_API int32 GetValue(const UBlackboardKeyType_Int* KeyOb, const uint8* RawData);
+	static AIMODULE_API bool SetValue(UBlackboardKeyType_Int* KeyOb, uint8* RawData, int32 Value);
 
-	virtual FString DescribeArithmeticParam(int32 IntValue, float FloatValue) const override;
-	virtual EBlackboardCompare::Type CompareValues(const UBlackboardComponent& OwnerComp, const uint8* MemoryBlock,
+	AIMODULE_API virtual FString DescribeArithmeticParam(int32 IntValue, float FloatValue) const override;
+	AIMODULE_API virtual EBlackboardCompare::Type CompareValues(const UBlackboardComponent& OwnerComp, const uint8* MemoryBlock,
 		const UBlackboardKeyType* OtherKeyOb, const uint8* OtherMemoryBlock) const override;
 
 protected:
-	virtual FString DescribeValue(const UBlackboardComponent& OwnerComp, const uint8* RawData) const override;
-	virtual bool TestArithmeticOperation(const UBlackboardComponent& OwnerComp, const uint8* MemoryBlock, EArithmeticKeyOperation::Type Op, int32 OtherIntValue, float OtherFloatValue) const override;
+	AIMODULE_API virtual FString DescribeValue(const UBlackboardComponent& OwnerComp, const uint8* RawData) const override;
+	AIMODULE_API virtual bool TestArithmeticOperation(const UBlackboardComponent& OwnerComp, const uint8* MemoryBlock, EArithmeticKeyOperation::Type Op, int32 OtherIntValue, float OtherFloatValue) const override;
 };
diff --git a/Engine/Source/Runtime/AIModule/Classes/BehaviorTree/Blackboard/BlackboardKeyType_Name.h b/Engine/Source/Runtime/AIModule/Classes/BehaviorTree/Blackboard/BlackboardKeyType_Name.h
--- a/Engine/Source/Runtime/AIModule/Classes/BehaviorTree/Blackboard/BlackboardKeyType_Name.h
+++ b/Engine/Source/Runtime/AIModule/Classes/BehaviorTree/Blackboard/BlackboardKeyType_Name.h
@@ -8,21 +8,21 @@
 
 class UBlackboardComponent;
 
-UCLASS(EditInlineNew, meta=(DisplayName="Name"))
-class AIMODULE_API UBlackboardKeyType_Name : public UBlackboardKeyType
+UCLASS(EditInlineNew, meta=(DisplayName="Name"), MinimalAPI)
+class UBlackboardKeyType_Name : public UBlackboardKeyType
 {
 	GENERATED_UCLASS_BODY()
 
 	typedef FName FDataType;
-	static const FDataType InvalidValue;
+	static AIMODULE_API const FDataType InvalidValue;
 
-	static FName GetValue(const UBlackboardKeyType_Name* KeyOb, const uint8* RawData);
-	static bool SetValue(UBlackboardKeyType_Name* KeyOb, uint8* RawData, const FName& Value);
+	static AIMODULE_API FName GetValue(const UBlackboardKeyType_Name* KeyOb, const uint8* RawData);
+	static AIMODULE_API bool SetValue(UBlackboardKeyType_Name* KeyOb, uint8* RawData, const FName& Value);
 
-	virtual EBlackboardCompare::Type CompareValues(const UBlackboardComponent& OwnerComp, const uint8* MemoryBlock,
+	AIMODULE_API virtual EBlackboardCompare::Type CompareValues(const UBlackboardComponent& OwnerComp, const uint8* MemoryBlock,
 		const UBlackboardKeyType* OtherKeyOb, const uint8* OtherMemoryBlock) const override;
 
 protected:
-	virtual FString DescribeValue(const UBlackboardComponent& OwnerComp, const uint8* RawData) const override;
-	virtual bool TestTextOperation(const UBlackboardComponent& OwnerComp, const uint8* MemoryBlock, ETextKeyOperation::Type Op, const FString& OtherString) const override;
+	AIMODULE_API virtual FString DescribeValue(const UBlackboardComponent& OwnerComp, const uint8* RawData) const override;
+	AIMODULE_API virtual bool TestTextOperation(const UBlackboardComponent& OwnerComp, const uint8* MemoryBlock, ETextKeyOperation::Type Op, const FString& OtherString) const override;
 };
diff --git a/Engine/Source/Runtime/AIModule/Classes/BehaviorTree/Blackboard/BlackboardKeyType_NativeEnum.h b/Engine/Source/Runtime/AIModule/Classes/BehaviorTree/Blackboard/BlackboardKeyType_NativeEnum.h
--- a/Engine/Source/Runtime/AIModule/Classes/BehaviorTree/Blackboard/BlackboardKeyType_NativeEnum.h
+++ b/Engine/Source/Runtime/AIModule/Classes/BehaviorTree/Blackboard/BlackboardKeyType_NativeEnum.h
@@ -9,13 +9,13 @@
 
 // DEPRECATED, please use UBlackboardKeyType_Enum instead
 
-UCLASS(NotEditInlineNew, HideDropDown)
-class AIMODULE_API UBlackboardKeyType_NativeEnum : public UBlackboardKeyType
+UCLASS(NotEditInlineNew, HideDropDown, MinimalAPI)
+class UBlackboardKeyType_NativeEnum : public UBlackboardKeyType
 {
 	GENERATED_UCLASS_BODY()
 
 	typedef uint8 FDataType;
-	static const FDataType InvalidValue;
+	static AIMODULE_API const FDataType InvalidValue;
 
 	UPROPERTY(Category=Blackboard, EditDefaultsOnly)
 	FString EnumName;
@@ -23,8 +23,8 @@ class AIMODULE_API UBlackboardKeyType_NativeEnum : public UBlackboardKeyType
 	UPROPERTY()
 	TObjectPtr<UEnum> EnumType;
 
-	virtual UBlackboardKeyType* UpdateDeprecatedKey() override;
+	AIMODULE_API virtual UBlackboardKeyType* UpdateDeprecatedKey() override;
 
-	static uint8 GetValue(const UBlackboardKeyType_NativeEnum* KeyOb, const uint8* RawData);
-	static bool SetValue(UBlackboardKeyType_NativeEnum* KeyOb, uint8* RawData, uint8 Value);
+	static AIMODULE_API uint8 GetValue(const UBlackboardKeyType_NativeEnum* KeyOb, const uint8* RawData);
+	static AIMODULE_API bool SetValue(UBlackboardKeyType_NativeEnum* KeyOb, uint8* RawData, uint8 Value);
 };
diff --git a/Engine/Source/Runtime/AIModule/Classes/BehaviorTree/Blackboard/BlackboardKeyType_Object.h b/Engine/Source/Runtime/AIModule/Classes/BehaviorTree/Blackboard/BlackboardKeyType_Object.h
--- a/Engine/Source/Runtime/AIModule/Classes/BehaviorTree/Blackboard/BlackboardKeyType_Object.h
+++ b/Engine/Source/Runtime/AIModule/Classes/BehaviorTree/Blackboard/BlackboardKeyType_Object.h
@@ -8,29 +8,29 @@
 
 class UBlackboardComponent;
 
-UCLASS(EditInlineNew, meta=(DisplayName="Object"))
-class AIMODULE_API UBlackboardKeyType_Object : public UBlackboardKeyType
+UCLASS(EditInlineNew, meta=(DisplayName="Object"), MinimalAPI)
+class UBlackboardKeyType_Object : public UBlackboardKeyType
 {
 	GENERATED_UCLASS_BODY()
 
 	typedef UObject* FDataType;
-	static const FDataType InvalidValue;
+	static AIMODULE_API const FDataType InvalidValue;
 
 	UPROPERTY(Category=Blackboard, EditDefaultsOnly, NoClear, meta=(AllowAbstract="1"))
 	TObjectPtr<UClass> BaseClass;
 
-	static UObject* GetValue(const UBlackboardKeyType_Object* KeyOb, const uint8* RawData);
-	static bool SetValue(UBlackboardKeyType_Object* KeyOb, uint8* RawData, UObject* Value);
+	static AIMODULE_API UObject* GetValue(const UBlackboardKeyType_Object* KeyOb, const uint8* RawData);
+	static AIMODULE_API bool SetValue(UBlackboardKeyType_Object* KeyOb, uint8* RawData, UObject* Value);
 
-	virtual EBlackboardCompare::Type CompareValues(const UBlackboardComponent& OwnerComp, const uint8* MemoryBlock,
+	AIMODULE_API virtual EBlackboardCompare::Type CompareValues(const UBlackboardComponent& OwnerComp, const uint8* MemoryBlock,
 		const UBlackboardKeyType* OtherKeyOb, const uint8* OtherMemoryBlock) const override;
 
-	virtual FString DescribeSelf() const override;
-	virtual bool IsAllowedByFilter(UBlackboardKeyType* FilterOb) const override;
+	AIMODULE_API virtual FString DescribeSelf() const override;
+	AIMODULE_API virtual bool IsAllowedByFilter(UBlackboardKeyType* FilterOb) const override;
 
 protected:
-	virtual FString DescribeValue(const UBlackboardComponent& OwnerComp, const uint8* RawData) const override;
-	virtual bool GetLocation(const UBlackboardComponent& OwnerComp, const uint8* RawData, FVector& Location) const override;
-	virtual bool GetRotation(const UBlackboardComponent& OwnerComp, const uint8* MemoryBlock, FRotator& Rotation) const override;
-	virtual bool TestBasicOperation(const UBlackboardComponent& OwnerComp, const uint8* MemoryBlock, EBasicKeyOperation::Type Op) const override;
+	AIMODULE_API virtual FString DescribeValue(const UBlackboardComponent& OwnerComp, const uint8* RawData) const override;
+	AIMODULE_API virtual bool GetLocation(const UBlackboardComponent& OwnerComp, const uint8* RawData, FVector& Location) const override;
+	AIMODULE_API virtual bool GetRotation(const UBlackboardComponent& OwnerComp, const uint8* MemoryBlock, FRotator& Rotation) const override;
+	AIMODULE_API virtual bool TestBasicOperation(const UBlackboardComponent& OwnerComp, const uint8* MemoryBlock, EBasicKeyOperation::Type Op) const override;
 };
diff --git a/Engine/Source/Runtime/AIModule/Classes/BehaviorTree/Blackboard/BlackboardKeyType_Rotator.h b/Engine/Source/Runtime/AIModule/Classes/BehaviorTree/Blackboard/BlackboardKeyType_Rotator.h
--- a/Engine/Source/Runtime/AIModule/Classes/BehaviorTree/Blackboard/BlackboardKeyType_Rotator.h
+++ b/Engine/Source/Runtime/AIModule/Classes/BehaviorTree/Blackboard/BlackboardKeyType_Rotator.h
@@ -8,25 +8,25 @@
 
 class UBlackboardComponent;
 
-UCLASS(EditInlineNew, meta=(DisplayName="Rotator"))
-class AIMODULE_API UBlackboardKeyType_Rotator : public UBlackboardKeyType
+UCLASS(EditInlineNew, meta=(DisplayName="Rotator"), MinimalAPI)
+class UBlackboardKeyType_Rotator : public UBlackboardKeyType
 {
 	GENERATED_UCLASS_BODY()
 	
 	typedef FRotator FDataType; 
-	static const FDataType InvalidValue;
+	static AIMODULE_API const FDataType InvalidValue;
 
-	static FRotator GetValue(const UBlackboardKeyType_Rotator* KeyOb, const uint8* RawData);
-	static bool SetValue(UBlackboardKeyType_Rotator* KeyOb, uint8* RawData, const FRotator& Value);
+	static AIMODULE_API FRotator GetValue(const UBlackboardKeyType_Rotator* KeyOb, const uint8* RawData);
+	static AIMODULE_API bool SetValue(UBlackboardKeyType_Rotator* KeyOb, uint8* RawData, const FRotator& Value);
 
-	virtual EBlackboardCompare::Type CompareValues(const UBlackboardComponent& OwnerComp, const uint8* MemoryBlock,
+	AIMODULE_API virtual EBlackboardCompare::Type CompareValues(const UBlackboardComponent& OwnerComp, const uint8* MemoryBlock,
 		const UBlackboardKeyType* OtherKeyOb, const uint8* OtherMemoryBlock) const override;
 
 protected:
-	virtual void InitializeMemory(UBlackboardComponent& OwnerComp, uint8* RawData) override;
-	virtual FString DescribeValue(const UBlackboardComponent& OwnerComp, const uint8* RawData) const override;
-	virtual bool GetRotation(const UBlackboardComponent& OwnerComp, const uint8* RawData, FRotator& Rotation) const override;
-	virtual bool IsEmpty(const UBlackboardComponent& OwnerComp, const uint8* RawData) const override;
-	virtual void Clear(UBlackboardComponent& OwnerComp, uint8* RawData) override;
-	virtual bool TestBasicOperation(const UBlackboardComponent& OwnerComp, const uint8* MemoryBlock, EBasicKeyOperation::Type Op) const override;
+	AIMODULE_API virtual void InitializeMemory(UBlackboardComponent& OwnerComp, uint8* RawData) override;
+	AIMODULE_API virtual FString DescribeValue(const UBlackboardComponent& OwnerComp, const uint8* RawData) const override;
+	AIMODULE_API virtual bool GetRotation(const UBlackboardComponent& OwnerComp, const uint8* RawData, FRotator& Rotation) const override;
+	AIMODULE_API virtual bool IsEmpty(const UBlackboardComponent& OwnerComp, const uint8* RawData) const override;
+	AIMODULE_API virtual void Clear(UBlackboardComponent& OwnerComp, uint8* RawData) override;
+	AIMODULE_API virtual bool TestBasicOperation(const UBlackboardComponent& OwnerComp, const uint8* MemoryBlock, EBasicKeyOperation::Type Op) const override;
 };
diff --git a/Engine/Source/Runtime/AIModule/Classes/BehaviorTree/Blackboard/BlackboardKeyType_String.h b/Engine/Source/Runtime/AIModule/Classes/BehaviorTree/Blackboard/BlackboardKeyType_String.h
--- a/Engine/Source/Runtime/AIModule/Classes/BehaviorTree/Blackboard/BlackboardKeyType_String.h
+++ b/Engine/Source/Runtime/AIModule/Classes/BehaviorTree/Blackboard/BlackboardKeyType_String.h
@@ -8,27 +8,27 @@
 
 class UBlackboardComponent;
 
-UCLASS(EditInlineNew, meta=(DisplayName="String"))
-class AIMODULE_API UBlackboardKeyType_String : public UBlackboardKeyType
+UCLASS(EditInlineNew, meta=(DisplayName="String"), MinimalAPI)
+class UBlackboardKeyType_String : public UBlackboardKeyType
 {
 	GENERATED_UCLASS_BODY()
 
 	typedef FString FDataType;
-	static const FDataType InvalidValue;
+	static AIMODULE_API const FDataType InvalidValue;
 
-	static FString GetValue(const UBlackboardKeyType_String* KeyOb, const uint8* RawData);
-	static bool SetValue(UBlackboardKeyType_String* KeyOb, uint8* RawData, const FString& Value);
+	static AIMODULE_API FString GetValue(const UBlackboardKeyType_String* KeyOb, const uint8* RawData);
+	static AIMODULE_API bool SetValue(UBlackboardKeyType_String* KeyOb, uint8* RawData, const FString& Value);
 
-	virtual EBlackboardCompare::Type CompareValues(const UBlackboardComponent& OwnerComp, const uint8* MemoryBlock,
+	AIMODULE_API virtual EBlackboardCompare::Type CompareValues(const UBlackboardComponent& OwnerComp, const uint8* MemoryBlock,
 		const UBlackboardKeyType* OtherKeyOb, const uint8* OtherMemoryBlock) const override;
 
 	UPROPERTY()
 	FString StringValue;
 
 protected:
-	virtual FString DescribeValue(const UBlackboardComponent& OwnerComp, const uint8* RawData) const override;
-	virtual bool TestTextOperation(const UBlackboardComponent& OwnerComp, const uint8* MemoryBlock, ETextKeyOperation::Type Op, const FString& OtherString) const override;
-	virtual void Clear(UBlackboardComponent& OwnerComp, uint8* MemoryBlock) override;
-	virtual bool IsEmpty(const UBlackboardComponent& OwnerComp, const uint8* MemoryBlock) const override;
-	virtual void CopyValues(UBlackboardComponent& OwnerComp, uint8* MemoryBlock, const UBlackboardKeyType* SourceKeyOb, const uint8* SourceBlock) override;
+	AIMODULE_API virtual FString DescribeValue(const UBlackboardComponent& OwnerComp, const uint8* RawData) const override;
+	AIMODULE_API virtual bool TestTextOperation(const UBlackboardComponent& OwnerComp, const uint8* MemoryBlock, ETextKeyOperation::Type Op, const FString& OtherString) const override;
+	AIMODULE_API virtual void Clear(UBlackboardComponent& OwnerComp, uint8* MemoryBlock) override;
+	AIMODULE_API virtual bool IsEmpty(const UBlackboardComponent& OwnerComp, const uint8* MemoryBlock) const override;
+	AIMODULE_API virtual void CopyValues(UBlackboardComponent& OwnerComp, uint8* MemoryBlock, const UBlackboardKeyType* SourceKeyOb, const uint8* SourceBlock) override;
 };
diff --git a/Engine/Source/Runtime/AIModule/Classes/BehaviorTree/Blackboard/BlackboardKeyType_Vector.h b/Engine/Source/Runtime/AIModule/Classes/BehaviorTree/Blackboard/BlackboardKeyType_Vector.h
--- a/Engine/Source/Runtime/AIModule/Classes/BehaviorTree/Blackboard/BlackboardKeyType_Vector.h
+++ b/Engine/Source/Runtime/AIModule/Classes/BehaviorTree/Blackboard/BlackboardKeyType_Vector.h
@@ -8,25 +8,25 @@
 
 class UBlackboardComponent;
 
-UCLASS(EditInlineNew, meta=(DisplayName="Vector"))
-class AIMODULE_API UBlackboardKeyType_Vector : public UBlackboardKeyType
+UCLASS(EditInlineNew, meta=(DisplayName="Vector"), MinimalAPI)
+class UBlackboardKeyType_Vector : public UBlackboardKeyType
 {
 	GENERATED_UCLASS_BODY()
 
 	typedef FVector FDataType; 
-	static const FDataType InvalidValue;
+	static AIMODULE_API const FDataType InvalidValue;
 	
-	static FVector GetValue(const UBlackboardKeyType_Vector* KeyOb, const uint8* RawData);
-	static bool SetValue(UBlackboardKeyType_Vector* KeyOb, uint8* RawData, const FVector& Value);
+	static AIMODULE_API FVector GetValue(const UBlackboardKeyType_Vector* KeyOb, const uint8* RawData);
+	static AIMODULE_API bool SetValue(UBlackboardKeyType_Vector* KeyOb, uint8* RawData, const FVector& Value);
 
-	virtual EBlackboardCompare::Type CompareValues(const UBlackboardComponent& OwnerComp, const uint8* MemoryBlock,
+	AIMODULE_API virtual EBlackboardCompare::Type CompareValues(const UBlackboardComponent& OwnerComp, const uint8* MemoryBlock,
 		const UBlackboardKeyType* OtherKeyOb, const uint8* OtherMemoryBlock) const override;
 
 protected:
-	virtual void InitializeMemory(UBlackboardComponent& OwnerComp, uint8* RawData) override;
-	virtual FString DescribeValue(const UBlackboardComponent& OwnerComp, const uint8* RawData) const override;
-	virtual bool GetLocation(const UBlackboardComponent& OwnerComp, const uint8* RawData, FVector& Location) const override;
-	virtual bool IsEmpty(const UBlackboardComponent& OwnerComp, const uint8* MemoryBlock) const override;
-	virtual void Clear(UBlackboardComponent& OwnerComp, uint8* MemoryBlock) override;
-	virtual bool TestBasicOperation(const UBlackboardComponent& OwnerComp, const uint8* MemoryBlock, EBasicKeyOperation::Type Op) const override;
+	AIMODULE_API virtual void InitializeMemory(UBlackboardComponent& OwnerComp, uint8* RawData) override;
+	AIMODULE_API virtual FString DescribeValue(const UBlackboardComponent& OwnerComp, const uint8* RawData) const override;
+	AIMODULE_API virtual bool GetLocation(const UBlackboardComponent& OwnerComp, const uint8* RawData, FVector& Location) const override;
+	AIMODULE_API virtual bool IsEmpty(const UBlackboardComponent& OwnerComp, const uint8* MemoryBlock) const override;
+	AIMODULE_API virtual void Clear(UBlackboardComponent& OwnerComp, uint8* MemoryBlock) override;
+	AIMODULE_API virtual bool TestBasicOperation(const UBlackboardComponent& OwnerComp, const uint8* MemoryBlock, EBasicKeyOperation::Type Op) const override;
 };
diff --git a/Engine/Source/Runtime/AIModule/Classes/BehaviorTree/BlackboardComponent.h b/Engine/Source/Runtime/AIModule/Classes/BehaviorTree/BlackboardComponent.h
--- a/Engine/Source/Runtime/AIModule/Classes/BehaviorTree/BlackboardComponent.h
+++ b/Engine/Source/Runtime/AIModule/Classes/BehaviorTree/BlackboardComponent.h
@@ -38,148 +38,148 @@ namespace EBlackboardDescription
 }
 
 
-UCLASS(ClassGroup = AI, meta = (BlueprintSpawnableComponent), hidecategories = (Sockets, Collision))
-class AIMODULE_API UBlackboardComponent : public UActorComponent
+UCLASS(ClassGroup = AI, meta = (BlueprintSpawnableComponent), hidecategories = (Sockets, Collision), MinimalAPI)
+class UBlackboardComponent : public UActorComponent
 {
 	GENERATED_BODY()
 
 public:
-	UBlackboardComponent(const FObjectInitializer& ObjectInitializer = FObjectInitializer::Get());
+	AIMODULE_API UBlackboardComponent(const FObjectInitializer& ObjectInitializer = FObjectInitializer::Get());
 
 	/** BEGIN UActorComponent overrides */
-	virtual void InitializeComponent() override;
-	virtual void UninitializeComponent() override;
+	AIMODULE_API virtual void InitializeComponent() override;
+	AIMODULE_API virtual void UninitializeComponent() override;
 	/** END UActorComponent overrides */
 
 	/** @return name of key */
-	FName GetKeyName(FBlackboard::FKey KeyID) const;
+	AIMODULE_API FName GetKeyName(FBlackboard::FKey KeyID) const;
 
 	/** @return key ID from name */
-	FBlackboard::FKey GetKeyID(const FName& KeyName) const;
+	AIMODULE_API FBlackboard::FKey GetKeyID(const FName& KeyName) const;
 
 	/** @return class of value for given key */
-	TSubclassOf<UBlackboardKeyType> GetKeyType(FBlackboard::FKey KeyID) const;
+	AIMODULE_API TSubclassOf<UBlackboardKeyType> GetKeyType(FBlackboard::FKey KeyID) const;
 
 	/** @return true if the key is marked as instance synced */
-	bool IsKeyInstanceSynced(FBlackboard::FKey KeyID) const;
+	AIMODULE_API bool IsKeyInstanceSynced(FBlackboard::FKey KeyID) const;
 	
 	/** @return number of entries in data asset */
-	int32 GetNumKeys() const;
+	AIMODULE_API int32 GetNumKeys() const;
 
 	/** @return true if blackboard have valid data asset */
-	bool HasValidAsset() const;
+	AIMODULE_API bool HasValidAsset() const;
 
 	/** register observer for blackboard key */
-	FDelegateHandle RegisterObserver(FBlackboard::FKey KeyID, const UObject* NotifyOwner, FOnBlackboardChangeNotification ObserverDelegate);
+	AIMODULE_API FDelegateHandle RegisterObserver(FBlackboard::FKey KeyID, const UObject* NotifyOwner, FOnBlackboardChangeNotification ObserverDelegate);
 
 	/** unregister observer from blackboard key */
-	void UnregisterObserver(FBlackboard::FKey KeyID, FDelegateHandle ObserverHandle);
+	AIMODULE_API void UnregisterObserver(FBlackboard::FKey KeyID, FDelegateHandle ObserverHandle);
 
 	/** unregister all observers associated with given owner */
-	void UnregisterObserversFrom(const UObject* NotifyOwner);
+	AIMODULE_API void UnregisterObserversFrom(const UObject* NotifyOwner);
 
 	/** pause observer change notifications, any new ones will be added to a queue */
-	void PauseObserverNotifications();
+	AIMODULE_API void PauseObserverNotifications();
 
 	/** resume observer change notifications and, optionally, process the queued observation list */
-	void ResumeObserverNotifications(bool bSendQueuedObserverNotifications);
+	AIMODULE_API void ResumeObserverNotifications(bool bSendQueuedObserverNotifications);
 
 	/** @return associated behavior tree component */
-	UBrainComponent* GetBrainComponent() const;
+	AIMODULE_API UBrainComponent* GetBrainComponent() const;
 
 	/** @return blackboard data asset */
-	UBlackboardData* GetBlackboardAsset() const;
+	AIMODULE_API UBlackboardData* GetBlackboardAsset() const;
 
 	/** caches UBrainComponent pointer to be used in communication */
-	void CacheBrainComponent(UBrainComponent& BrainComponent);
+	AIMODULE_API void CacheBrainComponent(UBrainComponent& BrainComponent);
 
 	/** setup component for using given blackboard asset, returns true if blackboard is properly initialized for specified blackboard data */
-	bool InitializeBlackboard(UBlackboardData& NewAsset);
+	AIMODULE_API bool InitializeBlackboard(UBlackboardData& NewAsset);
 	
 	/** @return true if component can be used with specified blackboard asset */
-	virtual bool IsCompatibleWith(const UBlackboardData* TestAsset) const;
+	AIMODULE_API virtual bool IsCompatibleWith(const UBlackboardData* TestAsset) const;
 
 	UFUNCTION(BlueprintCallable, Category="AI|Components|Blackboard")
-	UObject* GetValueAsObject(const FName& KeyName) const;
+	AIMODULE_API UObject* GetValueAsObject(const FName& KeyName) const;
 
 	UFUNCTION(BlueprintCallable, Category="AI|Components|Blackboard")
-	UClass* GetValueAsClass(const FName& KeyName) const;
+	AIMODULE_API UClass* GetValueAsClass(const FName& KeyName) const;
 
 	UFUNCTION(BlueprintCallable, Category="AI|Components|Blackboard")
-	uint8 GetValueAsEnum(const FName& KeyName) const;
+	AIMODULE_API uint8 GetValueAsEnum(const FName& KeyName) const;
 
 	UFUNCTION(BlueprintCallable, Category="AI|Components|Blackboard")
-	int32 GetValueAsInt(const FName& KeyName) const;
+	AIMODULE_API int32 GetValueAsInt(const FName& KeyName) const;
 
 	UFUNCTION(BlueprintCallable, Category="AI|Components|Blackboard")
-	float GetValueAsFloat(const FName& KeyName) const;
+	AIMODULE_API float GetValueAsFloat(const FName& KeyName) const;
 
 	UFUNCTION(BlueprintCallable, Category="AI|Components|Blackboard")
-	bool GetValueAsBool(const FName& KeyName) const;
+	AIMODULE_API bool GetValueAsBool(const FName& KeyName) const;
 
 	UFUNCTION(BlueprintCallable, Category="AI|Components|Blackboard")
-	FString GetValueAsString(const FName& KeyName) const;
+	AIMODULE_API FString GetValueAsString(const FName& KeyName) const;
 	
 	UFUNCTION(BlueprintCallable, Category="AI|Components|Blackboard")
-	FName GetValueAsName(const FName& KeyName) const;
+	AIMODULE_API FName GetValueAsName(const FName& KeyName) const;
 
 	UFUNCTION(BlueprintCallable, Category="AI|Components|Blackboard")
-	FVector GetValueAsVector(const FName& KeyName) const;
+	AIMODULE_API FVector GetValueAsVector(const FName& KeyName) const;
 
 	UFUNCTION(BlueprintCallable, Category="AI|Components|Blackboard")
-	FRotator GetValueAsRotator(const FName& KeyName) const;
+	AIMODULE_API FRotator GetValueAsRotator(const FName& KeyName) const;
 
 	UFUNCTION(BlueprintCallable, Category="AI|Components|Blackboard")
-	void SetValueAsObject(const FName& KeyName, UObject* ObjectValue);
+	AIMODULE_API void SetValueAsObject(const FName& KeyName, UObject* ObjectValue);
 	
 	UFUNCTION(BlueprintCallable, Category="AI|Components|Blackboard")
-	void SetValueAsClass(const FName& KeyName, UClass* ClassValue);
+	AIMODULE_API void SetValueAsClass(const FName& KeyName, UClass* ClassValue);
 
 	UFUNCTION(BlueprintCallable, Category="AI|Components|Blackboard")
-	void SetValueAsEnum(const FName& KeyName, uint8 EnumValue);
+	AIMODULE_API void SetValueAsEnum(const FName& KeyName, uint8 EnumValue);
 
 	UFUNCTION(BlueprintCallable, Category="AI|Components|Blackboard")
-	void SetValueAsInt(const FName& KeyName, int32 IntValue);
+	AIMODULE_API void SetValueAsInt(const FName& KeyName, int32 IntValue);
 
 	UFUNCTION(BlueprintCallable, Category="AI|Components|Blackboard")
-	void SetValueAsFloat(const FName& KeyName, float FloatValue);
+	AIMODULE_API void SetValueAsFloat(const FName& KeyName, float FloatValue);
 
 	UFUNCTION(BlueprintCallable, Category="AI|Components|Blackboard")
-	void SetValueAsBool(const FName& KeyName, bool BoolValue);
+	AIMODULE_API void SetValueAsBool(const FName& KeyName, bool BoolValue);
 
 	UFUNCTION(BlueprintCallable, Category="AI|Components|Blackboard")
-	void SetValueAsString(const FName& KeyName, FString StringValue);
+	AIMODULE_API void SetValueAsString(const FName& KeyName, FString StringValue);
 
 	UFUNCTION(BlueprintCallable, Category="AI|Components|Blackboard")
-	void SetValueAsName(const FName& KeyName, FName NameValue);
+	AIMODULE_API void SetValueAsName(const FName& KeyName, FName NameValue);
 
 	UFUNCTION(BlueprintCallable, Category="AI|Components|Blackboard")
-	void SetValueAsVector(const FName& KeyName, FVector VectorValue);
+	AIMODULE_API void SetValueAsVector(const FName& KeyName, FVector VectorValue);
 
 	UFUNCTION(BlueprintCallable, Category = "AI|Components|Blackboard")
-	void SetValueAsRotator(const FName& KeyName, FRotator VectorValue);
+	AIMODULE_API void SetValueAsRotator(const FName& KeyName, FRotator VectorValue);
 
 	UFUNCTION(BlueprintCallable, Category="AI|Components|Blackboard", Meta=(
 		Tooltip="If the vector value has been set (and not cleared), this function returns true (indicating that the value should be valid).  If it's not set, the vector value is invalid and this function will return false.  (Also returns false if the key specified does not hold a vector.)"))
-	bool IsVectorValueSet(const FName& KeyName) const;
-	bool IsVectorValueSet(FBlackboard::FKey KeyID) const;
+	AIMODULE_API bool IsVectorValueSet(const FName& KeyName) const;
+	AIMODULE_API bool IsVectorValueSet(FBlackboard::FKey KeyID) const;
 		
 	/** return false if call failed (most probably no such entry in BB) */
 	UFUNCTION(BlueprintCallable, Category="AI|Components|Blackboard")
-	bool GetLocationFromEntry(const FName& KeyName, FVector& ResultLocation) const;
-	bool GetLocationFromEntry(FBlackboard::FKey KeyID, FVector& ResultLocation) const;
+	AIMODULE_API bool GetLocationFromEntry(const FName& KeyName, FVector& ResultLocation) const;
+	AIMODULE_API bool GetLocationFromEntry(FBlackboard::FKey KeyID, FVector& ResultLocation) const;
 
 	/** return false if call failed (most probably no such entry in BB) */
 	UFUNCTION(BlueprintCallable, Category="AI|Components|Blackboard")
-	bool GetRotationFromEntry(const FName& KeyName, FRotator& ResultRotation) const;
-	bool GetRotationFromEntry(FBlackboard::FKey KeyID, FRotator& ResultRotation) const;
+	AIMODULE_API bool GetRotationFromEntry(const FName& KeyName, FRotator& ResultRotation) const;
+	AIMODULE_API bool GetRotationFromEntry(FBlackboard::FKey KeyID, FRotator& ResultRotation) const;
 
 	UFUNCTION(BlueprintCallable, Category = "AI|Components|Blackboard")
-	void ClearValue(const FName& KeyName);
-	void ClearValue(FBlackboard::FKey KeyID);
+	AIMODULE_API void ClearValue(const FName& KeyName);
+	AIMODULE_API void ClearValue(FBlackboard::FKey KeyID);
 
 	/** Copy content from SourceKeyID to DestinationID and return true if it worked */
-	bool CopyKeyValue(FBlackboard::FKey SourceKeyID, FBlackboard::FKey DestinationID);
+	AIMODULE_API bool CopyKeyValue(FBlackboard::FKey SourceKeyID, FBlackboard::FKey DestinationID);
 
 	template<class TDataClass>
 	bool IsKeyOfType(FBlackboard::FKey KeyID) const;
@@ -205,20 +205,20 @@ public:
 
 	PRAGMA_DISABLE_DEPRECATION_WARNINGS // re BlackboardAsset
 	FORCEINLINE bool IsValidKey(FBlackboard::FKey KeyID) const { check(BlackboardAsset); return BlackboardAsset->IsValidKey(KeyID); }
-	PRAGMA_ENABLE_DEPRECATION_WARNINGS // re BlackboardAsset
+	AIMODULE_API PRAGMA_ENABLE_DEPRECATION_WARNINGS // re BlackboardAsset
 
 	/** compares blackboard's values under specified keys */
 	EBlackboardCompare::Type CompareKeyValues(TSubclassOf<UBlackboardKeyType> KeyType, FBlackboard::FKey KeyA, FBlackboard::FKey KeyB) const;
 
-	FString GetDebugInfoString(EBlackboardDescription::Type Mode) const;
+	AIMODULE_API FString GetDebugInfoString(EBlackboardDescription::Type Mode) const;
 
 	/** get description of value under given key */
-	FString DescribeKeyValue(const FName& KeyName, EBlackboardDescription::Type Mode) const;
-	FString DescribeKeyValue(FBlackboard::FKey KeyID, EBlackboardDescription::Type Mode) const;
+	AIMODULE_API FString DescribeKeyValue(const FName& KeyName, EBlackboardDescription::Type Mode) const;
+	AIMODULE_API FString DescribeKeyValue(FBlackboard::FKey KeyID, EBlackboardDescription::Type Mode) const;
 
 #if ENABLE_VISUAL_LOG
 	/** prepare blackboard snapshot for logs */
-	virtual void DescribeSelfToVisLog(struct FVisualLogEntry* Snapshot) const;
+	AIMODULE_API virtual void DescribeSelfToVisLog(struct FVisualLogEntry* Snapshot) const;
 #endif
 	
 protected:
@@ -288,18 +288,18 @@ protected:
 	uint32 bSynchronizedKeyPopulated : 1;
 
 	/** notifies behavior tree decorators about change in blackboard */
-	void NotifyObservers(FBlackboard::FKey KeyID) const;
+	AIMODULE_API void NotifyObservers(FBlackboard::FKey KeyID) const;
 
 	/** initializes parent chain in asset */
-	void InitializeParentChain(UBlackboardData* NewAsset);
+	AIMODULE_API void InitializeParentChain(UBlackboardData* NewAsset);
 
 	/** destroy allocated values */
-	void DestroyValues();
+	AIMODULE_API void DestroyValues();
 
 	/** populates BB's synchronized entries */
-	void PopulateSynchronizedKeys();
+	AIMODULE_API void PopulateSynchronizedKeys();
 
-	bool ShouldSyncWithBlackboard(UBlackboardComponent& OtherBlackboardComponent) const;
+	AIMODULE_API bool ShouldSyncWithBlackboard(UBlackboardComponent& OtherBlackboardComponent) const;
 
 	friend UBlackboardKeyType;
 };
@@ -479,4 +479,4 @@ public:
 	}
 
 	bool IsValid() const { return BBKey != FBlackboard::InvalidKey; }
-};
\ No newline at end of file
+};
