[Engine]

diff --git a/Engine/Source/Runtime/AIModule/Classes/BehaviorTree/BlackboardData.h b/Engine/Source/Runtime/AIModule/Classes/BehaviorTree/BlackboardData.h
--- a/Engine/Source/Runtime/AIModule/Classes/BehaviorTree/BlackboardData.h
+++ b/Engine/Source/Runtime/AIModule/Classes/BehaviorTree/BlackboardData.h
@@ -11,7 +11,7 @@
 
 /** blackboard entry definition */
 USTRUCT()
-struct AIMODULE_API FBlackboardEntry
+struct FBlackboardEntry
 {
 	GENERATED_USTRUCT_BODY()
 
@@ -38,11 +38,11 @@ struct AIMODULE_API FBlackboardEntry
 		: KeyType(nullptr), bInstanceSynced(0)
 	{}
 
-	bool operator==(const FBlackboardEntry& Other) const;
+	AIMODULE_API bool operator==(const FBlackboardEntry& Other) const;
 };
 
-UCLASS(BlueprintType, AutoExpandCategories=(Blackboard))
-class AIMODULE_API UBlackboardData : public UDataAsset
+UCLASS(BlueprintType, AutoExpandCategories=(Blackboard), MinimalAPI)
+class UBlackboardData : public UDataAsset
 {
 	GENERATED_UCLASS_BODY()
 	DECLARE_MULTICAST_DELEGATE_OneParam(FKeyUpdate, UBlackboardData* /*asset*/);
@@ -70,38 +70,38 @@ public:
 	FORCEINLINE bool HasSynchronizedKeys() const { return bHasSynchronizedKeys; }
 
 	/** @return true if the key is instance synced */
-	bool IsKeyInstanceSynced(FBlackboard::FKey KeyID) const;
+	AIMODULE_API bool IsKeyInstanceSynced(FBlackboard::FKey KeyID) const;
 	
 	/** @return key ID from name */
-	FBlackboard::FKey GetKeyID(const FName& KeyName) const;
+	AIMODULE_API FBlackboard::FKey GetKeyID(const FName& KeyName) const;
 
 	/** @return name of key */
-	FName GetKeyName(FBlackboard::FKey KeyID) const;
+	AIMODULE_API FName GetKeyName(FBlackboard::FKey KeyID) const;
 
 	/** @return class of value for given key */
-	TSubclassOf<UBlackboardKeyType> GetKeyType(FBlackboard::FKey KeyID) const;
+	AIMODULE_API TSubclassOf<UBlackboardKeyType> GetKeyType(FBlackboard::FKey KeyID) const;
 
 	/** @return number of defined keys, including parent chain */
-	int32 GetNumKeys() const;
+	AIMODULE_API int32 GetNumKeys() const;
 
 	FORCEINLINE FBlackboard::FKey GetFirstKeyID() const { return FirstKeyID; }
 
 	bool IsValidKey(FBlackboard::FKey KeyID) const { return KeyID != FBlackboard::InvalidKey && (int32)KeyID < GetNumKeys(); }
 
 	/** @return key data */
-	const FBlackboardEntry* GetKey(FBlackboard::FKey KeyID) const;
+	AIMODULE_API const FBlackboardEntry* GetKey(FBlackboard::FKey KeyID) const;
 
 	const TArray<FBlackboardEntry>& GetKeys() const { return Keys; }
 
-	virtual void PostInitProperties() override;
-	virtual void PostLoad() override;
+	AIMODULE_API virtual void PostInitProperties() override;
+	AIMODULE_API virtual void PostLoad() override;
 #if WITH_EDITOR
-	virtual void PostEditChangeProperty(struct FPropertyChangedEvent& PropertyChangedEvent) override;
+	AIMODULE_API virtual void PostEditChangeProperty(struct FPropertyChangedEvent& PropertyChangedEvent) override;
 #endif
-	void PropagateKeyChangesToDerivedBlackboardAssets();
+	AIMODULE_API void PropagateKeyChangesToDerivedBlackboardAssets();
 
 	/** @return true if blackboard keys are not conflicting with parent key chain */
-	bool IsValid() const;
+	AIMODULE_API bool IsValid() const;
 
 	/** updates persistent key with given name, depending on currently defined entries and parent chain
 	 *  @return key type of newly created entry for further setup
@@ -138,26 +138,26 @@ public:
 #if WITH_EDITOR
 	/** A delegate called on PostEditChangeProperty. Can be used in editor to react to asset changes. */
 	DECLARE_MULTICAST_DELEGATE_OneParam(FBlackboardDataChanged, UBlackboardData* /*Asset*/);
-	static FBlackboardDataChanged OnBlackboardDataChanged;
+	static AIMODULE_API FBlackboardDataChanged OnBlackboardDataChanged;
 #endif
 
 	/** delegate called for every loaded blackboard asset
 	 *  meant for adding game specific persistent keys */
-	static FKeyUpdate OnUpdateKeys;
+	static AIMODULE_API FKeyUpdate OnUpdateKeys;
 
 	/** updates parent key cache for editor */
-	void UpdateParentKeys();
+	AIMODULE_API void UpdateParentKeys();
 
 	/** forces update of FirstKeyID, which depends on parent chain */
-	void UpdateKeyIDs();
+	AIMODULE_API void UpdateKeyIDs();
 
-	void UpdateIfHasSynchronizedKeys();
+	AIMODULE_API void UpdateIfHasSynchronizedKeys();
 
 	/** fix entries with deprecated key types */
-	void UpdateDeprecatedKeys();
+	AIMODULE_API void UpdateDeprecatedKeys();
 
 	/** returns true if OtherAsset is somewhere up the parent chain of this asset. Node that it will return false if *this == OtherAsset */
-	bool IsChildOf(const UBlackboardData& OtherAsset) const;
+	AIMODULE_API bool IsChildOf(const UBlackboardData& OtherAsset) const;
 
 	/** returns true if OtherAsset is equal to *this, or is it's parent, or *this is OtherAsset's parent */
 	bool IsRelatedTo(const UBlackboardData& OtherAsset) const
@@ -178,5 +178,5 @@ protected:
 	FBlackboard::FKey FirstKeyID;
 
 	/** @return key ID from name */
-	FBlackboard::FKey InternalGetKeyID(const FName& KeyName, EKeyLookupMode LookupMode) const;
+	AIMODULE_API FBlackboard::FKey InternalGetKeyID(const FName& KeyName, EKeyLookupMode LookupMode) const;
 };
diff --git a/Engine/Source/Runtime/AIModule/Classes/BehaviorTree/Composites/BTComposite_Selector.h b/Engine/Source/Runtime/AIModule/Classes/BehaviorTree/Composites/BTComposite_Selector.h
--- a/Engine/Source/Runtime/AIModule/Classes/BehaviorTree/Composites/BTComposite_Selector.h
+++ b/Engine/Source/Runtime/AIModule/Classes/BehaviorTree/Composites/BTComposite_Selector.h
@@ -11,14 +11,14 @@
  * Selector Nodes execute their children from left to right, and will stop executing its children when one of their children succeeds.
  * If a Selector's child succeeds, the Selector succeeds. If all the Selector's children fail, the Selector fails.
  */
-UCLASS()
-class AIMODULE_API UBTComposite_Selector: public UBTCompositeNode
+UCLASS(MinimalAPI)
+class UBTComposite_Selector: public UBTCompositeNode
 {
 	GENERATED_UCLASS_BODY()
 
-	virtual int32 GetNextChildHandler(struct FBehaviorTreeSearchData& SearchData, int32 PrevChild, EBTNodeResult::Type LastResult) const override;
+	AIMODULE_API virtual int32 GetNextChildHandler(struct FBehaviorTreeSearchData& SearchData, int32 PrevChild, EBTNodeResult::Type LastResult) const override;
 
 #if WITH_EDITOR
-	virtual FName GetNodeIconName() const override;
+	AIMODULE_API virtual FName GetNodeIconName() const override;
 #endif
 };
diff --git a/Engine/Source/Runtime/AIModule/Classes/BehaviorTree/Composites/BTComposite_Sequence.h b/Engine/Source/Runtime/AIModule/Classes/BehaviorTree/Composites/BTComposite_Sequence.h
--- a/Engine/Source/Runtime/AIModule/Classes/BehaviorTree/Composites/BTComposite_Sequence.h
+++ b/Engine/Source/Runtime/AIModule/Classes/BehaviorTree/Composites/BTComposite_Sequence.h
@@ -11,15 +11,15 @@
  * Sequence Nodes execute their children from left to right, and will stop executing its children when one of their children fails.
  * If a child fails, then the Sequence fails. If all the Sequence's children succeed, then the Sequence succeeds.
  */
-UCLASS()
-class AIMODULE_API UBTComposite_Sequence : public UBTCompositeNode
+UCLASS(MinimalAPI)
+class UBTComposite_Sequence : public UBTCompositeNode
 {
 	GENERATED_UCLASS_BODY()
 
-	virtual int32 GetNextChildHandler(struct FBehaviorTreeSearchData& SearchData, int32 PrevChild, EBTNodeResult::Type LastResult) const override;
+	AIMODULE_API virtual int32 GetNextChildHandler(struct FBehaviorTreeSearchData& SearchData, int32 PrevChild, EBTNodeResult::Type LastResult) const override;
 
 #if WITH_EDITOR
-	virtual bool CanAbortLowerPriority() const override;
-	virtual FName GetNodeIconName() const override;
+	AIMODULE_API virtual bool CanAbortLowerPriority() const override;
+	AIMODULE_API virtual FName GetNodeIconName() const override;
 #endif
 };
diff --git a/Engine/Source/Runtime/AIModule/Classes/BehaviorTree/Composites/BTComposite_SimpleParallel.h b/Engine/Source/Runtime/AIModule/Classes/BehaviorTree/Composites/BTComposite_SimpleParallel.h
--- a/Engine/Source/Runtime/AIModule/Classes/BehaviorTree/Composites/BTComposite_SimpleParallel.h
+++ b/Engine/Source/Runtime/AIModule/Classes/BehaviorTree/Composites/BTComposite_SimpleParallel.h
@@ -50,8 +50,8 @@ struct FBTParallelMemory : public FBTCompositeMemory
  * Simple Parallel composite node.
  * Allows for running two children: one which must be a single task node (with optional decorators), and the other of which can be a complete subtree.
  */
-UCLASS(HideCategories=(Composite))
-class AIMODULE_API UBTComposite_SimpleParallel : public UBTCompositeNode
+UCLASS(HideCategories=(Composite), MinimalAPI)
+class UBTComposite_SimpleParallel : public UBTCompositeNode
 {
 	GENERATED_UCLASS_BODY()
 
@@ -60,23 +60,23 @@ class AIMODULE_API UBTComposite_SimpleParallel : public UBTCompositeNode
 	TEnumAsByte<EBTParallelMode::Type> FinishMode;
 
 	/** handle child updates */
-	virtual int32 GetNextChildHandler(FBehaviorTreeSearchData& SearchData, int32 PrevChild, EBTNodeResult::Type LastResult) const override;
+	AIMODULE_API virtual int32 GetNextChildHandler(FBehaviorTreeSearchData& SearchData, int32 PrevChild, EBTNodeResult::Type LastResult) const override;
 
-	virtual void NotifyChildExecution(UBehaviorTreeComponent& OwnerComp, uint8* NodeMemory, int32 ChildIdx, EBTNodeResult::Type& NodeResult) const override;
-	virtual void NotifyNodeDeactivation(FBehaviorTreeSearchData& SearchData, EBTNodeResult::Type& NodeResult) const override;
-	virtual bool CanNotifyDecoratorsOnDeactivation(FBehaviorTreeSearchData& SearchData, int32 ChildIdx, EBTNodeResult::Type& NodeResult) const override;
-	virtual bool CanPushSubtree(UBehaviorTreeComponent& OwnerComp, uint8* NodeMemory, int32 ChildIdx) const override;
-	virtual void SetChildOverride(FBehaviorTreeSearchData& SearchData, int8 Index) const override;
-	virtual uint16 GetInstanceMemorySize() const override;
-	virtual FString GetStaticDescription() const override;
-	virtual void DescribeRuntimeValues(const UBehaviorTreeComponent& OwnerComp, uint8* NodeMemory, EBTDescriptionVerbosity::Type Verbosity, TArray<FString>& Values) const override;
+	AIMODULE_API virtual void NotifyChildExecution(UBehaviorTreeComponent& OwnerComp, uint8* NodeMemory, int32 ChildIdx, EBTNodeResult::Type& NodeResult) const override;
+	AIMODULE_API virtual void NotifyNodeDeactivation(FBehaviorTreeSearchData& SearchData, EBTNodeResult::Type& NodeResult) const override;
+	AIMODULE_API virtual bool CanNotifyDecoratorsOnDeactivation(FBehaviorTreeSearchData& SearchData, int32 ChildIdx, EBTNodeResult::Type& NodeResult) const override;
+	AIMODULE_API virtual bool CanPushSubtree(UBehaviorTreeComponent& OwnerComp, uint8* NodeMemory, int32 ChildIdx) const override;
+	AIMODULE_API virtual void SetChildOverride(FBehaviorTreeSearchData& SearchData, int8 Index) const override;
+	AIMODULE_API virtual uint16 GetInstanceMemorySize() const override;
+	AIMODULE_API virtual FString GetStaticDescription() const override;
+	AIMODULE_API virtual void DescribeRuntimeValues(const UBehaviorTreeComponent& OwnerComp, uint8* NodeMemory, EBTDescriptionVerbosity::Type Verbosity, TArray<FString>& Values) const override;
 
 	/** helper for showing values of EBTParallelMode enum */
-	static FString DescribeFinishMode(EBTParallelMode::Type Mode);
+	static AIMODULE_API FString DescribeFinishMode(EBTParallelMode::Type Mode);
 
 #if WITH_EDITOR
-	virtual bool CanAbortLowerPriority() const override;
-	virtual bool CanAbortSelf() const override;
-	virtual FName GetNodeIconName() const override;
+	AIMODULE_API virtual bool CanAbortLowerPriority() const override;
+	AIMODULE_API virtual bool CanAbortSelf() const override;
+	AIMODULE_API virtual FName GetNodeIconName() const override;
 #endif // WITH_EDITOR
 };
diff --git a/Engine/Source/Runtime/AIModule/Classes/BehaviorTree/Decorators/BTDecorator_Blackboard.h b/Engine/Source/Runtime/AIModule/Classes/BehaviorTree/Decorators/BTDecorator_Blackboard.h
--- a/Engine/Source/Runtime/AIModule/Classes/BehaviorTree/Decorators/BTDecorator_Blackboard.h
+++ b/Engine/Source/Runtime/AIModule/Classes/BehaviorTree/Decorators/BTDecorator_Blackboard.h
@@ -32,15 +32,15 @@ namespace EBTBlackboardRestart
  * Blackboard decorator node.
  * A decorator node that bases its condition on a Blackboard key.
  */
-UCLASS(HideCategories=(Condition))
-class AIMODULE_API UBTDecorator_Blackboard : public UBTDecorator_BlackboardBase
+UCLASS(HideCategories=(Condition), MinimalAPI)
+class UBTDecorator_Blackboard : public UBTDecorator_BlackboardBase
 {
 	GENERATED_UCLASS_BODY()
 
-	virtual bool CalculateRawConditionValue(UBehaviorTreeComponent& OwnerComp, uint8* NodeMemory) const override;
-	virtual EBlackboardNotificationResult OnBlackboardKeyValueChange(const UBlackboardComponent& Blackboard, FBlackboard::FKey ChangedKeyID) override;
-	virtual void DescribeRuntimeValues(const UBehaviorTreeComponent& OwnerComp, uint8* NodeMemory, EBTDescriptionVerbosity::Type Verbosity, TArray<FString>& Values) const override;
-	virtual FString GetStaticDescription() const override;
+	AIMODULE_API virtual bool CalculateRawConditionValue(UBehaviorTreeComponent& OwnerComp, uint8* NodeMemory) const override;
+	AIMODULE_API virtual EBlackboardNotificationResult OnBlackboardKeyValueChange(const UBlackboardComponent& Blackboard, FBlackboard::FKey ChangedKeyID) override;
+	AIMODULE_API virtual void DescribeRuntimeValues(const UBehaviorTreeComponent& OwnerComp, uint8* NodeMemory, EBTDescriptionVerbosity::Type Verbosity, TArray<FString>& Values) const override;
+	AIMODULE_API virtual FString GetStaticDescription() const override;
 
 protected:
 
@@ -84,22 +84,22 @@ protected:
 #if WITH_EDITOR
 public:
 	/** describe decorator and cache it */
-	virtual void BuildDescription();
+	AIMODULE_API virtual void BuildDescription();
 protected:
-	virtual void PostEditChangeProperty(FPropertyChangedEvent& PropertyChangedEvent) override;
+	AIMODULE_API virtual void PostEditChangeProperty(FPropertyChangedEvent& PropertyChangedEvent) override;
 #endif
 
-	virtual void InitializeFromAsset(UBehaviorTree& Asset) override;
+	AIMODULE_API virtual void InitializeFromAsset(UBehaviorTree& Asset) override;
 
 	/**
 	 * decorator using enum based key requires to synchronize 'StringValue' used for storage with 'IntValue' used for runtime evaluation.
 	 * @note using custom serialization version was not possible since we require blackboard and used types (e.g. UserDefinedEnumerations)
 	 * to be loaded and keys to be resolved (i.e. InitializeFromAsset).
 	 */
-	void RefreshEnumBasedDecorator(const FBlackboardEntry& Entry);
+	AIMODULE_API void RefreshEnumBasedDecorator(const FBlackboardEntry& Entry);
 
 	/** take blackboard value and evaluate decorator's condition */
-	bool EvaluateOnBlackboard(const UBlackboardComponent& BlackboardComp) const;
+	AIMODULE_API bool EvaluateOnBlackboard(const UBlackboardComponent& BlackboardComp) const;
 
 	friend FBlackboardDecoratorDetails;
 };
diff --git a/Engine/Source/Runtime/AIModule/Classes/BehaviorTree/Decorators/BTDecorator_BlackboardBase.h b/Engine/Source/Runtime/AIModule/Classes/BehaviorTree/Decorators/BTDecorator_BlackboardBase.h
--- a/Engine/Source/Runtime/AIModule/Classes/BehaviorTree/Decorators/BTDecorator_BlackboardBase.h
+++ b/Engine/Source/Runtime/AIModule/Classes/BehaviorTree/Decorators/BTDecorator_BlackboardBase.h
@@ -11,24 +11,24 @@
 class UBehaviorTree;
 class UBlackboardComponent;
 
-UCLASS(Abstract)
-class AIMODULE_API UBTDecorator_BlackboardBase : public UBTDecorator
+UCLASS(Abstract, MinimalAPI)
+class UBTDecorator_BlackboardBase : public UBTDecorator
 {
 	GENERATED_UCLASS_BODY()
 
 	/** initialize any asset related data */
-	virtual void InitializeFromAsset(UBehaviorTree& Asset) override;
+	AIMODULE_API virtual void InitializeFromAsset(UBehaviorTree& Asset) override;
 
 	/** notify about change in blackboard keys */
-	virtual EBlackboardNotificationResult OnBlackboardKeyValueChange(const UBlackboardComponent& Blackboard, FBlackboard::FKey ChangedKeyID);
+	AIMODULE_API virtual EBlackboardNotificationResult OnBlackboardKeyValueChange(const UBlackboardComponent& Blackboard, FBlackboard::FKey ChangedKeyID);
 
 #if WITH_EDITOR
-	virtual FName GetNodeIconName() const override;
+	AIMODULE_API virtual FName GetNodeIconName() const override;
 #endif
 
 public:
 	/** get name of selected blackboard key */
-	FName GetSelectedBlackboardKey() const;
+	AIMODULE_API FName GetSelectedBlackboardKey() const;
 
 protected:
 
@@ -37,10 +37,10 @@ protected:
 	FBlackboardKeySelector BlackboardKey;
 
 	/** called when execution flow controller becomes active */
-	virtual void OnBecomeRelevant(UBehaviorTreeComponent& OwnerComp, uint8* NodeMemory) override;
+	AIMODULE_API virtual void OnBecomeRelevant(UBehaviorTreeComponent& OwnerComp, uint8* NodeMemory) override;
 
 	/** called when execution flow controller becomes inactive */
-	virtual void OnCeaseRelevant(UBehaviorTreeComponent& OwnerComp, uint8* NodeMemory) override;
+	AIMODULE_API virtual void OnCeaseRelevant(UBehaviorTreeComponent& OwnerComp, uint8* NodeMemory) override;
 };
 
 //////////////////////////////////////////////////////////////////////////
diff --git a/Engine/Source/Runtime/AIModule/Classes/BehaviorTree/Decorators/BTDecorator_BlueprintBase.h b/Engine/Source/Runtime/AIModule/Classes/BehaviorTree/Decorators/BTDecorator_BlueprintBase.h
--- a/Engine/Source/Runtime/AIModule/Classes/BehaviorTree/Decorators/BTDecorator_BlueprintBase.h
+++ b/Engine/Source/Runtime/AIModule/Classes/BehaviorTree/Decorators/BTDecorator_BlueprintBase.h
@@ -29,29 +29,29 @@ class UBlackboardComponent;
  *  - IsDecoratorObserverActive (true after ObserverActivated, until ObserverDeactivated)
  */
 
-UCLASS(Abstract, Blueprintable)
-class AIMODULE_API UBTDecorator_BlueprintBase : public UBTDecorator
+UCLASS(Abstract, Blueprintable, MinimalAPI)
+class UBTDecorator_BlueprintBase : public UBTDecorator
 {
 	GENERATED_BODY()
 
 public:
-	UBTDecorator_BlueprintBase(const FObjectInitializer& ObjectInitializer = FObjectInitializer::Get());
+	AIMODULE_API UBTDecorator_BlueprintBase(const FObjectInitializer& ObjectInitializer = FObjectInitializer::Get());
 
 	/** initialize data about blueprint defined properties */
-	void InitializeProperties();
+	AIMODULE_API void InitializeProperties();
 
 	/** setup node name */
-	virtual void PostInitProperties() override;
-	virtual void PostLoad() override;
+	AIMODULE_API virtual void PostInitProperties() override;
+	AIMODULE_API virtual void PostLoad() override;
 
 	/** notify about changes in blackboard */
-	EBlackboardNotificationResult OnBlackboardKeyValueChange(const UBlackboardComponent& Blackboard, FBlackboard::FKey ChangedKeyID);
+	AIMODULE_API EBlackboardNotificationResult OnBlackboardKeyValueChange(const UBlackboardComponent& Blackboard, FBlackboard::FKey ChangedKeyID);
 
-	virtual FString GetStaticDescription() const override;
-	virtual void DescribeRuntimeValues(const UBehaviorTreeComponent& OwnerComp, uint8* NodeMemory, EBTDescriptionVerbosity::Type Verbosity, TArray<FString>& Values) const override;
-	virtual bool CalculateRawConditionValue(UBehaviorTreeComponent& OwnerComp, uint8* NodeMemory) const override;
-	virtual void OnInstanceDestroyed(UBehaviorTreeComponent& OwnerComp) override;
-	virtual void InitializeFromAsset(UBehaviorTree& Asset) override;
+	AIMODULE_API virtual FString GetStaticDescription() const override;
+	AIMODULE_API virtual void DescribeRuntimeValues(const UBehaviorTreeComponent& OwnerComp, uint8* NodeMemory, EBTDescriptionVerbosity::Type Verbosity, TArray<FString>& Values) const override;
+	AIMODULE_API virtual bool CalculateRawConditionValue(UBehaviorTreeComponent& OwnerComp, uint8* NodeMemory) const override;
+	AIMODULE_API virtual void OnInstanceDestroyed(UBehaviorTreeComponent& OwnerComp) override;
+	AIMODULE_API virtual void InitializeFromAsset(UBehaviorTree& Asset) override;
 
 	enum class EAbortType : uint8
 	{
@@ -62,7 +62,7 @@ public:
 	};
 
 	/** return this decorator abort type in current circumstances */
-	EAbortType EvaluateAbortType(UBehaviorTreeComponent& OwnerComp) const;
+	AIMODULE_API EAbortType EvaluateAbortType(UBehaviorTreeComponent& OwnerComp) const;
 
 	UE_DEPRECATED(5.2, "Use GetShouldAbortType instead")
 	bool GetShouldAbort(UBehaviorTreeComponent& OwnerComp) const
@@ -70,11 +70,11 @@ public:
 		return  EvaluateAbortType(OwnerComp) != EAbortType::NoAbort;
 	}
 
-	virtual void SetOwner(AActor* ActorOwner) override;
+	AIMODULE_API virtual void SetOwner(AActor* ActorOwner) override;
 
 #if WITH_EDITOR
-	virtual FName GetNodeIconName() const override;
-	virtual bool UsesBlueprint() const override;
+	AIMODULE_API virtual FName GetNodeIconName() const override;
+	AIMODULE_API virtual bool UsesBlueprint() const override;
 #endif
 
 protected:
@@ -129,100 +129,100 @@ protected:
 	/** set if ReceiveConditionCheck is implemented by blueprint */
 	uint32 PerformConditionCheckImplementations : 2;
 
-	virtual void OnBecomeRelevant(UBehaviorTreeComponent& OwnerComp, uint8* NodeMemory) override;
-	virtual void OnCeaseRelevant(UBehaviorTreeComponent& OwnerComp, uint8* NodeMemory) override;
-	virtual void OnNodeActivation(FBehaviorTreeSearchData& SearchData) override;
-	virtual void OnNodeDeactivation(FBehaviorTreeSearchData& SearchData, EBTNodeResult::Type NodeResult) override;
-	virtual void TickNode(UBehaviorTreeComponent& OwnerComp, uint8* NodeMemory, float DeltaSeconds) override;
+	AIMODULE_API virtual void OnBecomeRelevant(UBehaviorTreeComponent& OwnerComp, uint8* NodeMemory) override;
+	AIMODULE_API virtual void OnCeaseRelevant(UBehaviorTreeComponent& OwnerComp, uint8* NodeMemory) override;
+	AIMODULE_API virtual void OnNodeActivation(FBehaviorTreeSearchData& SearchData) override;
+	AIMODULE_API virtual void OnNodeDeactivation(FBehaviorTreeSearchData& SearchData, EBTNodeResult::Type NodeResult) override;
+	AIMODULE_API virtual void TickNode(UBehaviorTreeComponent& OwnerComp, uint8* NodeMemory, float DeltaSeconds) override;
 
 	/** tick function
 	 *	@Note that if both generic and AI event versions are implemented only the more
 	 *	suitable one will be called, meaning the AI version if called for AI, generic one otherwise */
 	UFUNCTION(BlueprintImplementableEvent)
-	void ReceiveTick(AActor* OwnerActor, float DeltaSeconds);
+	AIMODULE_API void ReceiveTick(AActor* OwnerActor, float DeltaSeconds);
 
 	/** called on execution of underlying node 
 	 *	@Note that if both generic and AI event versions are implemented only the more
 	 *	suitable one will be called, meaning the AI version if called for AI, generic one otherwise */
 	UFUNCTION(BlueprintImplementableEvent)
-	void ReceiveExecutionStart(AActor* OwnerActor);
+	AIMODULE_API void ReceiveExecutionStart(AActor* OwnerActor);
 
 	/** called when execution of underlying node is finished 
 	 *	@Note that if both generic and AI event versions are implemented only the more
 	 *	suitable one will be called, meaning the AI version if called for AI, generic one otherwise */
 	UFUNCTION(BlueprintImplementableEvent)
-	void ReceiveExecutionFinish(AActor* OwnerActor, enum EBTNodeResult::Type NodeResult);
+	AIMODULE_API void ReceiveExecutionFinish(AActor* OwnerActor, enum EBTNodeResult::Type NodeResult);
 
 	/** called when observer is activated (flow controller) 
 	 *	@Note that if both generic and AI event versions are implemented only the more
 	 *	suitable one will be called, meaning the AI version if called for AI, generic one otherwise */
 	UFUNCTION(BlueprintImplementableEvent)
-	void ReceiveObserverActivated(AActor* OwnerActor);
+	AIMODULE_API void ReceiveObserverActivated(AActor* OwnerActor);
 
 	/** called when observer is deactivated (flow controller) 
 	 *	@Note that if both generic and AI event versions are implemented only the more
 	 *	suitable one will be called, meaning the AI version if called for AI, generic one otherwise */
 	UFUNCTION(BlueprintImplementableEvent)
-	void ReceiveObserverDeactivated(AActor* OwnerActor);
+	AIMODULE_API void ReceiveObserverDeactivated(AActor* OwnerActor);
 
 	/** called when testing if underlying node can be executed, must call FinishConditionCheck
 	 *	@Note that if both generic and AI event versions are implemented only the more
 	 *	suitable one will be called, meaning the AI version if called for AI, generic one otherwise */
 	UFUNCTION(BlueprintImplementableEvent)
-	bool PerformConditionCheck(AActor* OwnerActor);
+	AIMODULE_API bool PerformConditionCheck(AActor* OwnerActor);
 
 	/** Alternative AI version of ReceiveTick
 	 *	@see ReceiveTick for more details
 	 *	@Note that if both generic and AI event versions are implemented only the more
 	 *	suitable one will be called, meaning the AI version if called for AI, generic one otherwise */
 	UFUNCTION(BlueprintImplementableEvent, Category = AI)
-	void ReceiveTickAI(AAIController* OwnerController, APawn* ControlledPawn, float DeltaSeconds);
+	AIMODULE_API void ReceiveTickAI(AAIController* OwnerController, APawn* ControlledPawn, float DeltaSeconds);
 
 	/** Alternative AI version of ReceiveExecutionStart
 	 *	@see ReceiveExecutionStart for more details
 	 *	@Note that if both generic and AI event versions are implemented only the more
 	 *	suitable one will be called, meaning the AI version if called for AI, generic one otherwise */
 	UFUNCTION(BlueprintImplementableEvent, Category = AI)
-	void ReceiveExecutionStartAI(AAIController* OwnerController, APawn* ControlledPawn);
+	AIMODULE_API void ReceiveExecutionStartAI(AAIController* OwnerController, APawn* ControlledPawn);
 
 	/** Alternative AI version of ReceiveExecutionFinish
 	 *	@see ReceiveExecutionFinish for more details
 	 *	@Note that if both generic and AI event versions are implemented only the more
 	 *	suitable one will be called, meaning the AI version if called for AI, generic one otherwise */
 	UFUNCTION(BlueprintImplementableEvent, Category = AI)
-	void ReceiveExecutionFinishAI(AAIController* OwnerController, APawn* ControlledPawn, enum EBTNodeResult::Type NodeResult);
+	AIMODULE_API void ReceiveExecutionFinishAI(AAIController* OwnerController, APawn* ControlledPawn, enum EBTNodeResult::Type NodeResult);
 
 	/** Alternative AI version of ReceiveObserverActivated
 	 *	@see ReceiveObserverActivated for more details
 	 *	@Note that if both generic and AI event versions are implemented only the more
 	 *	suitable one will be called, meaning the AI version if called for AI, generic one otherwise */
 	UFUNCTION(BlueprintImplementableEvent, Category = AI)
-	void ReceiveObserverActivatedAI(AAIController* OwnerController, APawn* ControlledPawn);
+	AIMODULE_API void ReceiveObserverActivatedAI(AAIController* OwnerController, APawn* ControlledPawn);
 
 	/** Alternative AI version of ReceiveObserverDeactivated
 	 *	@see ReceiveObserverDeactivated for more details
 	 *	@Note that if both generic and AI event versions are implemented only the more
 	 *	suitable one will be called, meaning the AI version if called for AI, generic one otherwise */
 	UFUNCTION(BlueprintImplementableEvent, Category = AI)
-	void ReceiveObserverDeactivatedAI(AAIController* OwnerController, APawn* ControlledPawn);
+	AIMODULE_API void ReceiveObserverDeactivatedAI(AAIController* OwnerController, APawn* ControlledPawn);
 
 	/** Alternative AI version of ReceiveConditionCheck
 	 *	@see ReceiveConditionCheck for more details
 	 *	@Note that if both generic and AI event versions are implemented only the more
 	 *	suitable one will be called, meaning the AI version if called for AI, generic one otherwise */
 	UFUNCTION(BlueprintImplementableEvent, Category = AI)
-	bool PerformConditionCheckAI(AAIController* OwnerController, APawn* ControlledPawn);
+	AIMODULE_API bool PerformConditionCheckAI(AAIController* OwnerController, APawn* ControlledPawn);
 		
 	/** check if decorator is part of currently active branch */
 	UFUNCTION(BlueprintCallable, Category="AI|BehaviorTree")
-	bool IsDecoratorExecutionActive() const;
+	AIMODULE_API bool IsDecoratorExecutionActive() const;
 
 	/** check if decorator's observer is currently active */
 	UFUNCTION(BlueprintCallable, Category="AI|BehaviorTree")
-	bool IsDecoratorObserverActive() const;
+	AIMODULE_API bool IsDecoratorObserverActive() const;
 
 	FORCEINLINE bool GetNeedsTickForConditionChecking() const { return PerformConditionCheckImplementations != 0 && (bIsObservingBB == false || bCheckConditionOnlyBlackBoardChanges == false); }
 
 	/** Handle abort request calls from the specified abort type */
-	void RequestAbort(UBehaviorTreeComponent& OwnerComp, const EAbortType Type);
+	AIMODULE_API void RequestAbort(UBehaviorTreeComponent& OwnerComp, const EAbortType Type);
 };
diff --git a/Engine/Source/Runtime/AIModule/Classes/BehaviorTree/Decorators/BTDecorator_CheckGameplayTagsOnActor.h b/Engine/Source/Runtime/AIModule/Classes/BehaviorTree/Decorators/BTDecorator_CheckGameplayTagsOnActor.h
--- a/Engine/Source/Runtime/AIModule/Classes/BehaviorTree/Decorators/BTDecorator_CheckGameplayTagsOnActor.h
+++ b/Engine/Source/Runtime/AIModule/Classes/BehaviorTree/Decorators/BTDecorator_CheckGameplayTagsOnActor.h
@@ -15,14 +15,14 @@ class UBehaviorTree;
  * A decorator node that bases its condition on whether the specified Actor (in the blackboard) has a Gameplay Tag or
  * Tags specified.
  */
-UCLASS()
-class AIMODULE_API UBTDecorator_CheckGameplayTagsOnActor : public UBTDecorator
+UCLASS(MinimalAPI)
+class UBTDecorator_CheckGameplayTagsOnActor : public UBTDecorator
 {
 	GENERATED_UCLASS_BODY()
 
-	virtual bool CalculateRawConditionValue(UBehaviorTreeComponent& OwnerComp, uint8* NodeMemory) const override;
-	virtual void DescribeRuntimeValues(const UBehaviorTreeComponent& OwnerComp, uint8* NodeMemory, EBTDescriptionVerbosity::Type Verbosity, TArray<FString>& Values) const override;
-	virtual FString GetStaticDescription() const override;
+	AIMODULE_API virtual bool CalculateRawConditionValue(UBehaviorTreeComponent& OwnerComp, uint8* NodeMemory) const override;
+	AIMODULE_API virtual void DescribeRuntimeValues(const UBehaviorTreeComponent& OwnerComp, uint8* NodeMemory, EBTDescriptionVerbosity::Type Verbosity, TArray<FString>& Values) const override;
+	AIMODULE_API virtual FString GetStaticDescription() const override;
 
 protected:
 
@@ -42,10 +42,10 @@ protected:
 
 #if WITH_EDITOR
 	/** describe decorator and cache it */
-	virtual void BuildDescription();
+	AIMODULE_API virtual void BuildDescription();
 
-	virtual void PostEditChangeProperty(FPropertyChangedEvent& PropertyChangedEvent) override;
+	AIMODULE_API virtual void PostEditChangeProperty(FPropertyChangedEvent& PropertyChangedEvent) override;
 #endif
 
-	virtual void InitializeFromAsset(UBehaviorTree& Asset) override;
+	AIMODULE_API virtual void InitializeFromAsset(UBehaviorTree& Asset) override;
 };
diff --git a/Engine/Source/Runtime/AIModule/Classes/BehaviorTree/Decorators/BTDecorator_CompareBBEntries.h b/Engine/Source/Runtime/AIModule/Classes/BehaviorTree/Decorators/BTDecorator_CompareBBEntries.h
--- a/Engine/Source/Runtime/AIModule/Classes/BehaviorTree/Decorators/BTDecorator_CompareBBEntries.h
+++ b/Engine/Source/Runtime/AIModule/Classes/BehaviorTree/Decorators/BTDecorator_CompareBBEntries.h
@@ -26,8 +26,8 @@ namespace EBlackBoardEntryComparison
  * Blackboard comparison decorator node.
  * A decorator node that bases its condition on a comparison between two Blackboard keys.
  */
-UCLASS(HideCategories=(Condition))
-class AIMODULE_API UBTDecorator_CompareBBEntries : public UBTDecorator
+UCLASS(HideCategories=(Condition), MinimalAPI)
+class UBTDecorator_CompareBBEntries : public UBTDecorator
 {
 	GENERATED_UCLASS_BODY()
 
@@ -47,16 +47,16 @@ protected:
 
 public:
 
-	virtual bool CalculateRawConditionValue(UBehaviorTreeComponent& OwnerComp, uint8* NodeMemory) const override;
-	virtual FString GetStaticDescription() const override;
+	AIMODULE_API virtual bool CalculateRawConditionValue(UBehaviorTreeComponent& OwnerComp, uint8* NodeMemory) const override;
+	AIMODULE_API virtual FString GetStaticDescription() const override;
 
-	virtual void InitializeFromAsset(UBehaviorTree& Asset) override;
+	AIMODULE_API virtual void InitializeFromAsset(UBehaviorTree& Asset) override;
 
-	virtual EBlackboardNotificationResult OnBlackboardKeyValueChange(const UBlackboardComponent& Blackboard, FBlackboard::FKey ChangedKeyID);
-	virtual void OnBecomeRelevant(UBehaviorTreeComponent& OwnerComp, uint8* NodeMemory) override;
-	virtual void OnCeaseRelevant(UBehaviorTreeComponent& OwnerComp, uint8* NodeMemory) override;
+	AIMODULE_API virtual EBlackboardNotificationResult OnBlackboardKeyValueChange(const UBlackboardComponent& Blackboard, FBlackboard::FKey ChangedKeyID);
+	AIMODULE_API virtual void OnBecomeRelevant(UBehaviorTreeComponent& OwnerComp, uint8* NodeMemory) override;
+	AIMODULE_API virtual void OnCeaseRelevant(UBehaviorTreeComponent& OwnerComp, uint8* NodeMemory) override;
 
 #if WITH_EDITOR
-	virtual FName GetNodeIconName() const override;
+	AIMODULE_API virtual FName GetNodeIconName() const override;
 #endif // WITH_EDITOR
 };
diff --git a/Engine/Source/Runtime/AIModule/Classes/BehaviorTree/Decorators/BTDecorator_ConditionalLoop.h b/Engine/Source/Runtime/AIModule/Classes/BehaviorTree/Decorators/BTDecorator_ConditionalLoop.h
--- a/Engine/Source/Runtime/AIModule/Classes/BehaviorTree/Decorators/BTDecorator_ConditionalLoop.h
+++ b/Engine/Source/Runtime/AIModule/Classes/BehaviorTree/Decorators/BTDecorator_ConditionalLoop.h
@@ -19,18 +19,18 @@ struct FBTConditionalLoopDecoratorMemory
  * Conditional loop decorator node.
  * A decorator node that loops execution as long as condition is satisfied.
  */
-UCLASS(HideCategories=(FlowControl))
-class AIMODULE_API UBTDecorator_ConditionalLoop : public UBTDecorator_Blackboard
+UCLASS(HideCategories=(FlowControl), MinimalAPI)
+class UBTDecorator_ConditionalLoop : public UBTDecorator_Blackboard
 {
 	GENERATED_UCLASS_BODY()
 
-	virtual bool CalculateRawConditionValue(UBehaviorTreeComponent& OwnerComp, uint8* NodeMemory) const override;
-	virtual EBlackboardNotificationResult OnBlackboardKeyValueChange(const UBlackboardComponent& Blackboard, FBlackboard::FKey ChangedKeyID) override;
-	virtual void OnNodeDeactivation(FBehaviorTreeSearchData& SearchData, EBTNodeResult::Type NodeResult) override;
+	AIMODULE_API virtual bool CalculateRawConditionValue(UBehaviorTreeComponent& OwnerComp, uint8* NodeMemory) const override;
+	AIMODULE_API virtual EBlackboardNotificationResult OnBlackboardKeyValueChange(const UBlackboardComponent& Blackboard, FBlackboard::FKey ChangedKeyID) override;
+	AIMODULE_API virtual void OnNodeDeactivation(FBehaviorTreeSearchData& SearchData, EBTNodeResult::Type NodeResult) override;
 
-	virtual uint16 GetInstanceMemorySize() const override;
+	AIMODULE_API virtual uint16 GetInstanceMemorySize() const override;
 
 #if WITH_EDITOR
-	virtual FName GetNodeIconName() const override;
+	AIMODULE_API virtual FName GetNodeIconName() const override;
 #endif // WITH_EDITOR
 };
diff --git a/Engine/Source/Runtime/AIModule/Classes/BehaviorTree/Decorators/BTDecorator_ConeCheck.h b/Engine/Source/Runtime/AIModule/Classes/BehaviorTree/Decorators/BTDecorator_ConeCheck.h
--- a/Engine/Source/Runtime/AIModule/Classes/BehaviorTree/Decorators/BTDecorator_ConeCheck.h
+++ b/Engine/Source/Runtime/AIModule/Classes/BehaviorTree/Decorators/BTDecorator_ConeCheck.h
@@ -21,8 +21,8 @@ struct FBTConeCheckDecoratorMemory
  * Cone check decorator node.
  * A decorator node that bases its condition on a cone check, using Blackboard entries to form the parameters of the check.
  */
-UCLASS()
-class AIMODULE_API UBTDecorator_ConeCheck : public UBTDecorator
+UCLASS(MinimalAPI)
+class UBTDecorator_ConeCheck : public UBTDecorator
 {
 	GENERATED_UCLASS_BODY()
 
@@ -46,24 +46,24 @@ class AIMODULE_API UBTDecorator_ConeCheck : public UBTDecorator
 	
 	float ConeHalfAngleDot;
 
-	virtual void InitializeFromAsset(UBehaviorTree& Asset) override;
-	virtual uint16 GetInstanceMemorySize() const override;
-	virtual void DescribeRuntimeValues(const UBehaviorTreeComponent& OwnerComp, uint8* NodeMemory, EBTDescriptionVerbosity::Type Verbosity, TArray<FString>& Values) const override;
-	virtual FString GetStaticDescription() const override;
+	AIMODULE_API virtual void InitializeFromAsset(UBehaviorTree& Asset) override;
+	AIMODULE_API virtual uint16 GetInstanceMemorySize() const override;
+	AIMODULE_API virtual void DescribeRuntimeValues(const UBehaviorTreeComponent& OwnerComp, uint8* NodeMemory, EBTDescriptionVerbosity::Type Verbosity, TArray<FString>& Values) const override;
+	AIMODULE_API virtual FString GetStaticDescription() const override;
 
 #if WITH_EDITOR
-	virtual FName GetNodeIconName() const override;
+	AIMODULE_API virtual FName GetNodeIconName() const override;
 #endif // WITH_EDITOR
 
 protected:
 
-	virtual bool CalculateRawConditionValue(UBehaviorTreeComponent& OwnerComp, uint8* NodeMemory) const override;
-	virtual void OnBecomeRelevant(UBehaviorTreeComponent& OwnerComp, uint8* NodeMemory) override;
-	void OnBlackboardChange(const UBlackboardComponent& Blackboard, FBlackboard::FKey ChangedKeyID);
+	AIMODULE_API virtual bool CalculateRawConditionValue(UBehaviorTreeComponent& OwnerComp, uint8* NodeMemory) const override;
+	AIMODULE_API virtual void OnBecomeRelevant(UBehaviorTreeComponent& OwnerComp, uint8* NodeMemory) override;
+	AIMODULE_API void OnBlackboardChange(const UBlackboardComponent& Blackboard, FBlackboard::FKey ChangedKeyID);
 
-	virtual void TickNode(UBehaviorTreeComponent& OwnerComp, uint8* NodeMemory, float DeltaSeconds) override;
+	AIMODULE_API virtual void TickNode(UBehaviorTreeComponent& OwnerComp, uint8* NodeMemory, float DeltaSeconds) override;
 	
-	bool CalculateDirection(const UBlackboardComponent* BlackboardComp, const FBlackboardKeySelector& Origin, const FBlackboardKeySelector& End, FVector& Direction) const;
+	AIMODULE_API bool CalculateDirection(const UBlackboardComponent* BlackboardComp, const FBlackboardKeySelector& Origin, const FBlackboardKeySelector& End, FVector& Direction) const;
 
 private:
 	bool CalcConditionImpl(UBehaviorTreeComponent& OwnerComp, uint8* NodeMemory) const;
diff --git a/Engine/Source/Runtime/AIModule/Classes/BehaviorTree/Decorators/BTDecorator_Cooldown.h b/Engine/Source/Runtime/AIModule/Classes/BehaviorTree/Decorators/BTDecorator_Cooldown.h
--- a/Engine/Source/Runtime/AIModule/Classes/BehaviorTree/Decorators/BTDecorator_Cooldown.h
+++ b/Engine/Source/Runtime/AIModule/Classes/BehaviorTree/Decorators/BTDecorator_Cooldown.h
@@ -17,8 +17,8 @@ struct FBTCooldownDecoratorMemory
  * Cooldown decorator node.
  * A decorator node that bases its condition on whether a cooldown timer has expired.
  */
-UCLASS(HideCategories=(Condition))
-class AIMODULE_API UBTDecorator_Cooldown : public UBTDecorator
+UCLASS(HideCategories=(Condition), MinimalAPI)
+class UBTDecorator_Cooldown : public UBTDecorator
 {
 	GENERATED_UCLASS_BODY()
 
@@ -27,21 +27,21 @@ class AIMODULE_API UBTDecorator_Cooldown : public UBTDecorator
 	float CoolDownTime;
 
 	//~ Begin UObject Interface
-	virtual void PostLoad() override;
+	AIMODULE_API virtual void PostLoad() override;
 	//~ End UObject Interface
 
-	virtual bool CalculateRawConditionValue(UBehaviorTreeComponent& OwnerComp, uint8* NodeMemory) const override;
-	virtual void InitializeMemory(UBehaviorTreeComponent& OwnerComp, uint8* NodeMemory, EBTMemoryInit::Type InitType) const override;
-	virtual uint16 GetInstanceMemorySize() const override;
-	virtual void DescribeRuntimeValues(const UBehaviorTreeComponent& OwnerComp, uint8* NodeMemory, EBTDescriptionVerbosity::Type Verbosity, TArray<FString>& Values) const override;
-	virtual FString GetStaticDescription() const override;
+	AIMODULE_API virtual bool CalculateRawConditionValue(UBehaviorTreeComponent& OwnerComp, uint8* NodeMemory) const override;
+	AIMODULE_API virtual void InitializeMemory(UBehaviorTreeComponent& OwnerComp, uint8* NodeMemory, EBTMemoryInit::Type InitType) const override;
+	AIMODULE_API virtual uint16 GetInstanceMemorySize() const override;
+	AIMODULE_API virtual void DescribeRuntimeValues(const UBehaviorTreeComponent& OwnerComp, uint8* NodeMemory, EBTDescriptionVerbosity::Type Verbosity, TArray<FString>& Values) const override;
+	AIMODULE_API virtual FString GetStaticDescription() const override;
 
 #if WITH_EDITOR
-	virtual FName GetNodeIconName() const override;
+	AIMODULE_API virtual FName GetNodeIconName() const override;
 #endif // WITH_EDITOR
 
 protected:
 
-	virtual void OnNodeDeactivation(FBehaviorTreeSearchData& SearchData, EBTNodeResult::Type NodeResult) override;
-	virtual void TickNode(UBehaviorTreeComponent& OwnerComp, uint8* NodeMemory, float DeltaSeconds) override;
+	AIMODULE_API virtual void OnNodeDeactivation(FBehaviorTreeSearchData& SearchData, EBTNodeResult::Type NodeResult) override;
+	AIMODULE_API virtual void TickNode(UBehaviorTreeComponent& OwnerComp, uint8* NodeMemory, float DeltaSeconds) override;
 };
diff --git a/Engine/Source/Runtime/AIModule/Classes/BehaviorTree/Decorators/BTDecorator_DoesPathExist.h b/Engine/Source/Runtime/AIModule/Classes/BehaviorTree/Decorators/BTDecorator_DoesPathExist.h
--- a/Engine/Source/Runtime/AIModule/Classes/BehaviorTree/Decorators/BTDecorator_DoesPathExist.h
+++ b/Engine/Source/Runtime/AIModule/Classes/BehaviorTree/Decorators/BTDecorator_DoesPathExist.h
@@ -27,8 +27,8 @@ namespace EPathExistanceQueryType
  * Cooldown decorator node.
  * A decorator node that bases its condition on whether a path exists between two points from the Blackboard.
  */
-UCLASS()
-class AIMODULE_API UBTDecorator_DoesPathExist : public UBTDecorator
+UCLASS(MinimalAPI)
+class UBTDecorator_DoesPathExist : public UBTDecorator
 {
 	GENERATED_UCLASS_BODY()
 
@@ -55,11 +55,11 @@ public:
 	UPROPERTY(Category=Node, EditAnywhere)
 	TSubclassOf<UNavigationQueryFilter> FilterClass;
 
-	virtual bool CalculateRawConditionValue(UBehaviorTreeComponent& OwnerComp, uint8* NodeMemory) const override;
-	virtual FString GetStaticDescription() const override;
-	virtual void InitializeFromAsset(UBehaviorTree& Asset) override;
+	AIMODULE_API virtual bool CalculateRawConditionValue(UBehaviorTreeComponent& OwnerComp, uint8* NodeMemory) const override;
+	AIMODULE_API virtual FString GetStaticDescription() const override;
+	AIMODULE_API virtual void InitializeFromAsset(UBehaviorTree& Asset) override;
 
 #if WITH_EDITOR
-	virtual FName GetNodeIconName() const override;
+	AIMODULE_API virtual FName GetNodeIconName() const override;
 #endif // WITH_EDITOR
 };
diff --git a/Engine/Source/Runtime/AIModule/Classes/BehaviorTree/Decorators/BTDecorator_ForceSuccess.h b/Engine/Source/Runtime/AIModule/Classes/BehaviorTree/Decorators/BTDecorator_ForceSuccess.h
--- a/Engine/Source/Runtime/AIModule/Classes/BehaviorTree/Decorators/BTDecorator_ForceSuccess.h
+++ b/Engine/Source/Runtime/AIModule/Classes/BehaviorTree/Decorators/BTDecorator_ForceSuccess.h
@@ -16,16 +16,16 @@
  * - selector = would allow executing multiple nodes, turning it into a sequence...
  */
 
-UCLASS(HideCategories=(Condition))
-class AIMODULE_API UBTDecorator_ForceSuccess : public UBTDecorator
+UCLASS(HideCategories=(Condition), MinimalAPI)
+class UBTDecorator_ForceSuccess : public UBTDecorator
 {
 	GENERATED_UCLASS_BODY()
 
 protected:
 
-	virtual void OnNodeProcessed(struct FBehaviorTreeSearchData& SearchData, EBTNodeResult::Type& NodeResult) override;
+	AIMODULE_API virtual void OnNodeProcessed(struct FBehaviorTreeSearchData& SearchData, EBTNodeResult::Type& NodeResult) override;
 
 #if WITH_EDITOR
-	virtual FName GetNodeIconName() const override;
+	AIMODULE_API virtual FName GetNodeIconName() const override;
 #endif // WITH_EDITOR
 };
diff --git a/Engine/Source/Runtime/AIModule/Classes/BehaviorTree/Decorators/BTDecorator_IsAtLocation.h b/Engine/Source/Runtime/AIModule/Classes/BehaviorTree/Decorators/BTDecorator_IsAtLocation.h
--- a/Engine/Source/Runtime/AIModule/Classes/BehaviorTree/Decorators/BTDecorator_IsAtLocation.h
+++ b/Engine/Source/Runtime/AIModule/Classes/BehaviorTree/Decorators/BTDecorator_IsAtLocation.h
@@ -13,8 +13,8 @@
  * Is At Location decorator node.
  * A decorator node that checks if AI controlled pawn is at given location.
  */
-UCLASS()
-class AIMODULE_API UBTDecorator_IsAtLocation : public UBTDecorator_BlackboardBase
+UCLASS(MinimalAPI)
+class UBTDecorator_IsAtLocation : public UBTDecorator_BlackboardBase
 {
 	GENERATED_UCLASS_BODY()
 
@@ -40,14 +40,14 @@ class AIMODULE_API UBTDecorator_IsAtLocation : public UBTDecorator_BlackboardBas
 	UPROPERTY(EditAnywhere, Category = Condition)
 	uint32 bPathFindingBasedTest : 1;
 
-	virtual bool CalculateRawConditionValue(UBehaviorTreeComponent& OwnerComp, uint8* NodeMemory) const override;
+	AIMODULE_API virtual bool CalculateRawConditionValue(UBehaviorTreeComponent& OwnerComp, uint8* NodeMemory) const override;
 	
-	virtual FString GetStaticDescription() const override;
+	AIMODULE_API virtual FString GetStaticDescription() const override;
 
 #if WITH_EDITOR
-	virtual FName GetNodeIconName() const override;
+	AIMODULE_API virtual FName GetNodeIconName() const override;
 #endif // WITH_EDITOR
 
 protected:
-	FVector::FReal GetGeometricDistanceSquared(const FVector& A, const FVector& B) const;
+	AIMODULE_API FVector::FReal GetGeometricDistanceSquared(const FVector& A, const FVector& B) const;
 };
diff --git a/Engine/Source/Runtime/AIModule/Classes/BehaviorTree/Decorators/BTDecorator_IsBBEntryOfClass.h b/Engine/Source/Runtime/AIModule/Classes/BehaviorTree/Decorators/BTDecorator_IsBBEntryOfClass.h
--- a/Engine/Source/Runtime/AIModule/Classes/BehaviorTree/Decorators/BTDecorator_IsBBEntryOfClass.h
+++ b/Engine/Source/Runtime/AIModule/Classes/BehaviorTree/Decorators/BTDecorator_IsBBEntryOfClass.h
@@ -13,21 +13,21 @@
 
 class UBlackboardComponent;
 
-UCLASS(HideCategories=(Condition))
-class AIMODULE_API UBTDecorator_IsBBEntryOfClass : public UBTDecorator_BlackboardBase
+UCLASS(HideCategories=(Condition), MinimalAPI)
+class UBTDecorator_IsBBEntryOfClass : public UBTDecorator_BlackboardBase
 {
 	GENERATED_BODY()
 		
 public:
-	UBTDecorator_IsBBEntryOfClass(const FObjectInitializer& ObjectInitializer = FObjectInitializer::Get());
+	AIMODULE_API UBTDecorator_IsBBEntryOfClass(const FObjectInitializer& ObjectInitializer = FObjectInitializer::Get());
 
 protected:
 
 	UPROPERTY(Category = Blackboard, EditAnywhere)
 	TSubclassOf<UObject> TestClass;
 
-	virtual bool CalculateRawConditionValue(UBehaviorTreeComponent& OwnerComp, uint8* NodeMemory) const override;
-	virtual EBlackboardNotificationResult OnBlackboardKeyValueChange(const UBlackboardComponent& Blackboard, FBlackboard::FKey ChangedKeyID) override;
-	virtual void DescribeRuntimeValues(const UBehaviorTreeComponent& OwnerComp, uint8* NodeMemory, EBTDescriptionVerbosity::Type Verbosity, TArray<FString>& Values) const override;
-	virtual FString GetStaticDescription() const override;
+	AIMODULE_API virtual bool CalculateRawConditionValue(UBehaviorTreeComponent& OwnerComp, uint8* NodeMemory) const override;
+	AIMODULE_API virtual EBlackboardNotificationResult OnBlackboardKeyValueChange(const UBlackboardComponent& Blackboard, FBlackboard::FKey ChangedKeyID) override;
+	AIMODULE_API virtual void DescribeRuntimeValues(const UBehaviorTreeComponent& OwnerComp, uint8* NodeMemory, EBTDescriptionVerbosity::Type Verbosity, TArray<FString>& Values) const override;
+	AIMODULE_API virtual FString GetStaticDescription() const override;
 };
diff --git a/Engine/Source/Runtime/AIModule/Classes/BehaviorTree/Decorators/BTDecorator_KeepInCone.h b/Engine/Source/Runtime/AIModule/Classes/BehaviorTree/Decorators/BTDecorator_KeepInCone.h
--- a/Engine/Source/Runtime/AIModule/Classes/BehaviorTree/Decorators/BTDecorator_KeepInCone.h
+++ b/Engine/Source/Runtime/AIModule/Classes/BehaviorTree/Decorators/BTDecorator_KeepInCone.h
@@ -19,8 +19,8 @@ struct FBTKeepInConeDecoratorMemory
  * Cooldown decorator node.
  * A decorator node that bases its condition on whether the observed position is still inside a cone. The cone's direction is calculated when the node first becomes relevant.
  */
-UCLASS(HideCategories=(Condition))
-class AIMODULE_API UBTDecorator_KeepInCone : public UBTDecorator
+UCLASS(HideCategories=(Condition), MinimalAPI)
+class UBTDecorator_KeepInCone : public UBTDecorator
 {
 	GENERATED_UCLASS_BODY()
 
@@ -48,19 +48,19 @@ class AIMODULE_API UBTDecorator_KeepInCone : public UBTDecorator
 	
 	float ConeHalfAngleDot;
 
-	virtual void InitializeFromAsset(UBehaviorTree& Asset) override;
-	virtual uint16 GetInstanceMemorySize() const override;
-	virtual void DescribeRuntimeValues(const UBehaviorTreeComponent& OwnerComp, uint8* NodeMemory, EBTDescriptionVerbosity::Type Verbosity, TArray<FString>& Values) const override;
-	virtual FString GetStaticDescription() const override;
+	AIMODULE_API virtual void InitializeFromAsset(UBehaviorTree& Asset) override;
+	AIMODULE_API virtual uint16 GetInstanceMemorySize() const override;
+	AIMODULE_API virtual void DescribeRuntimeValues(const UBehaviorTreeComponent& OwnerComp, uint8* NodeMemory, EBTDescriptionVerbosity::Type Verbosity, TArray<FString>& Values) const override;
+	AIMODULE_API virtual FString GetStaticDescription() const override;
 
 #if WITH_EDITOR
-	virtual FName GetNodeIconName() const override;
+	AIMODULE_API virtual FName GetNodeIconName() const override;
 #endif // WITH_EDITOR
 
 protected:
 
-	virtual void OnBecomeRelevant(UBehaviorTreeComponent& OwnerComp, uint8* NodeMemory) override;
-	virtual void TickNode(UBehaviorTreeComponent& OwnerComp, uint8* NodeMemory, float DeltaSeconds) override;
+	AIMODULE_API virtual void OnBecomeRelevant(UBehaviorTreeComponent& OwnerComp, uint8* NodeMemory) override;
+	AIMODULE_API virtual void TickNode(UBehaviorTreeComponent& OwnerComp, uint8* NodeMemory, float DeltaSeconds) override;
 
-	bool CalculateCurrentDirection(const UBehaviorTreeComponent& OwnerComp, FVector& Direction) const;
+	AIMODULE_API bool CalculateCurrentDirection(const UBehaviorTreeComponent& OwnerComp, FVector& Direction) const;
 };
diff --git a/Engine/Source/Runtime/AIModule/Classes/BehaviorTree/Decorators/BTDecorator_Loop.h b/Engine/Source/Runtime/AIModule/Classes/BehaviorTree/Decorators/BTDecorator_Loop.h
--- a/Engine/Source/Runtime/AIModule/Classes/BehaviorTree/Decorators/BTDecorator_Loop.h
+++ b/Engine/Source/Runtime/AIModule/Classes/BehaviorTree/Decorators/BTDecorator_Loop.h
@@ -17,8 +17,8 @@ struct FBTLoopDecoratorMemory
  * Loop decorator node.
  * A decorator node that bases its condition on whether its loop counter has been exceeded.
  */
-UCLASS(HideCategories=(Condition))
-class AIMODULE_API UBTDecorator_Loop : public UBTDecorator
+UCLASS(HideCategories=(Condition), MinimalAPI)
+class UBTDecorator_Loop : public UBTDecorator
 {
 	GENERATED_UCLASS_BODY()
 
@@ -34,15 +34,15 @@ class AIMODULE_API UBTDecorator_Loop : public UBTDecorator
 	UPROPERTY(Category = Decorator, EditAnywhere, meta = (EditCondition = "bInfiniteLoop"))
 	float InfiniteLoopTimeoutTime;
 
-	virtual uint16 GetInstanceMemorySize() const override;
-	virtual void DescribeRuntimeValues(const UBehaviorTreeComponent& OwnerComp, uint8* NodeMemory, EBTDescriptionVerbosity::Type Verbosity, TArray<FString>& Values) const override;
-	virtual FString GetStaticDescription() const override;
+	AIMODULE_API virtual uint16 GetInstanceMemorySize() const override;
+	AIMODULE_API virtual void DescribeRuntimeValues(const UBehaviorTreeComponent& OwnerComp, uint8* NodeMemory, EBTDescriptionVerbosity::Type Verbosity, TArray<FString>& Values) const override;
+	AIMODULE_API virtual FString GetStaticDescription() const override;
 
 #if WITH_EDITOR
-	virtual FName GetNodeIconName() const override;
+	AIMODULE_API virtual FName GetNodeIconName() const override;
 #endif // WITH_EDITOR
 
 protected:
 
-	virtual void OnNodeActivation(FBehaviorTreeSearchData& SearchData) override;
+	AIMODULE_API virtual void OnNodeActivation(FBehaviorTreeSearchData& SearchData) override;
 };
diff --git a/Engine/Source/Runtime/AIModule/Classes/BehaviorTree/Decorators/BTDecorator_ReachedMoveGoal.h b/Engine/Source/Runtime/AIModule/Classes/BehaviorTree/Decorators/BTDecorator_ReachedMoveGoal.h
--- a/Engine/Source/Runtime/AIModule/Classes/BehaviorTree/Decorators/BTDecorator_ReachedMoveGoal.h
+++ b/Engine/Source/Runtime/AIModule/Classes/BehaviorTree/Decorators/BTDecorator_ReachedMoveGoal.h
@@ -11,14 +11,14 @@
  * Reached Move Goal decorator node.
  * A decorator node that bases its condition on whether the AI controller's path following component returns that it has reached its goal.
  */
-UCLASS(meta = (DeprecatedNode, DeprecationMessage = "Please use IsAtLocation decorator instead."))
-class AIMODULE_API UBTDecorator_ReachedMoveGoal : public UBTDecorator
+UCLASS(meta = (DeprecatedNode, DeprecationMessage = "Please use IsAtLocation decorator instead."), MinimalAPI)
+class UBTDecorator_ReachedMoveGoal : public UBTDecorator
 {
 	GENERATED_UCLASS_BODY()
 
-	virtual bool CalculateRawConditionValue(UBehaviorTreeComponent& OwnerComp, uint8* NodeMemory) const override;
+	AIMODULE_API virtual bool CalculateRawConditionValue(UBehaviorTreeComponent& OwnerComp, uint8* NodeMemory) const override;
 
 #if WITH_EDITOR
-	virtual FName GetNodeIconName() const override;
+	AIMODULE_API virtual FName GetNodeIconName() const override;
 #endif // WITH_EDITOR
 };
diff --git a/Engine/Source/Runtime/AIModule/Classes/BehaviorTree/Decorators/BTDecorator_SetTagCooldown.h b/Engine/Source/Runtime/AIModule/Classes/BehaviorTree/Decorators/BTDecorator_SetTagCooldown.h
--- a/Engine/Source/Runtime/AIModule/Classes/BehaviorTree/Decorators/BTDecorator_SetTagCooldown.h
+++ b/Engine/Source/Runtime/AIModule/Classes/BehaviorTree/Decorators/BTDecorator_SetTagCooldown.h
@@ -12,8 +12,8 @@
  * Set tag cooldown decorator node.
  * A decorator node that sets a gameplay tag cooldown.
  */
-UCLASS(HideCategories=(Condition))
-class AIMODULE_API UBTDecorator_SetTagCooldown : public UBTDecorator
+UCLASS(HideCategories=(Condition), MinimalAPI)
+class UBTDecorator_SetTagCooldown : public UBTDecorator
 {
 	GENERATED_UCLASS_BODY()
 
@@ -29,12 +29,12 @@ class AIMODULE_API UBTDecorator_SetTagCooldown : public UBTDecorator
 	UPROPERTY(Category = Decorator, EditAnywhere)
 	bool bAddToExistingDuration;
 
-	virtual FString GetStaticDescription() const override;
+	AIMODULE_API virtual FString GetStaticDescription() const override;
 
 #if WITH_EDITOR
-	virtual FName GetNodeIconName() const override;
+	AIMODULE_API virtual FName GetNodeIconName() const override;
 #endif // WITH_EDITOR
 
 protected:
-	virtual void OnNodeDeactivation(FBehaviorTreeSearchData& SearchData, EBTNodeResult::Type NodeResult) override;
+	AIMODULE_API virtual void OnNodeDeactivation(FBehaviorTreeSearchData& SearchData, EBTNodeResult::Type NodeResult) override;
 };
diff --git a/Engine/Source/Runtime/AIModule/Classes/BehaviorTree/Decorators/BTDecorator_TagCooldown.h b/Engine/Source/Runtime/AIModule/Classes/BehaviorTree/Decorators/BTDecorator_TagCooldown.h
--- a/Engine/Source/Runtime/AIModule/Classes/BehaviorTree/Decorators/BTDecorator_TagCooldown.h
+++ b/Engine/Source/Runtime/AIModule/Classes/BehaviorTree/Decorators/BTDecorator_TagCooldown.h
@@ -17,8 +17,8 @@ struct FBTTagCooldownDecoratorMemory
  * Cooldown decorator node.
  * A decorator node that bases its condition on whether a cooldown timer based on a gameplay tag has expired.
  */
-UCLASS(HideCategories=(Condition))
-class AIMODULE_API UBTDecorator_TagCooldown : public UBTDecorator
+UCLASS(HideCategories=(Condition), MinimalAPI)
+class UBTDecorator_TagCooldown : public UBTDecorator
 {
 	GENERATED_UCLASS_BODY()
 
@@ -39,23 +39,23 @@ class AIMODULE_API UBTDecorator_TagCooldown : public UBTDecorator
 	bool bActivatesCooldown;
 
 	//~ Begin UObject Interface
-	virtual void PostLoad() override;
+	AIMODULE_API virtual void PostLoad() override;
 	//~ End UObject Interface
 
-	virtual bool CalculateRawConditionValue(UBehaviorTreeComponent& OwnerComp, uint8* NodeMemory) const override;
-	virtual void InitializeMemory(UBehaviorTreeComponent& OwnerComp, uint8* NodeMemory, EBTMemoryInit::Type InitType) const override;
-	virtual uint16 GetInstanceMemorySize() const override;
-	virtual void DescribeRuntimeValues(const UBehaviorTreeComponent& OwnerComp, uint8* NodeMemory, EBTDescriptionVerbosity::Type Verbosity, TArray<FString>& Values) const override;
-	virtual FString GetStaticDescription() const override;
+	AIMODULE_API virtual bool CalculateRawConditionValue(UBehaviorTreeComponent& OwnerComp, uint8* NodeMemory) const override;
+	AIMODULE_API virtual void InitializeMemory(UBehaviorTreeComponent& OwnerComp, uint8* NodeMemory, EBTMemoryInit::Type InitType) const override;
+	AIMODULE_API virtual uint16 GetInstanceMemorySize() const override;
+	AIMODULE_API virtual void DescribeRuntimeValues(const UBehaviorTreeComponent& OwnerComp, uint8* NodeMemory, EBTDescriptionVerbosity::Type Verbosity, TArray<FString>& Values) const override;
+	AIMODULE_API virtual FString GetStaticDescription() const override;
 
 #if WITH_EDITOR
-	virtual FName GetNodeIconName() const override;
+	AIMODULE_API virtual FName GetNodeIconName() const override;
 #endif // WITH_EDITOR
 
 protected:
 
-	virtual void OnNodeDeactivation(FBehaviorTreeSearchData& SearchData, EBTNodeResult::Type NodeResult) override;
-	virtual void TickNode(UBehaviorTreeComponent& OwnerComp, uint8* NodeMemory, float DeltaSeconds) override;
+	AIMODULE_API virtual void OnNodeDeactivation(FBehaviorTreeSearchData& SearchData, EBTNodeResult::Type NodeResult) override;
+	AIMODULE_API virtual void TickNode(UBehaviorTreeComponent& OwnerComp, uint8* NodeMemory, float DeltaSeconds) override;
 
 private:
 	bool HasCooldownFinished(const UBehaviorTreeComponent& OwnerComp) const;
diff --git a/Engine/Source/Runtime/AIModule/Classes/BehaviorTree/Decorators/BTDecorator_TimeLimit.h b/Engine/Source/Runtime/AIModule/Classes/BehaviorTree/Decorators/BTDecorator_TimeLimit.h
--- a/Engine/Source/Runtime/AIModule/Classes/BehaviorTree/Decorators/BTDecorator_TimeLimit.h
+++ b/Engine/Source/Runtime/AIModule/Classes/BehaviorTree/Decorators/BTDecorator_TimeLimit.h
@@ -10,8 +10,8 @@
  * Time Limit decorator node.
  * A decorator node that bases its condition on whether a timer has exceeded a specified value. The timer is reset each time the node becomes relevant.
  */
-UCLASS(HideCategories=(Condition))
-class AIMODULE_API UBTDecorator_TimeLimit : public UBTDecorator
+UCLASS(HideCategories=(Condition), MinimalAPI)
+class UBTDecorator_TimeLimit : public UBTDecorator
 {
 	GENERATED_UCLASS_BODY()
 		
@@ -19,14 +19,14 @@ class AIMODULE_API UBTDecorator_TimeLimit : public UBTDecorator
 	UPROPERTY(Category=Decorator, EditAnywhere)
 	float TimeLimit;
 
-	virtual void DescribeRuntimeValues(const UBehaviorTreeComponent& OwnerComp, uint8* NodeMemory, EBTDescriptionVerbosity::Type Verbosity, TArray<FString>& Values) const override;
-	virtual FString GetStaticDescription() const override;
+	AIMODULE_API virtual void DescribeRuntimeValues(const UBehaviorTreeComponent& OwnerComp, uint8* NodeMemory, EBTDescriptionVerbosity::Type Verbosity, TArray<FString>& Values) const override;
+	AIMODULE_API virtual FString GetStaticDescription() const override;
 
 #if WITH_EDITOR
-	virtual FName GetNodeIconName() const override;
+	AIMODULE_API virtual FName GetNodeIconName() const override;
 #endif // WITH_EDITOR
 
 protected:
-	virtual void OnNodeActivation(FBehaviorTreeSearchData& SearchData) override;
-	virtual void TickNode(UBehaviorTreeComponent& OwnerComp, uint8* NodeMemory, float DeltaSeconds) override;
+	AIMODULE_API virtual void OnNodeActivation(FBehaviorTreeSearchData& SearchData) override;
+	AIMODULE_API virtual void TickNode(UBehaviorTreeComponent& OwnerComp, uint8* NodeMemory, float DeltaSeconds) override;
 };
diff --git a/Engine/Source/Runtime/AIModule/Classes/BehaviorTree/Services/BTService_BlackboardBase.h b/Engine/Source/Runtime/AIModule/Classes/BehaviorTree/Services/BTService_BlackboardBase.h
--- a/Engine/Source/Runtime/AIModule/Classes/BehaviorTree/Services/BTService_BlackboardBase.h
+++ b/Engine/Source/Runtime/AIModule/Classes/BehaviorTree/Services/BTService_BlackboardBase.h
@@ -8,16 +8,16 @@
 
 class UBehaviorTree;
 
-UCLASS(Abstract)
-class AIMODULE_API UBTService_BlackboardBase : public UBTService
+UCLASS(Abstract, MinimalAPI)
+class UBTService_BlackboardBase : public UBTService
 {
 	GENERATED_UCLASS_BODY()
 
 	/** initialize any asset related data */
-	virtual void InitializeFromAsset(UBehaviorTree& Asset) override;
+	AIMODULE_API virtual void InitializeFromAsset(UBehaviorTree& Asset) override;
 
 	/** get name of selected blackboard key */
-	FName GetSelectedBlackboardKey() const;
+	AIMODULE_API FName GetSelectedBlackboardKey() const;
 
 protected:
 
diff --git a/Engine/Source/Runtime/AIModule/Classes/BehaviorTree/Services/BTService_BlueprintBase.h b/Engine/Source/Runtime/AIModule/Classes/BehaviorTree/Services/BTService_BlueprintBase.h
--- a/Engine/Source/Runtime/AIModule/Classes/BehaviorTree/Services/BTService_BlueprintBase.h
+++ b/Engine/Source/Runtime/AIModule/Classes/BehaviorTree/Services/BTService_BlueprintBase.h
@@ -20,19 +20,19 @@ class UBehaviorTree;
  *  Please use them safely (unregister at abort) and call IsServiceActive() when in doubt.
  */
 
-UCLASS(Abstract, Blueprintable)
-class AIMODULE_API UBTService_BlueprintBase : public UBTService
+UCLASS(Abstract, Blueprintable, MinimalAPI)
+class UBTService_BlueprintBase : public UBTService
 {
 	GENERATED_UCLASS_BODY()
 
-	virtual void DescribeRuntimeValues(const UBehaviorTreeComponent& OwnerComp, uint8* NodeMemory, EBTDescriptionVerbosity::Type Verbosity, TArray<FString>& Values) const override;
-	virtual void OnInstanceDestroyed(UBehaviorTreeComponent& OwnerComp) override;
-	virtual void InitializeFromAsset(UBehaviorTree& Asset) override;
+	AIMODULE_API virtual void DescribeRuntimeValues(const UBehaviorTreeComponent& OwnerComp, uint8* NodeMemory, EBTDescriptionVerbosity::Type Verbosity, TArray<FString>& Values) const override;
+	AIMODULE_API virtual void OnInstanceDestroyed(UBehaviorTreeComponent& OwnerComp) override;
+	AIMODULE_API virtual void InitializeFromAsset(UBehaviorTree& Asset) override;
 
-	virtual void SetOwner(AActor* ActorOwner) override;
+	AIMODULE_API virtual void SetOwner(AActor* ActorOwner) override;
 
 #if WITH_EDITOR
-	virtual bool UsesBlueprint() const override;
+	AIMODULE_API virtual bool UsesBlueprint() const override;
 #endif
 
 protected:
@@ -45,7 +45,7 @@ protected:
 	TObjectPtr<AActor> ActorOwner;
 
 	// Gets the description for our service
-	virtual FString GetStaticServiceDescription() const override;
+	AIMODULE_API virtual FString GetStaticServiceDescription() const override;
 
 	/** properties with runtime values, stored only in class default object */
 	TArray<FProperty*> PropertyData;
@@ -75,64 +75,64 @@ protected:
 	/** set if ReceiveSearchStart is implemented by blueprint */
 	uint32 ReceiveSearchStartImplementations : 2;
 
-	virtual void OnBecomeRelevant(UBehaviorTreeComponent& OwnerComp, uint8* NodeMemory) override;
-	virtual void OnCeaseRelevant(UBehaviorTreeComponent& OwnerComp, uint8* NodeMemory) override;
-	virtual void TickNode(UBehaviorTreeComponent& OwnerComp, uint8* NodeMemory, float DeltaSeconds) override;
-	virtual void OnSearchStart(FBehaviorTreeSearchData& SearchData) override;
+	AIMODULE_API virtual void OnBecomeRelevant(UBehaviorTreeComponent& OwnerComp, uint8* NodeMemory) override;
+	AIMODULE_API virtual void OnCeaseRelevant(UBehaviorTreeComponent& OwnerComp, uint8* NodeMemory) override;
+	AIMODULE_API virtual void TickNode(UBehaviorTreeComponent& OwnerComp, uint8* NodeMemory, float DeltaSeconds) override;
+	AIMODULE_API virtual void OnSearchStart(FBehaviorTreeSearchData& SearchData) override;
 
 	/** tick function
 	 *	@Note that if both generic and AI event versions are implemented only the more
 	 *	suitable one will be called, meaning the AI version if called for AI, generic one otherwise */
 	UFUNCTION(BlueprintImplementableEvent)
-	void ReceiveTick(AActor* OwnerActor, float DeltaSeconds);
+	AIMODULE_API void ReceiveTick(AActor* OwnerActor, float DeltaSeconds);
 
 	/** task search enters branch of tree
 	 *	@Note that if both generic and AI event versions are implemented only the more
 	 *	suitable one will be called, meaning the AI version if called for AI, generic one otherwise */
 	UFUNCTION(BlueprintImplementableEvent)
-	void ReceiveSearchStart(AActor* OwnerActor);
+	AIMODULE_API void ReceiveSearchStart(AActor* OwnerActor);
 
 	/** service became active
 	 *	@Note that if both generic and AI event versions are implemented only the more
 	 *	suitable one will be called, meaning the AI version if called for AI, generic one otherwise */
 	UFUNCTION(BlueprintImplementableEvent)
-	void ReceiveActivation(AActor* OwnerActor);
+	AIMODULE_API void ReceiveActivation(AActor* OwnerActor);
 
 	/** service became inactive
 	 *	@Note that if both generic and AI event versions are implemented only the more
 	 *	suitable one will be called, meaning the AI version if called for AI, generic one otherwise */
 	UFUNCTION(BlueprintImplementableEvent)
-	void ReceiveDeactivation(AActor* OwnerActor);
+	AIMODULE_API void ReceiveDeactivation(AActor* OwnerActor);
 
 	/** Alternative AI version of ReceiveTick function.
 	 *	@see ReceiveTick for more details
 	 *	@Note that if both generic and AI event versions are implemented only the more
 	 *	suitable one will be called, meaning the AI version if called for AI, generic one otherwise */
 	UFUNCTION(BlueprintImplementableEvent, Category = AI)
-	void ReceiveTickAI(AAIController* OwnerController, APawn* ControlledPawn, float DeltaSeconds);
+	AIMODULE_API void ReceiveTickAI(AAIController* OwnerController, APawn* ControlledPawn, float DeltaSeconds);
 
 	/** Alternative AI version of ReceiveSearchStart function.
 	 *	@see ReceiveSearchStart for more details
 	 *	@Note that if both generic and AI event versions are implemented only the more
 	 *	suitable one will be called, meaning the AI version if called for AI, generic one otherwise */
 	UFUNCTION(BlueprintImplementableEvent, Category = AI)
-	void ReceiveSearchStartAI(AAIController* OwnerController, APawn* ControlledPawn);
+	AIMODULE_API void ReceiveSearchStartAI(AAIController* OwnerController, APawn* ControlledPawn);
 
 	/** Alternative AI version of ReceiveActivation function.
 	 *	@see ReceiveActivation for more details
 	 *	@Note that if both generic and AI event versions are implemented only the more
 	 *	suitable one will be called, meaning the AI version if called for AI, generic one otherwise */
 	UFUNCTION(BlueprintImplementableEvent, Category = AI)
-	void ReceiveActivationAI(AAIController* OwnerController, APawn* ControlledPawn);
+	AIMODULE_API void ReceiveActivationAI(AAIController* OwnerController, APawn* ControlledPawn);
 
 	/** Alternative AI version of ReceiveDeactivation function.
 	 *	@see ReceiveDeactivation for more details
 	 *	@Note that if both generic and AI event versions are implemented only the more
 	 *	suitable one will be called, meaning the AI version if called for AI, generic one otherwise */
 	UFUNCTION(BlueprintImplementableEvent, Category = AI)
-	void ReceiveDeactivationAI(AAIController* OwnerController, APawn* ControlledPawn);
+	AIMODULE_API void ReceiveDeactivationAI(AAIController* OwnerController, APawn* ControlledPawn);
 
 	/** check if service is currently being active */
 	UFUNCTION(BlueprintCallable, Category="AI|BehaviorTree")
-	bool IsServiceActive() const;
+	AIMODULE_API bool IsServiceActive() const;
 };
diff --git a/Engine/Source/Runtime/AIModule/Classes/BehaviorTree/Services/BTService_DefaultFocus.h b/Engine/Source/Runtime/AIModule/Classes/BehaviorTree/Services/BTService_DefaultFocus.h
--- a/Engine/Source/Runtime/AIModule/Classes/BehaviorTree/Services/BTService_DefaultFocus.h
+++ b/Engine/Source/Runtime/AIModule/Classes/BehaviorTree/Services/BTService_DefaultFocus.h
@@ -30,8 +30,8 @@ struct FBTFocusMemory
  * Default Focus service node.
  * A service node that automatically sets the AI controller's focus when it becomes active.
  */
-UCLASS(hidecategories=(Service))
-class AIMODULE_API UBTService_DefaultFocus : public UBTService_BlackboardBase
+UCLASS(hidecategories=(Service), MinimalAPI)
+class UBTService_DefaultFocus : public UBTService_BlackboardBase
 {
 	GENERATED_BODY()
 
@@ -40,17 +40,17 @@ protected:
 	UPROPERTY()
 	uint8 FocusPriority;
 
-	UBTService_DefaultFocus(const FObjectInitializer& ObjectInitializer = FObjectInitializer::Get());
+	AIMODULE_API UBTService_DefaultFocus(const FObjectInitializer& ObjectInitializer = FObjectInitializer::Get());
 
 	virtual uint16 GetInstanceMemorySize() const override { return sizeof(FBTFocusMemory); }
-	virtual void OnBecomeRelevant(UBehaviorTreeComponent& OwnerComp, uint8* NodeMemory) override;
-	virtual void OnCeaseRelevant(UBehaviorTreeComponent& OwnerComp, uint8* NodeMemory) override;
+	AIMODULE_API virtual void OnBecomeRelevant(UBehaviorTreeComponent& OwnerComp, uint8* NodeMemory) override;
+	AIMODULE_API virtual void OnCeaseRelevant(UBehaviorTreeComponent& OwnerComp, uint8* NodeMemory) override;
 
-	virtual FString GetStaticDescription() const override;
+	AIMODULE_API virtual FString GetStaticDescription() const override;
 
-	EBlackboardNotificationResult OnBlackboardKeyValueChange(const UBlackboardComponent& Blackboard, FBlackboard::FKey ChangedKeyID);
+	AIMODULE_API EBlackboardNotificationResult OnBlackboardKeyValueChange(const UBlackboardComponent& Blackboard, FBlackboard::FKey ChangedKeyID);
 
 #if WITH_EDITOR
-	virtual FName GetNodeIconName() const override;
+	AIMODULE_API virtual FName GetNodeIconName() const override;
 #endif // WITH_EDITOR
 };
diff --git a/Engine/Source/Runtime/AIModule/Classes/BehaviorTree/Services/BTService_RunEQS.h b/Engine/Source/Runtime/AIModule/Classes/BehaviorTree/Services/BTService_RunEQS.h
--- a/Engine/Source/Runtime/AIModule/Classes/BehaviorTree/Services/BTService_RunEQS.h
+++ b/Engine/Source/Runtime/AIModule/Classes/BehaviorTree/Services/BTService_RunEQS.h
@@ -17,8 +17,8 @@ struct FBTEQSServiceMemory
 	int32 RequestID;
 };
 
-UCLASS()
-class AIMODULE_API UBTService_RunEQS : public UBTService_BlackboardBase
+UCLASS(MinimalAPI)
+class UBTService_RunEQS : public UBTService_BlackboardBase
 {
 	GENERATED_BODY()
 
@@ -32,23 +32,23 @@ protected:
 	FQueryFinishedSignature QueryFinishedDelegate;
 
 public:
-	UBTService_RunEQS(const FObjectInitializer& ObjectInitializer = FObjectInitializer::Get());
+	AIMODULE_API UBTService_RunEQS(const FObjectInitializer& ObjectInitializer = FObjectInitializer::Get());
 
-	virtual void InitializeFromAsset(UBehaviorTree& Asset) override;
-	virtual void TickNode(UBehaviorTreeComponent& OwnerComp, uint8* NodeMemory, float DeltaSeconds) override;
-	virtual void OnCeaseRelevant(UBehaviorTreeComponent& OwnerComp, uint8* NodeMemory) override;
-	virtual void InitializeMemory(UBehaviorTreeComponent& OwnerComp, uint8* NodeMemory, EBTMemoryInit::Type InitType) const override;
+	AIMODULE_API virtual void InitializeFromAsset(UBehaviorTree& Asset) override;
+	AIMODULE_API virtual void TickNode(UBehaviorTreeComponent& OwnerComp, uint8* NodeMemory, float DeltaSeconds) override;
+	AIMODULE_API virtual void OnCeaseRelevant(UBehaviorTreeComponent& OwnerComp, uint8* NodeMemory) override;
+	AIMODULE_API virtual void InitializeMemory(UBehaviorTreeComponent& OwnerComp, uint8* NodeMemory, EBTMemoryInit::Type InitType) const override;
 	virtual uint16 GetInstanceMemorySize() const override { return sizeof(FBTEQSServiceMemory); }
 
-	virtual FString GetStaticDescription() const override;
+	AIMODULE_API virtual FString GetStaticDescription() const override;
 
 #if WITH_EDITOR
 	/** We need this only for verification, no need to have it in shipped builds */
-	virtual void CleanupMemory(UBehaviorTreeComponent& OwnerComp, uint8* NodeMemory, EBTMemoryClear::Type CleanupType) const override;
+	AIMODULE_API virtual void CleanupMemory(UBehaviorTreeComponent& OwnerComp, uint8* NodeMemory, EBTMemoryClear::Type CleanupType) const override;
 	/** prepare query params */
-	virtual void PostEditChangeProperty(struct FPropertyChangedEvent& PropertyChangedEvent) override;
+	AIMODULE_API virtual void PostEditChangeProperty(struct FPropertyChangedEvent& PropertyChangedEvent) override;
 #endif // WITH_EDITOR
 
 protected:
-	void OnQueryFinished(TSharedPtr<FEnvQueryResult> Result);
+	AIMODULE_API void OnQueryFinished(TSharedPtr<FEnvQueryResult> Result);
 };
diff --git a/Engine/Source/Runtime/AIModule/Classes/BehaviorTree/Tasks/BTTask_BlackboardBase.h b/Engine/Source/Runtime/AIModule/Classes/BehaviorTree/Tasks/BTTask_BlackboardBase.h
--- a/Engine/Source/Runtime/AIModule/Classes/BehaviorTree/Tasks/BTTask_BlackboardBase.h
+++ b/Engine/Source/Runtime/AIModule/Classes/BehaviorTree/Tasks/BTTask_BlackboardBase.h
@@ -9,16 +9,16 @@
 
 class UBehaviorTree;
 
-UCLASS(Abstract)
-class AIMODULE_API UBTTask_BlackboardBase : public UBTTaskNode
+UCLASS(Abstract, MinimalAPI)
+class UBTTask_BlackboardBase : public UBTTaskNode
 {
 	GENERATED_UCLASS_BODY()
 
 	/** initialize any asset related data */
-	virtual void InitializeFromAsset(UBehaviorTree& Asset) override;
+	AIMODULE_API virtual void InitializeFromAsset(UBehaviorTree& Asset) override;
 
 	/** get name of selected blackboard key */
-	FName GetSelectedBlackboardKey() const;
+	AIMODULE_API FName GetSelectedBlackboardKey() const;
 
 protected:
 
diff --git a/Engine/Source/Runtime/AIModule/Classes/BehaviorTree/Tasks/BTTask_BlueprintBase.h b/Engine/Source/Runtime/AIModule/Classes/BehaviorTree/Tasks/BTTask_BlueprintBase.h
--- a/Engine/Source/Runtime/AIModule/Classes/BehaviorTree/Tasks/BTTask_BlueprintBase.h
+++ b/Engine/Source/Runtime/AIModule/Classes/BehaviorTree/Tasks/BTTask_BlueprintBase.h
@@ -20,24 +20,24 @@ class UBehaviorTree;
  *  Please use them safely (unregister at abort) and call IsTaskExecuting() when in doubt.
  */
 
-UCLASS(Abstract, Blueprintable)
-class AIMODULE_API UBTTask_BlueprintBase : public UBTTaskNode
+UCLASS(Abstract, Blueprintable, MinimalAPI)
+class UBTTask_BlueprintBase : public UBTTaskNode
 {
 	GENERATED_UCLASS_BODY()
 	
-	virtual EBTNodeResult::Type ExecuteTask(UBehaviorTreeComponent& OwnerComp, uint8* NodeMemory) override;
-	virtual EBTNodeResult::Type AbortTask(UBehaviorTreeComponent& OwnerComp, uint8* NodeMemory) override;
+	AIMODULE_API virtual EBTNodeResult::Type ExecuteTask(UBehaviorTreeComponent& OwnerComp, uint8* NodeMemory) override;
+	AIMODULE_API virtual EBTNodeResult::Type AbortTask(UBehaviorTreeComponent& OwnerComp, uint8* NodeMemory) override;
 
-	virtual FString GetStaticDescription() const override;
-	virtual void DescribeRuntimeValues(const UBehaviorTreeComponent& OwnerComp, uint8* NodeMemory, EBTDescriptionVerbosity::Type Verbosity, TArray<FString>& Values) const override;
-	virtual void OnInstanceDestroyed(UBehaviorTreeComponent& OwnerComp) override;
-	virtual void OnTaskFinished(UBehaviorTreeComponent& OwnerComp, uint8* NodeMemory, EBTNodeResult::Type TaskResult) override;
-	virtual void InitializeFromAsset(UBehaviorTree& Asset) override;
+	AIMODULE_API virtual FString GetStaticDescription() const override;
+	AIMODULE_API virtual void DescribeRuntimeValues(const UBehaviorTreeComponent& OwnerComp, uint8* NodeMemory, EBTDescriptionVerbosity::Type Verbosity, TArray<FString>& Values) const override;
+	AIMODULE_API virtual void OnInstanceDestroyed(UBehaviorTreeComponent& OwnerComp) override;
+	AIMODULE_API virtual void OnTaskFinished(UBehaviorTreeComponent& OwnerComp, uint8* NodeMemory, EBTNodeResult::Type TaskResult) override;
+	AIMODULE_API virtual void InitializeFromAsset(UBehaviorTree& Asset) override;
 
-	virtual void SetOwner(AActor* ActorOwner) override;
+	AIMODULE_API virtual void SetOwner(AActor* ActorOwner) override;
 
 #if WITH_EDITOR
-	virtual bool UsesBlueprint() const override;
+	AIMODULE_API virtual bool UsesBlueprint() const override;
 #endif
 
 protected:
@@ -89,65 +89,65 @@ protected:
 	 *	@Note that if both generic and AI event versions are implemented only the more 
 	 *	suitable one will be called, meaning the AI version if called for AI, generic one otherwise */
 	UFUNCTION(BlueprintImplementableEvent)
-	void ReceiveExecute(AActor* OwnerActor);
+	AIMODULE_API void ReceiveExecute(AActor* OwnerActor);
 
 	/** if blueprint graph contains this event, task will stay active until FinishAbort is called
 	 *	@Note that if both generic and AI event versions are implemented only the more
 	 *	suitable one will be called, meaning the AI version if called for AI, generic one otherwise */
 	UFUNCTION(BlueprintImplementableEvent)
-	void ReceiveAbort(AActor* OwnerActor);
+	AIMODULE_API void ReceiveAbort(AActor* OwnerActor);
 
 	/** tick function
 	 *	@Note that if both generic and AI event versions are implemented only the more
 	 *	suitable one will be called, meaning the AI version if called for AI, generic one otherwise */
 	UFUNCTION(BlueprintImplementableEvent)
-	void ReceiveTick(AActor* OwnerActor, float DeltaSeconds);
+	AIMODULE_API void ReceiveTick(AActor* OwnerActor, float DeltaSeconds);
 
 	/** Alternative AI version of ReceiveExecute
 	*	@see ReceiveExecute for more details
 	 *	@Note that if both generic and AI event versions are implemented only the more
 	 *	suitable one will be called, meaning the AI version if called for AI, generic one otherwise */
 	UFUNCTION(BlueprintImplementableEvent, Category = AI)
-	void ReceiveExecuteAI(AAIController* OwnerController, APawn* ControlledPawn);
+	AIMODULE_API void ReceiveExecuteAI(AAIController* OwnerController, APawn* ControlledPawn);
 
 	/** Alternative AI version of ReceiveAbort
 	 *	@see ReceiveAbort for more details
 	 *	@Note that if both generic and AI event versions are implemented only the more
 	 *	suitable one will be called, meaning the AI version if called for AI, generic one otherwise */
 	UFUNCTION(BlueprintImplementableEvent, Category = AI)
-	void ReceiveAbortAI(AAIController* OwnerController, APawn* ControlledPawn);
+	AIMODULE_API void ReceiveAbortAI(AAIController* OwnerController, APawn* ControlledPawn);
 
 	/** Alternative AI version of tick function.
 	 *	@see ReceiveTick for more details
 	 *	@Note that if both generic and AI event versions are implemented only the more
 	 *	suitable one will be called, meaning the AI version if called for AI, generic one otherwise */
 	UFUNCTION(BlueprintImplementableEvent, Category = AI)
-	void ReceiveTickAI(AAIController* OwnerController, APawn* ControlledPawn, float DeltaSeconds);
+	AIMODULE_API void ReceiveTickAI(AAIController* OwnerController, APawn* ControlledPawn, float DeltaSeconds);
 
 	/** finishes task execution with Success or Fail result */
 	UFUNCTION(BlueprintCallable, Category="AI|BehaviorTree")
-	void FinishExecute(bool bSuccess);
+	AIMODULE_API void FinishExecute(bool bSuccess);
 
 	/** aborts task execution */
 	UFUNCTION(BlueprintCallable, Category="AI|BehaviorTree")
-	void FinishAbort();
+	AIMODULE_API void FinishAbort();
 
 	/** task execution will be finished (with result 'Success') after receiving specified message */
 	UFUNCTION(BlueprintCallable, Category="AI|BehaviorTree")
-	void SetFinishOnMessage(FName MessageName);
+	AIMODULE_API void SetFinishOnMessage(FName MessageName);
 
 	/** task execution will be finished (with result 'Success') after receiving specified message with indicated ID */
 	UFUNCTION(BlueprintCallable, Category="AI|BehaviorTree")
-	void SetFinishOnMessageWithId(FName MessageName, int32 RequestID = -1);
+	AIMODULE_API void SetFinishOnMessageWithId(FName MessageName, int32 RequestID = -1);
 
 	/** check if task is currently being executed */
 	UFUNCTION(BlueprintCallable, Category="AI|BehaviorTree")
-	bool IsTaskExecuting() const;
+	AIMODULE_API bool IsTaskExecuting() const;
 	
 	/** check if task is currently being aborted */
 	UFUNCTION(BlueprintCallable, Category = "AI|BehaviorTree")
-	bool IsTaskAborting() const;
+	AIMODULE_API bool IsTaskAborting() const;
 
 	/** ticks this task */
-	virtual void TickTask(UBehaviorTreeComponent& OwnerComp, uint8* NodeMemory, float DeltaSeconds) override;
+	AIMODULE_API virtual void TickTask(UBehaviorTreeComponent& OwnerComp, uint8* NodeMemory, float DeltaSeconds) override;
 };
diff --git a/Engine/Source/Runtime/AIModule/Classes/BehaviorTree/Tasks/BTTask_FinishWithResult.h b/Engine/Source/Runtime/AIModule/Classes/BehaviorTree/Tasks/BTTask_FinishWithResult.h
--- a/Engine/Source/Runtime/AIModule/Classes/BehaviorTree/Tasks/BTTask_FinishWithResult.h
+++ b/Engine/Source/Runtime/AIModule/Classes/BehaviorTree/Tasks/BTTask_FinishWithResult.h
@@ -7,16 +7,16 @@
 /**
  * Instantly finishes with given result
  */
-UCLASS()
-class AIMODULE_API UBTTask_FinishWithResult : public UBTTaskNode
+UCLASS(MinimalAPI)
+class UBTTask_FinishWithResult : public UBTTaskNode
 {
 	GENERATED_BODY()
 
 public:
-	UBTTask_FinishWithResult(const FObjectInitializer& ObjectInitializer = FObjectInitializer::Get());
+	AIMODULE_API UBTTask_FinishWithResult(const FObjectInitializer& ObjectInitializer = FObjectInitializer::Get());
 
-	virtual EBTNodeResult::Type ExecuteTask(UBehaviorTreeComponent& OwnerComp, uint8* NodeMemory) override;
-	virtual FString GetStaticDescription() const override;
+	AIMODULE_API virtual EBTNodeResult::Type ExecuteTask(UBehaviorTreeComponent& OwnerComp, uint8* NodeMemory) override;
+	AIMODULE_API virtual FString GetStaticDescription() const override;
 	
 protected:
 	/** allows adding random time to wait time */
diff --git a/Engine/Source/Runtime/AIModule/Classes/BehaviorTree/Tasks/BTTask_GameplayTaskBase.h b/Engine/Source/Runtime/AIModule/Classes/BehaviorTree/Tasks/BTTask_GameplayTaskBase.h
--- a/Engine/Source/Runtime/AIModule/Classes/BehaviorTree/Tasks/BTTask_GameplayTaskBase.h
+++ b/Engine/Source/Runtime/AIModule/Classes/BehaviorTree/Tasks/BTTask_GameplayTaskBase.h
@@ -22,15 +22,15 @@ struct FBTGameplayTaskMemory
  * - initialize task with values if needed
  * - use StartGameplayTask() helper to execute and get node result
  */
-UCLASS(Abstract)
-class AIMODULE_API UBTTask_GameplayTaskBase : public UBTTaskNode
+UCLASS(Abstract, MinimalAPI)
+class UBTTask_GameplayTaskBase : public UBTTaskNode
 {
 	GENERATED_UCLASS_BODY()
 
-	virtual EBTNodeResult::Type AbortTask(UBehaviorTreeComponent& OwnerComp, uint8* NodeMemory) override;
-	virtual void OnTaskFinished(UBehaviorTreeComponent& OwnerComp, uint8* NodeMemory, EBTNodeResult::Type TaskResult) override;
-	virtual void OnGameplayTaskDeactivated(UGameplayTask& Task) override;
-	virtual uint16 GetInstanceMemorySize() const override;
+	AIMODULE_API virtual EBTNodeResult::Type AbortTask(UBehaviorTreeComponent& OwnerComp, uint8* NodeMemory) override;
+	AIMODULE_API virtual void OnTaskFinished(UBehaviorTreeComponent& OwnerComp, uint8* NodeMemory, EBTNodeResult::Type TaskResult) override;
+	AIMODULE_API virtual void OnGameplayTaskDeactivated(UGameplayTask& Task) override;
+	AIMODULE_API virtual uint16 GetInstanceMemorySize() const override;
 
 protected:
 
@@ -39,8 +39,8 @@ protected:
 	uint32 bWaitForGameplayTask : 1;
 
 	/** start task and initialize FBTGameplayTaskMemory memory block */
-	EBTNodeResult::Type StartGameplayTask(UBehaviorTreeComponent& OwnerComp, uint8* NodeMemory, UAITask& Task);
+	AIMODULE_API EBTNodeResult::Type StartGameplayTask(UBehaviorTreeComponent& OwnerComp, uint8* NodeMemory, UAITask& Task);
 
 	/** get finish result from task */
-	virtual EBTNodeResult::Type DetermineGameplayTaskResult(UAITask& Task) const;
+	AIMODULE_API virtual EBTNodeResult::Type DetermineGameplayTaskResult(UAITask& Task) const;
 };
diff --git a/Engine/Source/Runtime/AIModule/Classes/BehaviorTree/Tasks/BTTask_MakeNoise.h b/Engine/Source/Runtime/AIModule/Classes/BehaviorTree/Tasks/BTTask_MakeNoise.h
--- a/Engine/Source/Runtime/AIModule/Classes/BehaviorTree/Tasks/BTTask_MakeNoise.h
+++ b/Engine/Source/Runtime/AIModule/Classes/BehaviorTree/Tasks/BTTask_MakeNoise.h
@@ -11,8 +11,8 @@
  * Make Noise task node.
  * A task node that calls MakeNoise() on this Pawn when executed.
  */
-UCLASS()
-class AIMODULE_API UBTTask_MakeNoise : public UBTTaskNode
+UCLASS(MinimalAPI)
+class UBTTask_MakeNoise : public UBTTaskNode
 {
 	GENERATED_UCLASS_BODY()
 
@@ -20,9 +20,9 @@ class AIMODULE_API UBTTask_MakeNoise : public UBTTaskNode
 	UPROPERTY(Category=Node, EditAnywhere, meta=(ClampMin = "0", ClampMax = "1", UIMin = "0", UIMax = "1"))
 	float Loudnes;
 
-	virtual EBTNodeResult::Type ExecuteTask(UBehaviorTreeComponent& OwnerComp, uint8* NodeMemory) override;
+	AIMODULE_API virtual EBTNodeResult::Type ExecuteTask(UBehaviorTreeComponent& OwnerComp, uint8* NodeMemory) override;
 
 #if WITH_EDITOR
-	virtual FName GetNodeIconName() const override;
+	AIMODULE_API virtual FName GetNodeIconName() const override;
 #endif // WITH_EDITOR
 };
diff --git a/Engine/Source/Runtime/AIModule/Classes/BehaviorTree/Tasks/BTTask_MoveDirectlyToward.h b/Engine/Source/Runtime/AIModule/Classes/BehaviorTree/Tasks/BTTask_MoveDirectlyToward.h
--- a/Engine/Source/Runtime/AIModule/Classes/BehaviorTree/Tasks/BTTask_MoveDirectlyToward.h
+++ b/Engine/Source/Runtime/AIModule/Classes/BehaviorTree/Tasks/BTTask_MoveDirectlyToward.h
@@ -13,15 +13,15 @@ class UBehaviorTree;
  * Move Directly Toward task node.
  * Moves the AI pawn toward the specified Actor or Location (Vector) blackboard entry in a straight line, without regard to any navigation system. If you need the AI to navigate, use the "Move To" node instead.
  */
-UCLASS(config=Game)
-class AIMODULE_API UBTTask_MoveDirectlyToward : public UBTTask_MoveTo
+UCLASS(config=Game, MinimalAPI)
+class UBTTask_MoveDirectlyToward : public UBTTask_MoveTo
 {
 	GENERATED_UCLASS_BODY()
 
-	virtual void PostLoad() override;
-	virtual void InitializeFromAsset(UBehaviorTree& Asset) override;
+	AIMODULE_API virtual void PostLoad() override;
+	AIMODULE_API virtual void InitializeFromAsset(UBehaviorTree& Asset) override;
 #if WITH_EDITOR
-	virtual FName GetNodeIconName() const override;
+	AIMODULE_API virtual FName GetNodeIconName() const override;
 #endif // WITH_EDITOR
 
 	UE_DEPRECATED_FORGAME(4.12, "This property is now deprecated, use UBTTask_MoveTo::bTrackMovingGoal instead.")
diff --git a/Engine/Source/Runtime/AIModule/Classes/BehaviorTree/Tasks/BTTask_MoveTo.h b/Engine/Source/Runtime/AIModule/Classes/BehaviorTree/Tasks/BTTask_MoveTo.h
--- a/Engine/Source/Runtime/AIModule/Classes/BehaviorTree/Tasks/BTTask_MoveTo.h
+++ b/Engine/Source/Runtime/AIModule/Classes/BehaviorTree/Tasks/BTTask_MoveTo.h
@@ -31,8 +31,8 @@ struct FBTMoveToTaskMemory
  * Move To task node.
  * Moves the AI pawn toward the specified Actor or Location blackboard entry using the navigation system.
  */
-UCLASS(config=Game)
-class AIMODULE_API UBTTask_MoveTo : public UBTTask_BlackboardBase
+UCLASS(config=Game, MinimalAPI)
+class UBTTask_MoveTo : public UBTTask_BlackboardBase
 {
 	GENERATED_UCLASS_BODY()
 
@@ -91,28 +91,28 @@ class AIMODULE_API UBTTask_MoveTo : public UBTTask_BlackboardBase
 	/** if set, move will use pathfinding. Not exposed on purpose, please use BTTask_MoveDirectlyToward */
 	uint32 bUsePathfinding : 1;
 
-	virtual EBTNodeResult::Type ExecuteTask(UBehaviorTreeComponent& OwnerComp, uint8* NodeMemory) override;
-	virtual EBTNodeResult::Type AbortTask(UBehaviorTreeComponent& OwnerComp, uint8* NodeMemory) override;
-	virtual void OnTaskFinished(UBehaviorTreeComponent& OwnerComp, uint8* NodeMemory, EBTNodeResult::Type TaskResult) override;
-	virtual uint16 GetInstanceMemorySize() const override;
-	virtual void PostLoad() override;
+	AIMODULE_API virtual EBTNodeResult::Type ExecuteTask(UBehaviorTreeComponent& OwnerComp, uint8* NodeMemory) override;
+	AIMODULE_API virtual EBTNodeResult::Type AbortTask(UBehaviorTreeComponent& OwnerComp, uint8* NodeMemory) override;
+	AIMODULE_API virtual void OnTaskFinished(UBehaviorTreeComponent& OwnerComp, uint8* NodeMemory, EBTNodeResult::Type TaskResult) override;
+	AIMODULE_API virtual uint16 GetInstanceMemorySize() const override;
+	AIMODULE_API virtual void PostLoad() override;
 
-	virtual void OnGameplayTaskDeactivated(UGameplayTask& Task) override;
-	virtual void OnMessage(UBehaviorTreeComponent& OwnerComp, uint8* NodeMemory, FName Message, int32 RequestID, bool bSuccess) override;
-	EBlackboardNotificationResult OnBlackboardValueChange(const UBlackboardComponent& Blackboard, FBlackboard::FKey ChangedKeyID);
+	AIMODULE_API virtual void OnGameplayTaskDeactivated(UGameplayTask& Task) override;
+	AIMODULE_API virtual void OnMessage(UBehaviorTreeComponent& OwnerComp, uint8* NodeMemory, FName Message, int32 RequestID, bool bSuccess) override;
+	AIMODULE_API EBlackboardNotificationResult OnBlackboardValueChange(const UBlackboardComponent& Blackboard, FBlackboard::FKey ChangedKeyID);
 
-	virtual void DescribeRuntimeValues(const UBehaviorTreeComponent& OwnerComp, uint8* NodeMemory, EBTDescriptionVerbosity::Type Verbosity, TArray<FString>& Values) const override;
-	virtual FString GetStaticDescription() const override;
+	AIMODULE_API virtual void DescribeRuntimeValues(const UBehaviorTreeComponent& OwnerComp, uint8* NodeMemory, EBTDescriptionVerbosity::Type Verbosity, TArray<FString>& Values) const override;
+	AIMODULE_API virtual FString GetStaticDescription() const override;
 
 #if WITH_EDITOR
-	virtual FName GetNodeIconName() const override;
-	virtual void OnNodeCreated() override;
+	AIMODULE_API virtual FName GetNodeIconName() const override;
+	AIMODULE_API virtual void OnNodeCreated() override;
 #endif // WITH_EDITOR
 
 protected:
 
-	virtual EBTNodeResult::Type PerformMoveTask(UBehaviorTreeComponent& OwnerComp, uint8* NodeMemory);
+	AIMODULE_API virtual EBTNodeResult::Type PerformMoveTask(UBehaviorTreeComponent& OwnerComp, uint8* NodeMemory);
 	
 	/** prepares move task for activation */
-	virtual UAITask_MoveTo* PrepareMoveTask(UBehaviorTreeComponent& OwnerComp, UAITask_MoveTo* ExistingTask, FAIMoveRequest& MoveRequest);
+	AIMODULE_API virtual UAITask_MoveTo* PrepareMoveTask(UBehaviorTreeComponent& OwnerComp, UAITask_MoveTo* ExistingTask, FAIMoveRequest& MoveRequest);
 };
diff --git a/Engine/Source/Runtime/AIModule/Classes/BehaviorTree/Tasks/BTTask_PawnActionBase.h b/Engine/Source/Runtime/AIModule/Classes/BehaviorTree/Tasks/BTTask_PawnActionBase.h
--- a/Engine/Source/Runtime/AIModule/Classes/BehaviorTree/Tasks/BTTask_PawnActionBase.h
+++ b/Engine/Source/Runtime/AIModule/Classes/BehaviorTree/Tasks/BTTask_PawnActionBase.h
@@ -26,27 +26,27 @@ enum class EPawnActionTaskResult : uint8
  * Please use result returned by PushAction for ExecuteTask function.
  */
 
-UCLASS(Abstract)
-class AIMODULE_API UBTTask_PawnActionBase : public UBTTaskNode
+UCLASS(Abstract, MinimalAPI)
+class UBTTask_PawnActionBase : public UBTTaskNode
 {
 	GENERATED_UCLASS_BODY()
 
-	virtual EBTNodeResult::Type AbortTask(UBehaviorTreeComponent& OwnerComp, uint8* NodeMemory) override;
+	AIMODULE_API virtual EBTNodeResult::Type AbortTask(UBehaviorTreeComponent& OwnerComp, uint8* NodeMemory) override;
 
 protected:
 
 	/** starts executing pawn action */
-	EBTNodeResult::Type PushAction(UBehaviorTreeComponent& OwnerComp, UDEPRECATED_PawnAction& Action);
+	AIMODULE_API EBTNodeResult::Type PushAction(UBehaviorTreeComponent& OwnerComp, UDEPRECATED_PawnAction& Action);
 
 	/** action observer, updates state of task */
-	virtual void OnActionEvent(UDEPRECATED_PawnAction& Action, EPawnActionEventType::Type Event);
+	AIMODULE_API virtual void OnActionEvent(UDEPRECATED_PawnAction& Action, EPawnActionEventType::Type Event);
 
 	/** called when action is removed from stack (FinishedAborting) by some external event
 	 *  default behavior: finish task as failed */
-	virtual void OnActionLost(UDEPRECATED_PawnAction& Action);
+	AIMODULE_API virtual void OnActionLost(UDEPRECATED_PawnAction& Action);
 
 public:
 
 	/** helper functions, should be used when behavior tree task deals with pawn actions, but can't derive from this class */
-	static EPawnActionTaskResult ActionEventHandler(UBTTaskNode* TaskNode, UDEPRECATED_PawnAction& Action, EPawnActionEventType::Type Event);
+	static AIMODULE_API EPawnActionTaskResult ActionEventHandler(UBTTaskNode* TaskNode, UDEPRECATED_PawnAction& Action, EPawnActionEventType::Type Event);
 };
diff --git a/Engine/Source/Runtime/AIModule/Classes/BehaviorTree/Tasks/BTTask_PlayAnimation.h b/Engine/Source/Runtime/AIModule/Classes/BehaviorTree/Tasks/BTTask_PlayAnimation.h
--- a/Engine/Source/Runtime/AIModule/Classes/BehaviorTree/Tasks/BTTask_PlayAnimation.h
+++ b/Engine/Source/Runtime/AIModule/Classes/BehaviorTree/Tasks/BTTask_PlayAnimation.h
@@ -15,8 +15,8 @@
  *	Note that this node is generic and is handing multiple special cases,
  *	If you want a more efficient solution you'll need to implement it yourself (or wait for our BTTask_PlayCharacterAnimation)
  */
-UCLASS()
-class AIMODULE_API UBTTask_PlayAnimation : public UBTTaskNode
+UCLASS(MinimalAPI)
+class UBTTask_PlayAnimation : public UBTTaskNode
 {
 	GENERATED_UCLASS_BODY()
 
@@ -42,16 +42,16 @@ class AIMODULE_API UBTTask_PlayAnimation : public UBTTaskNode
 	FTimerDelegate TimerDelegate;
 	FTimerHandle TimerHandle;
 
-	virtual EBTNodeResult::Type ExecuteTask(UBehaviorTreeComponent& OwnerComp, uint8* NodeMemory) override;
-	virtual EBTNodeResult::Type AbortTask(UBehaviorTreeComponent& OwnerComp, uint8* NodeMemory) override;
-	virtual FString GetStaticDescription() const override;
+	AIMODULE_API virtual EBTNodeResult::Type ExecuteTask(UBehaviorTreeComponent& OwnerComp, uint8* NodeMemory) override;
+	AIMODULE_API virtual EBTNodeResult::Type AbortTask(UBehaviorTreeComponent& OwnerComp, uint8* NodeMemory) override;
+	AIMODULE_API virtual FString GetStaticDescription() const override;
 
-	void OnAnimationTimerDone();
+	AIMODULE_API void OnAnimationTimerDone();
 	
 #if WITH_EDITOR
-	virtual FName GetNodeIconName() const override;
+	AIMODULE_API virtual FName GetNodeIconName() const override;
 #endif // WITH_EDITOR
 
 protected:
-	void CleanUp(UBehaviorTreeComponent& OwnerComp);
+	AIMODULE_API void CleanUp(UBehaviorTreeComponent& OwnerComp);
 };
diff --git a/Engine/Source/Runtime/AIModule/Classes/BehaviorTree/Tasks/BTTask_PlaySound.h b/Engine/Source/Runtime/AIModule/Classes/BehaviorTree/Tasks/BTTask_PlaySound.h
--- a/Engine/Source/Runtime/AIModule/Classes/BehaviorTree/Tasks/BTTask_PlaySound.h
+++ b/Engine/Source/Runtime/AIModule/Classes/BehaviorTree/Tasks/BTTask_PlaySound.h
@@ -13,8 +13,8 @@ class USoundCue;
  * Play Sound task node.
  * Plays the specified sound when executed.
  */
-UCLASS()
-class AIMODULE_API UBTTask_PlaySound : public UBTTaskNode
+UCLASS(MinimalAPI)
+class UBTTask_PlaySound : public UBTTaskNode
 {
 	GENERATED_UCLASS_BODY()
 
@@ -22,10 +22,10 @@ class AIMODULE_API UBTTask_PlaySound : public UBTTaskNode
 	UPROPERTY(Category=Node, EditAnywhere)
 	TObjectPtr<USoundCue> SoundToPlay;
 
-	virtual EBTNodeResult::Type ExecuteTask(UBehaviorTreeComponent& OwnerComp, uint8* NodeMemory) override;
-	virtual FString GetStaticDescription() const override;
+	AIMODULE_API virtual EBTNodeResult::Type ExecuteTask(UBehaviorTreeComponent& OwnerComp, uint8* NodeMemory) override;
+	AIMODULE_API virtual FString GetStaticDescription() const override;
 
 #if WITH_EDITOR
-	virtual FName GetNodeIconName() const override;
+	AIMODULE_API virtual FName GetNodeIconName() const override;
 #endif // WITH_EDITOR
 };
diff --git a/Engine/Source/Runtime/AIModule/Classes/BehaviorTree/Tasks/BTTask_PushPawnAction.h b/Engine/Source/Runtime/AIModule/Classes/BehaviorTree/Tasks/BTTask_PushPawnAction.h
--- a/Engine/Source/Runtime/AIModule/Classes/BehaviorTree/Tasks/BTTask_PushPawnAction.h
+++ b/Engine/Source/Runtime/AIModule/Classes/BehaviorTree/Tasks/BTTask_PushPawnAction.h
@@ -13,8 +13,8 @@ class UDEPRECATED_PawnAction;
  * Action task node.
  * Push pawn action to controller.
  */
-UCLASS()
-class AIMODULE_API UBTTask_PushPawnAction : public UBTTask_PawnActionBase
+UCLASS(MinimalAPI)
+class UBTTask_PushPawnAction : public UBTTask_PawnActionBase
 {
 	GENERATED_UCLASS_BODY()
 
@@ -22,6 +22,6 @@ protected:
 	UPROPERTY(Instanced)
 	TObjectPtr<UDEPRECATED_PawnAction> Action_DEPRECATED;
 
-	virtual EBTNodeResult::Type ExecuteTask(UBehaviorTreeComponent& OwnerComp, uint8* NodeMemory) override;
-	virtual FString GetStaticDescription() const override;
+	AIMODULE_API virtual EBTNodeResult::Type ExecuteTask(UBehaviorTreeComponent& OwnerComp, uint8* NodeMemory) override;
+	AIMODULE_API virtual FString GetStaticDescription() const override;
 };
diff --git a/Engine/Source/Runtime/AIModule/Classes/BehaviorTree/Tasks/BTTask_RotateToFaceBBEntry.h b/Engine/Source/Runtime/AIModule/Classes/BehaviorTree/Tasks/BTTask_RotateToFaceBBEntry.h
--- a/Engine/Source/Runtime/AIModule/Classes/BehaviorTree/Tasks/BTTask_RotateToFaceBBEntry.h
+++ b/Engine/Source/Runtime/AIModule/Classes/BehaviorTree/Tasks/BTTask_RotateToFaceBBEntry.h
@@ -13,8 +13,8 @@ class AAIController;
 /**
  * 
  */
-UCLASS(config = Game)
-class AIMODULE_API UBTTask_RotateToFaceBBEntry : public UBTTask_BlackboardBase
+UCLASS(config = Game, MinimalAPI)
+class UBTTask_RotateToFaceBBEntry : public UBTTask_BlackboardBase
 {
 	GENERATED_UCLASS_BODY()
 
@@ -29,19 +29,19 @@ private:
 
 public:
 
-	virtual void PostInitProperties() override;
-	virtual void PostLoad() override;
+	AIMODULE_API virtual void PostInitProperties() override;
+	AIMODULE_API virtual void PostLoad() override;
 
-	virtual EBTNodeResult::Type ExecuteTask(UBehaviorTreeComponent& OwnerComp, uint8* NodeMemory) override;
-	virtual void TickTask(UBehaviorTreeComponent& OwnerComp, uint8* NodeMemory, float DeltaSeconds) override;
-	virtual EBTNodeResult::Type AbortTask(UBehaviorTreeComponent& OwnerComp, uint8* NodeMemory) override;
-	virtual void DescribeRuntimeValues(const UBehaviorTreeComponent& OwnerComp, uint8* NodeMemory, EBTDescriptionVerbosity::Type Verbosity, TArray<FString>& Values) const override;
-	virtual FString GetStaticDescription() const override;
+	AIMODULE_API virtual EBTNodeResult::Type ExecuteTask(UBehaviorTreeComponent& OwnerComp, uint8* NodeMemory) override;
+	AIMODULE_API virtual void TickTask(UBehaviorTreeComponent& OwnerComp, uint8* NodeMemory, float DeltaSeconds) override;
+	AIMODULE_API virtual EBTNodeResult::Type AbortTask(UBehaviorTreeComponent& OwnerComp, uint8* NodeMemory) override;
+	AIMODULE_API virtual void DescribeRuntimeValues(const UBehaviorTreeComponent& OwnerComp, uint8* NodeMemory, EBTDescriptionVerbosity::Type Verbosity, TArray<FString>& Values) const override;
+	AIMODULE_API virtual FString GetStaticDescription() const override;
 	
 	virtual uint16 GetInstanceMemorySize() const override { return sizeof(FBTFocusMemory); }
 
 protected:
 
 	float GetPrecisionDot() const { return PrecisionDot; }
-	void CleanUp(AAIController& AIController, uint8* NodeMemory);
+	AIMODULE_API void CleanUp(AAIController& AIController, uint8* NodeMemory);
 };
