[Engine]

diff --git a/Engine/Source/Runtime/AIModule/Classes/BehaviorTree/Tasks/BTTask_RunBehavior.h b/Engine/Source/Runtime/AIModule/Classes/BehaviorTree/Tasks/BTTask_RunBehavior.h
--- a/Engine/Source/Runtime/AIModule/Classes/BehaviorTree/Tasks/BTTask_RunBehavior.h
+++ b/Engine/Source/Runtime/AIModule/Classes/BehaviorTree/Tasks/BTTask_RunBehavior.h
@@ -19,23 +19,23 @@
  * Use RunBehaviorDynamic task for subtrees that need to be changed in runtime.
  */
 
-UCLASS()
-class AIMODULE_API UBTTask_RunBehavior : public UBTTaskNode
+UCLASS(MinimalAPI)
+class UBTTask_RunBehavior : public UBTTaskNode
 {
 	GENERATED_UCLASS_BODY()
 
-	virtual EBTNodeResult::Type ExecuteTask(UBehaviorTreeComponent& OwnerComp, uint8* NodeMemory) override;
-	virtual FString GetStaticDescription() const override;
+	AIMODULE_API virtual EBTNodeResult::Type ExecuteTask(UBehaviorTreeComponent& OwnerComp, uint8* NodeMemory) override;
+	AIMODULE_API virtual FString GetStaticDescription() const override;
 
 #if WITH_EDITOR
-	virtual FName GetNodeIconName() const override;
+	AIMODULE_API virtual FName GetNodeIconName() const override;
 #endif // WITH_EDITOR
 
 	/** @returns number of injected nodes */
-	int32 GetInjectedNodesCount() const;
+	AIMODULE_API int32 GetInjectedNodesCount() const;
 
 	/** @returns subtree asset */
-	UBehaviorTree* GetSubtreeAsset() const;
+	AIMODULE_API UBehaviorTree* GetSubtreeAsset() const;
 
 protected:
 
@@ -44,7 +44,7 @@ protected:
 	TObjectPtr<UBehaviorTree> BehaviorAsset;
 
 	/** called when subtree is removed from active stack */
-	virtual void OnSubtreeDeactivated(UBehaviorTreeComponent& OwnerComp, EBTNodeResult::Type NodeResult);
+	AIMODULE_API virtual void OnSubtreeDeactivated(UBehaviorTreeComponent& OwnerComp, EBTNodeResult::Type NodeResult);
 };
 
 //////////////////////////////////////////////////////////////////////////
diff --git a/Engine/Source/Runtime/AIModule/Classes/BehaviorTree/Tasks/BTTask_RunBehaviorDynamic.h b/Engine/Source/Runtime/AIModule/Classes/BehaviorTree/Tasks/BTTask_RunBehaviorDynamic.h
--- a/Engine/Source/Runtime/AIModule/Classes/BehaviorTree/Tasks/BTTask_RunBehaviorDynamic.h
+++ b/Engine/Source/Runtime/AIModule/Classes/BehaviorTree/Tasks/BTTask_RunBehaviorDynamic.h
@@ -17,26 +17,26 @@ class UBehaviorTree;
  * Does NOT support subtree's root level decorators!
  */
 
-UCLASS()
-class AIMODULE_API UBTTask_RunBehaviorDynamic : public UBTTaskNode
+UCLASS(MinimalAPI)
+class UBTTask_RunBehaviorDynamic : public UBTTaskNode
 {
 	GENERATED_UCLASS_BODY()
 
-	virtual EBTNodeResult::Type ExecuteTask(UBehaviorTreeComponent& OwnerComp, uint8* NodeMemory) override;
-	virtual void OnInstanceCreated(UBehaviorTreeComponent& OwnerComp) override;
-	virtual FString GetStaticDescription() const override;
-	virtual void DescribeRuntimeValues(const UBehaviorTreeComponent& OwnerComp, uint8* NodeMemory, EBTDescriptionVerbosity::Type Verbosity, TArray<FString>& Values) const override;
+	AIMODULE_API virtual EBTNodeResult::Type ExecuteTask(UBehaviorTreeComponent& OwnerComp, uint8* NodeMemory) override;
+	AIMODULE_API virtual void OnInstanceCreated(UBehaviorTreeComponent& OwnerComp) override;
+	AIMODULE_API virtual FString GetStaticDescription() const override;
+	AIMODULE_API virtual void DescribeRuntimeValues(const UBehaviorTreeComponent& OwnerComp, uint8* NodeMemory, EBTDescriptionVerbosity::Type Verbosity, TArray<FString>& Values) const override;
 
 #if WITH_EDITOR
-	virtual FName GetNodeIconName() const override;
+	AIMODULE_API virtual FName GetNodeIconName() const override;
 #endif // WITH_EDITOR
 
-	bool HasMatchingTag(const FGameplayTag& Tag) const;
-	const FGameplayTag& GetInjectionTag() const;
-	bool SetBehaviorAsset(UBehaviorTree* NewBehaviorAsset);
+	AIMODULE_API bool HasMatchingTag(const FGameplayTag& Tag) const;
+	AIMODULE_API const FGameplayTag& GetInjectionTag() const;
+	AIMODULE_API bool SetBehaviorAsset(UBehaviorTree* NewBehaviorAsset);
 	
 	/** @returns default subtree asset */
-	UBehaviorTree* GetDefaultBehaviorAsset() const;
+	AIMODULE_API UBehaviorTree* GetDefaultBehaviorAsset() const;
 
 protected:
 
@@ -53,7 +53,7 @@ protected:
 	TObjectPtr<UBehaviorTree> BehaviorAsset;
 
 	/** called when subtree is removed from active stack */
-	virtual void OnSubtreeDeactivated(UBehaviorTreeComponent& OwnerComp, EBTNodeResult::Type NodeResult);
+	AIMODULE_API virtual void OnSubtreeDeactivated(UBehaviorTreeComponent& OwnerComp, EBTNodeResult::Type NodeResult);
 };
 
 //////////////////////////////////////////////////////////////////////////
diff --git a/Engine/Source/Runtime/AIModule/Classes/BehaviorTree/Tasks/BTTask_RunEQSQuery.h b/Engine/Source/Runtime/AIModule/Classes/BehaviorTree/Tasks/BTTask_RunEQSQuery.h
--- a/Engine/Source/Runtime/AIModule/Classes/BehaviorTree/Tasks/BTTask_RunEQSQuery.h
+++ b/Engine/Source/Runtime/AIModule/Classes/BehaviorTree/Tasks/BTTask_RunEQSQuery.h
@@ -21,8 +21,8 @@ struct FBTEnvQueryTaskMemory
  * Run Environment Query System Query task node.
  * Runs the specified environment query when executed.
  */
-UCLASS()
-class AIMODULE_API UBTTask_RunEQSQuery : public UBTTask_BlackboardBase
+UCLASS(MinimalAPI)
+class UBTTask_RunEQSQuery : public UBTTask_BlackboardBase
 {
 	GENERATED_UCLASS_BODY()
 
@@ -61,28 +61,28 @@ class AIMODULE_API UBTTask_RunEQSQuery : public UBTTask_BlackboardBase
 
 	FQueryFinishedSignature QueryFinishedDelegate;
 
-	virtual void InitializeFromAsset(UBehaviorTree& Asset) override;
-	virtual EBTNodeResult::Type ExecuteTask(UBehaviorTreeComponent& OwnerComp, uint8* NodeMemory) override;
-	virtual EBTNodeResult::Type AbortTask(UBehaviorTreeComponent& OwnerComp, uint8* NodeMemory) override;
+	AIMODULE_API virtual void InitializeFromAsset(UBehaviorTree& Asset) override;
+	AIMODULE_API virtual EBTNodeResult::Type ExecuteTask(UBehaviorTreeComponent& OwnerComp, uint8* NodeMemory) override;
+	AIMODULE_API virtual EBTNodeResult::Type AbortTask(UBehaviorTreeComponent& OwnerComp, uint8* NodeMemory) override;
 
-	virtual void DescribeRuntimeValues(const UBehaviorTreeComponent& OwnerComp, uint8* NodeMemory, EBTDescriptionVerbosity::Type Verbosity, TArray<FString>& Values) const override;
-	virtual FString GetStaticDescription() const override;
-	virtual uint16 GetInstanceMemorySize() const override;
+	AIMODULE_API virtual void DescribeRuntimeValues(const UBehaviorTreeComponent& OwnerComp, uint8* NodeMemory, EBTDescriptionVerbosity::Type Verbosity, TArray<FString>& Values) const override;
+	AIMODULE_API virtual FString GetStaticDescription() const override;
+	AIMODULE_API virtual uint16 GetInstanceMemorySize() const override;
 
 	/** finish task */
-	void OnQueryFinished(TSharedPtr<FEnvQueryResult> Result);
+	AIMODULE_API void OnQueryFinished(TSharedPtr<FEnvQueryResult> Result);
 
 	/** Convert QueryParams to QueryConfig */
-	virtual void PostLoad() override;
+	AIMODULE_API virtual void PostLoad() override;
 
 #if WITH_EDITOR
 	/** prepare query params */
-	virtual void PostEditChangeProperty(struct FPropertyChangedEvent& PropertyChangedEvent) override;
-	virtual FName GetNodeIconName() const override;
+	AIMODULE_API virtual void PostEditChangeProperty(struct FPropertyChangedEvent& PropertyChangedEvent) override;
+	AIMODULE_API virtual FName GetNodeIconName() const override;
 #endif
 
 protected:
 
 	/** gather all filters from existing EnvQueryItemTypes */
-	void CollectKeyFilters();
+	AIMODULE_API void CollectKeyFilters();
 };
diff --git a/Engine/Source/Runtime/AIModule/Classes/BehaviorTree/Tasks/BTTask_SetTagCooldown.h b/Engine/Source/Runtime/AIModule/Classes/BehaviorTree/Tasks/BTTask_SetTagCooldown.h
--- a/Engine/Source/Runtime/AIModule/Classes/BehaviorTree/Tasks/BTTask_SetTagCooldown.h
+++ b/Engine/Source/Runtime/AIModule/Classes/BehaviorTree/Tasks/BTTask_SetTagCooldown.h
@@ -11,8 +11,8 @@
  * Cooldown task node.
  * Sets a cooldown tag value.  Use with cooldown tag decorators to prevent behavior tree execution.
  */
-UCLASS()
-class AIMODULE_API UBTTask_SetTagCooldown : public UBTTaskNode
+UCLASS(MinimalAPI)
+class UBTTask_SetTagCooldown : public UBTTaskNode
 {
 	GENERATED_UCLASS_BODY()
 
@@ -28,10 +28,10 @@ class AIMODULE_API UBTTask_SetTagCooldown : public UBTTaskNode
 	UPROPERTY(Category = Cooldown, EditAnywhere)
 	float CooldownDuration;
 
-	virtual EBTNodeResult::Type ExecuteTask(UBehaviorTreeComponent& OwnerComp, uint8* NodeMemory) override;
-	virtual FString GetStaticDescription() const override;
+	AIMODULE_API virtual EBTNodeResult::Type ExecuteTask(UBehaviorTreeComponent& OwnerComp, uint8* NodeMemory) override;
+	AIMODULE_API virtual FString GetStaticDescription() const override;
 
 #if WITH_EDITOR
-	virtual FName GetNodeIconName() const override;
+	AIMODULE_API virtual FName GetNodeIconName() const override;
 #endif // WITH_EDITOR
 };
diff --git a/Engine/Source/Runtime/AIModule/Classes/BehaviorTree/Tasks/BTTask_Wait.h b/Engine/Source/Runtime/AIModule/Classes/BehaviorTree/Tasks/BTTask_Wait.h
--- a/Engine/Source/Runtime/AIModule/Classes/BehaviorTree/Tasks/BTTask_Wait.h
+++ b/Engine/Source/Runtime/AIModule/Classes/BehaviorTree/Tasks/BTTask_Wait.h
@@ -14,8 +14,8 @@ struct FBTWaitTaskMemory
  * Wait task node.
  * Wait for the specified time when executed.
  */
-UCLASS()
-class AIMODULE_API UBTTask_Wait : public UBTTaskNode
+UCLASS(MinimalAPI)
+class UBTTask_Wait : public UBTTaskNode
 {
 	GENERATED_UCLASS_BODY()
 
@@ -27,15 +27,15 @@ class AIMODULE_API UBTTask_Wait : public UBTTaskNode
 	UPROPERTY(Category = Wait, EditAnywhere, meta = (UIMin = 0, ClampMin = 0))
 	float RandomDeviation;
 
-	virtual EBTNodeResult::Type ExecuteTask(UBehaviorTreeComponent& OwnerComp, uint8* NodeMemory) override;
-	virtual void DescribeRuntimeValues(const UBehaviorTreeComponent& OwnerComp, uint8* NodeMemory, EBTDescriptionVerbosity::Type Verbosity, TArray<FString>& Values) const override;
-	virtual FString GetStaticDescription() const override;
+	AIMODULE_API virtual EBTNodeResult::Type ExecuteTask(UBehaviorTreeComponent& OwnerComp, uint8* NodeMemory) override;
+	AIMODULE_API virtual void DescribeRuntimeValues(const UBehaviorTreeComponent& OwnerComp, uint8* NodeMemory, EBTDescriptionVerbosity::Type Verbosity, TArray<FString>& Values) const override;
+	AIMODULE_API virtual FString GetStaticDescription() const override;
 
 #if WITH_EDITOR
-	virtual FName GetNodeIconName() const override;
+	AIMODULE_API virtual FName GetNodeIconName() const override;
 #endif // WITH_EDITOR
 
 protected:
 
-	virtual void TickTask(UBehaviorTreeComponent& OwnerComp, uint8* NodeMemory, float DeltaSeconds) override;
+	AIMODULE_API virtual void TickTask(UBehaviorTreeComponent& OwnerComp, uint8* NodeMemory, float DeltaSeconds) override;
 };
diff --git a/Engine/Source/Runtime/AIModule/Classes/BehaviorTree/Tasks/BTTask_WaitBlackboardTime.h b/Engine/Source/Runtime/AIModule/Classes/BehaviorTree/Tasks/BTTask_WaitBlackboardTime.h
--- a/Engine/Source/Runtime/AIModule/Classes/BehaviorTree/Tasks/BTTask_WaitBlackboardTime.h
+++ b/Engine/Source/Runtime/AIModule/Classes/BehaviorTree/Tasks/BTTask_WaitBlackboardTime.h
@@ -12,17 +12,17 @@ class UBehaviorTree;
  * Wait task node.
  * Wait for the time specified by a Blackboard key when executed.
  */
-UCLASS(hidecategories=Wait)
-class AIMODULE_API UBTTask_WaitBlackboardTime : public UBTTask_Wait
+UCLASS(hidecategories=Wait, MinimalAPI)
+class UBTTask_WaitBlackboardTime : public UBTTask_Wait
 {
 	GENERATED_UCLASS_BODY()
 
-	virtual void InitializeFromAsset(UBehaviorTree& Asset) override;
-	virtual EBTNodeResult::Type ExecuteTask(UBehaviorTreeComponent& OwnerComp, uint8* NodeMemory) override;
-	virtual FString GetStaticDescription() const override;
+	AIMODULE_API virtual void InitializeFromAsset(UBehaviorTree& Asset) override;
+	AIMODULE_API virtual EBTNodeResult::Type ExecuteTask(UBehaviorTreeComponent& OwnerComp, uint8* NodeMemory) override;
+	AIMODULE_API virtual FString GetStaticDescription() const override;
 
 	/** get name of selected blackboard key */
-	FName GetSelectedBlackboardKey() const;
+	AIMODULE_API FName GetSelectedBlackboardKey() const;
 
 
 protected:
diff --git a/Engine/Source/Runtime/AIModule/Classes/Blueprint/AIBlueprintHelperLibrary.h b/Engine/Source/Runtime/AIModule/Classes/Blueprint/AIBlueprintHelperLibrary.h
--- a/Engine/Source/Runtime/AIModule/Classes/Blueprint/AIBlueprintHelperLibrary.h
+++ b/Engine/Source/Runtime/AIModule/Classes/Blueprint/AIBlueprintHelperLibrary.h
@@ -22,16 +22,16 @@ class UBlackboardComponent;
 class UNavigationPath;
 class UPathFollowingComponent;
 
-UCLASS(meta=(ScriptName="AIHelperLibrary"))
-class AIMODULE_API UAIBlueprintHelperLibrary : public UBlueprintFunctionLibrary
+UCLASS(meta=(ScriptName="AIHelperLibrary"), MinimalAPI)
+class UAIBlueprintHelperLibrary : public UBlueprintFunctionLibrary
 {
 	GENERATED_UCLASS_BODY()
 
 	UFUNCTION(BlueprintCallable, meta=(WorldContext="WorldContextObject", BlueprintInternalUseOnly = "TRUE"))
-	static UAIAsyncTaskBlueprintProxy* CreateMoveToProxyObject(UObject* WorldContextObject, APawn* Pawn, FVector Destination, AActor* TargetActor = NULL, float AcceptanceRadius = 5.f, bool bStopOnOverlap = false);
+	static AIMODULE_API UAIAsyncTaskBlueprintProxy* CreateMoveToProxyObject(UObject* WorldContextObject, APawn* Pawn, FVector Destination, AActor* TargetActor = NULL, float AcceptanceRadius = 5.f, bool bStopOnOverlap = false);
 
 	UFUNCTION(BlueprintCallable, Category="AI", meta=(DefaultToSelf="MessageSource"))
-	static void SendAIMessage(APawn* Target, FName Message, UObject* MessageSource, bool bSuccess = true);
+	static AIMODULE_API void SendAIMessage(APawn* Target, FName Message, UObject* MessageSource, bool bSuccess = true);
 
 	/**	Spawns AI agent of a given class. The PawnClass needs to have AIController 
 	 *	set for the function to spawn a controller as well.
@@ -42,57 +42,57 @@ class AIMODULE_API UAIBlueprintHelperLibrary : public UBlueprintFunctionLibrary
 	 *		persistent level (which is the default behavior).
 	 */
 	UFUNCTION(BlueprintCallable, Category="AI", meta=(WorldContext="WorldContextObject", UnsafeDuringActorConstruction="true", AdvancedDisplay = "Owner"))
-	static APawn* SpawnAIFromClass(UObject* WorldContextObject, TSubclassOf<APawn> PawnClass, UBehaviorTree* BehaviorTree, FVector Location, FRotator Rotation = FRotator::ZeroRotator, bool bNoCollisionFail = false, AActor* Owner = nullptr);
+	static AIMODULE_API APawn* SpawnAIFromClass(UObject* WorldContextObject, TSubclassOf<APawn> PawnClass, UBehaviorTree* BehaviorTree, FVector Location, FRotator Rotation = FRotator::ZeroRotator, bool bNoCollisionFail = false, AActor* Owner = nullptr);
 
 	/** The way it works exactly is if the actor passed in is a pawn, then the function retrieves 
 	 *	pawn's controller cast to AIController. Otherwise the function returns actor cast to AIController. */
 	UFUNCTION(BlueprintPure, Category = "AI", meta = (DefaultToSelf = "ControlledObject"))
-	static AAIController* GetAIController(AActor* ControlledActor);
+	static AIMODULE_API AAIController* GetAIController(AActor* ControlledActor);
 
 	UFUNCTION(BlueprintPure, Category="AI", meta=(DefaultToSelf="Target"))
-	static UBlackboardComponent* GetBlackboard(AActor* Target);
+	static AIMODULE_API UBlackboardComponent* GetBlackboard(AActor* Target);
 
 	/** locks indicated AI resources of animated pawn */
 	UFUNCTION(BlueprintCallable, Category = "Animation", BlueprintAuthorityOnly, meta = (DefaultToSelf = "AnimInstance"))
-	static void LockAIResourcesWithAnimation(UAnimInstance* AnimInstance, bool bLockMovement, bool LockAILogic);
+	static AIMODULE_API void LockAIResourcesWithAnimation(UAnimInstance* AnimInstance, bool bLockMovement, bool LockAILogic);
 
 	/** unlocks indicated AI resources of animated pawn. Will unlock only animation-locked resources */
 	UFUNCTION(BlueprintCallable, Category = "Animation", BlueprintAuthorityOnly, meta = (DefaultToSelf = "AnimInstance"))
-	static void UnlockAIResourcesWithAnimation(UAnimInstance* AnimInstance, bool bUnlockMovement, bool UnlockAILogic);
+	static AIMODULE_API void UnlockAIResourcesWithAnimation(UAnimInstance* AnimInstance, bool bUnlockMovement, bool UnlockAILogic);
 
 	UFUNCTION(BlueprintPure, Category = "AI")
-	static bool IsValidAILocation(FVector Location);
+	static AIMODULE_API bool IsValidAILocation(FVector Location);
 
 	UFUNCTION(BlueprintPure, Category = "AI")
-	static bool IsValidAIDirection(FVector DirectionVector);
+	static AIMODULE_API bool IsValidAIDirection(FVector DirectionVector);
 
 	UFUNCTION(BlueprintPure, Category = "AI")
-	static bool IsValidAIRotation(FRotator Rotation);
+	static AIMODULE_API bool IsValidAIRotation(FRotator Rotation);
 
 	/** Returns a NEW UOBJECT that is a COPY of navigation path given controller is currently using. 
 	 *	The result being a copy means you won't be able to influence agent's pathfollowing 
 	 *	by manipulating received path.
 	 *	Please use GetCurrentPathPoints if you only need the array of path points. */
 	UFUNCTION(BlueprintPure, Category = "AI", meta = (UnsafeDuringActorConstruction = "true"))
-	static UNavigationPath* GetCurrentPath(AController* Controller);
+	static AIMODULE_API UNavigationPath* GetCurrentPath(AController* Controller);
 
 	/** Returns an array of navigation path points given controller is currently using. */
 	UFUNCTION(BlueprintPure, Category = "AI", meta = (UnsafeDuringActorConstruction = "true"))
-	static const TArray<FVector> GetCurrentPathPoints(AController* Controller);
+	static AIMODULE_API const TArray<FVector> GetCurrentPathPoints(AController* Controller);
 
 	/** Return the path index the given controller is currently at. Returns INDEX_NONE if no path. */
 	UFUNCTION(BlueprintPure, Category = "AI", meta = (UnsafeDuringActorConstruction = "true"))
-	static int32 GetCurrentPathIndex(const AController* Controller);
+	static AIMODULE_API int32 GetCurrentPathIndex(const AController* Controller);
 
 	/** Return the path index of the next nav link for the current path of the given controller. Returns INDEX_NONE if no path or no incoming nav link. */
 	UFUNCTION(BlueprintPure, Category = "AI", meta = (UnsafeDuringActorConstruction = "true"))
-	static int32 GetNextNavLinkIndex(const AController* Controller);
+	static AIMODULE_API int32 GetNextNavLinkIndex(const AController* Controller);
 
 	UFUNCTION(BlueprintCallable, Category = "AI|Navigation")
-	static void SimpleMoveToActor(AController* Controller, const AActor* Goal);
+	static AIMODULE_API void SimpleMoveToActor(AController* Controller, const AActor* Goal);
 
 	UFUNCTION(BlueprintCallable, Category = "AI|Navigation")
-	static void SimpleMoveToLocation(AController* Controller, const FVector& Goal);
+	static AIMODULE_API void SimpleMoveToLocation(AController* Controller, const FVector& Goal);
 
 private:
 	static UPathFollowingComponent* GetPathComp(const AController* Controller);
diff --git a/Engine/Source/Runtime/AIModule/Classes/BrainComponent.h b/Engine/Source/Runtime/AIModule/Classes/BrainComponent.h
--- a/Engine/Source/Runtime/AIModule/Classes/BrainComponent.h
+++ b/Engine/Source/Runtime/AIModule/Classes/BrainComponent.h
@@ -23,7 +23,7 @@ DECLARE_DELEGATE_TwoParams(FOnAIMessage, UBrainComponent*, const FAIMessage&);
 
 DECLARE_LOG_CATEGORY_EXTERN(LogBrain, Warning, All);
 
-struct AIMODULE_API FAIMessage
+struct FAIMessage
 {
 	enum EStatus
 	{
@@ -58,33 +58,33 @@ struct AIMODULE_API FAIMessage
 	void ClearFlag(uint8 Flag) { MessageFlags &= ~Flag; }
 	bool HasFlag(uint8 Flag) const { return (MessageFlags & Flag) != 0; }
 
-	static void Send(AController* Controller, const FAIMessage& Message);
-	static void Send(APawn* Pawn, const FAIMessage& Message);
-	static void Send(UBrainComponent* BrainComp, const FAIMessage& Message);
+	static AIMODULE_API void Send(AController* Controller, const FAIMessage& Message);
+	static AIMODULE_API void Send(APawn* Pawn, const FAIMessage& Message);
+	static AIMODULE_API void Send(UBrainComponent* BrainComp, const FAIMessage& Message);
 
-	static void Broadcast(UObject* WorldContextObject, const FAIMessage& Message);
+	static AIMODULE_API void Broadcast(UObject* WorldContextObject, const FAIMessage& Message);
 };
 
 typedef TSharedPtr<struct FAIMessageObserver> FAIMessageObserverHandle;
 
-struct AIMODULE_API FAIMessageObserver : public TSharedFromThis<FAIMessageObserver>
+struct FAIMessageObserver : public TSharedFromThis<FAIMessageObserver>
 {
 public:
-	FAIMessageObserver();
+	AIMODULE_API FAIMessageObserver();
 
-	static FAIMessageObserverHandle Create(AController* Controller, FName MessageType, FOnAIMessage const& Delegate);
-	static FAIMessageObserverHandle Create(AController* Controller, FName MessageType, FAIRequestID MessageID, FOnAIMessage const& Delegate);
+	static AIMODULE_API FAIMessageObserverHandle Create(AController* Controller, FName MessageType, FOnAIMessage const& Delegate);
+	static AIMODULE_API FAIMessageObserverHandle Create(AController* Controller, FName MessageType, FAIRequestID MessageID, FOnAIMessage const& Delegate);
 
-	static FAIMessageObserverHandle Create(APawn* Pawn, FName MessageType, FOnAIMessage const& Delegate);
-	static FAIMessageObserverHandle Create(APawn* Pawn, FName MessageType, FAIRequestID MessageID, FOnAIMessage const& Delegate);
+	static AIMODULE_API FAIMessageObserverHandle Create(APawn* Pawn, FName MessageType, FOnAIMessage const& Delegate);
+	static AIMODULE_API FAIMessageObserverHandle Create(APawn* Pawn, FName MessageType, FAIRequestID MessageID, FOnAIMessage const& Delegate);
 
-	static FAIMessageObserverHandle Create(UBrainComponent* BrainComp, FName MessageType, FOnAIMessage const& Delegate);
-	static FAIMessageObserverHandle Create(UBrainComponent* BrainComp, FName MessageType, FAIRequestID MessageID, FOnAIMessage const& Delegate);
+	static AIMODULE_API FAIMessageObserverHandle Create(UBrainComponent* BrainComp, FName MessageType, FOnAIMessage const& Delegate);
+	static AIMODULE_API FAIMessageObserverHandle Create(UBrainComponent* BrainComp, FName MessageType, FAIRequestID MessageID, FOnAIMessage const& Delegate);
 
-	~FAIMessageObserver();
+	AIMODULE_API ~FAIMessageObserver();
 
-	void OnMessage(const FAIMessage& Message);
-	FString DescribeObservedMessage() const;
+	AIMODULE_API void OnMessage(const FAIMessage& Message);
+	AIMODULE_API FString DescribeObservedMessage() const;
 	
 	FORCEINLINE FName GetObservedMessageType() const { return MessageType; }
 	FORCEINLINE FAIRequestID GetObservedMessageID() const { return MessageID; }
@@ -92,8 +92,8 @@ public:
 
 private:
 
-	void Register(UBrainComponent* OwnerComp);
-	void Unregister();
+	AIMODULE_API void Register(UBrainComponent* OwnerComp);
+	AIMODULE_API void Unregister();
 
 	/** observed message type */
 	FName MessageType;
@@ -109,12 +109,12 @@ private:
 	TWeakObjectPtr<UBrainComponent> Owner;
 
 	// Non-copyable
-	FAIMessageObserver(const FAIMessageObserver&);
-	FAIMessageObserver& operator=(const FAIMessageObserver&);
+	AIMODULE_API FAIMessageObserver(const FAIMessageObserver&);
+	AIMODULE_API FAIMessageObserver& operator=(const FAIMessageObserver&);
 };
 
-UCLASS(ClassGroup = AI, BlueprintType, hidecategories = (Sockets, Collision))
-class AIMODULE_API UBrainComponent : public UActorComponent, public IAIResourceInterface
+UCLASS(ClassGroup = AI, BlueprintType, hidecategories = (Sockets, Collision), MinimalAPI)
+class UBrainComponent : public UActorComponent, public IAIResourceInterface
 {
 	GENERATED_UCLASS_BODY()
 
@@ -136,19 +136,19 @@ public:
 	/** To be called in case we want to restart AI logic while it's still being locked.
 	 *	On subsequent ResumeLogic instead RestartLogic will be called. 
 	 *	@note this call does nothing if logic is not locked at the moment of call */
-	void RequestLogicRestartOnUnlock();
+	AIMODULE_API void RequestLogicRestartOnUnlock();
 
 	/** Starts brain logic. If brain is already running, will not do anything. */
 	UFUNCTION(BlueprintCallable, Category = "AI|Logic")
-	virtual void StartLogic();
+	AIMODULE_API virtual void StartLogic();
 
 	/** Restarts currently running or previously ran brain logic. */
 	UFUNCTION(BlueprintCallable, Category = "AI|Logic")
-	virtual void RestartLogic();
+	AIMODULE_API virtual void RestartLogic();
 
 	/** Stops currently running brain logic. */
 	UFUNCTION(BlueprintCallable, Category = "AI|Logic")
-	virtual void StopLogic(const FString& Reason);
+	AIMODULE_API virtual void StopLogic(const FString& Reason);
 
 	/** AI logic won't be needed anymore, stop all activity and run cleanup */
 	virtual void Cleanup() {}
@@ -160,42 +160,42 @@ public:
 	 *  MUST be called by child implementations!
 	 *	@return indicates whether child class' ResumeLogic should be called (true) or has it been 
 	 *	handled in a different way and no other actions are required (false)*/
-	virtual EAILogicResuming::Type ResumeLogic(const FString& Reason);
+	AIMODULE_API virtual EAILogicResuming::Type ResumeLogic(const FString& Reason);
 public:
 
 	UFUNCTION(BlueprintPure, Category = "AI|Logic")
-	virtual bool IsRunning() const;
+	AIMODULE_API virtual bool IsRunning() const;
 
 	UFUNCTION(BlueprintPure, Category = "AI|Logic")
-	virtual bool IsPaused() const;
+	AIMODULE_API virtual bool IsPaused() const;
 
 #if ENABLE_VISUAL_LOG
-	virtual void DescribeSelfToVisLog(struct FVisualLogEntry* Snapshot) const;
+	AIMODULE_API virtual void DescribeSelfToVisLog(struct FVisualLogEntry* Snapshot) const;
 #endif // ENABLE_VISUAL_LOG
 	
 	// IAIResourceInterface begin
-	virtual void LockResource(EAIRequestPriority::Type LockSource) override;
-	virtual void ClearResourceLock(EAIRequestPriority::Type LockSource) override;
-	virtual void ForceUnlockResource() override;
-	virtual bool IsResourceLocked() const override;
+	AIMODULE_API virtual void LockResource(EAIRequestPriority::Type LockSource) override;
+	AIMODULE_API virtual void ClearResourceLock(EAIRequestPriority::Type LockSource) override;
+	AIMODULE_API virtual void ForceUnlockResource() override;
+	AIMODULE_API virtual bool IsResourceLocked() const override;
 	// IAIResourceInterface end
 	
-	virtual void HandleMessage(const FAIMessage& Message);
+	AIMODULE_API virtual void HandleMessage(const FAIMessage& Message);
 	
 	/** BEGIN UActorComponent overrides */
-	virtual void InitializeComponent() override;
-	virtual void TickComponent(float DeltaTime, enum ELevelTick TickType, FActorComponentTickFunction *ThisTickFunction) override;
-	virtual void OnRegister() override;
+	AIMODULE_API virtual void InitializeComponent() override;
+	AIMODULE_API virtual void TickComponent(float DeltaTime, enum ELevelTick TickType, FActorComponentTickFunction *ThisTickFunction) override;
+	AIMODULE_API virtual void OnRegister() override;
 	/** END UActorComponent overrides */
 
 	/** caches BlackboardComponent's pointer to be used with this brain component */
-	void CacheBlackboardComponent(UBlackboardComponent* BBComp);
+	AIMODULE_API void CacheBlackboardComponent(UBlackboardComponent* BBComp);
 
 	/** @return blackboard used with this component */
-	UBlackboardComponent* GetBlackboardComponent();
+	AIMODULE_API UBlackboardComponent* GetBlackboardComponent();
 
 	/** @return blackboard used with this component */
-	const UBlackboardComponent* GetBlackboardComponent() const;
+	AIMODULE_API const UBlackboardComponent* GetBlackboardComponent() const;
 
 	FORCEINLINE AAIController* GetAIOwner() const { return AIOwner; }
 
@@ -216,9 +216,9 @@ private:
 public:
 	// static names to be used with SendMessage. Fell free to define game-specific
 	// messages anywhere you want
-	static const FName AIMessage_MoveFinished;
-	static const FName AIMessage_RepathFailed;
-	static const FName AIMessage_QueryFinished;
+	static AIMODULE_API const FName AIMessage_MoveFinished;
+	static AIMODULE_API const FName AIMessage_RepathFailed;
+	static AIMODULE_API const FName AIMessage_QueryFinished;
 };
 
 //////////////////////////////////////////////////////////////////////////
diff --git a/Engine/Source/Runtime/AIModule/Classes/DataProviders/AIDataProvider.h b/Engine/Source/Runtime/AIModule/Classes/DataProviders/AIDataProvider.h
--- a/Engine/Source/Runtime/AIModule/Classes/DataProviders/AIDataProvider.h
+++ b/Engine/Source/Runtime/AIModule/Classes/DataProviders/AIDataProvider.h
@@ -22,7 +22,7 @@ class UAIDataProvider;
  */
 
 USTRUCT()
-struct AIMODULE_API FAIDataProviderValue
+struct FAIDataProviderValue
 {
 	GENERATED_USTRUCT_BODY()
 
@@ -40,14 +40,14 @@ public:
 	FName DataField;
 
 	/** describe default data */
-	virtual FString ValueToString() const;
-	FString ToString() const;
+	AIMODULE_API virtual FString ValueToString() const;
+	AIMODULE_API FString ToString() const;
 
 	/** filter for provider's properties */
-	virtual bool IsMatchingType(FProperty* PropType) const;
+	AIMODULE_API virtual bool IsMatchingType(FProperty* PropType) const;
 
 	/** find all properties of provider that are matching filter */
-	void GetMatchingProperties(TArray<FName>& MatchingProperties) const;
+	AIMODULE_API void GetMatchingProperties(TArray<FName>& MatchingProperties) const;
 
 	/** return raw data from provider's property */
 	template<typename T>
@@ -57,7 +57,7 @@ public:
 	}
 
 	/** bind data in provider and cache property for faster access */
-	void BindData(const UObject* Owner, int32 RequestId) const;
+	AIMODULE_API void BindData(const UObject* Owner, int32 RequestId) const;
 
 	FORCEINLINE bool IsDynamic() const { return DataBinding != nullptr; }
 
@@ -71,7 +71,7 @@ public:
 };
 
 USTRUCT()
-struct AIMODULE_API FAIDataProviderTypedValue : public FAIDataProviderValue
+struct FAIDataProviderTypedValue : public FAIDataProviderValue
 {
 	GENERATED_USTRUCT_BODY()
 
@@ -86,10 +86,10 @@ struct AIMODULE_API FAIDataProviderTypedValue : public FAIDataProviderValue
 	FFieldClass* PropertyType;
 
 	/** filter for provider's properties */
-	virtual bool IsMatchingType(FProperty* PropType) const override;
+	AIMODULE_API virtual bool IsMatchingType(FProperty* PropType) const override;
 
 	/** Implementing Serialize to convert UClass to FFieldClass */
-	bool Serialize(FArchive& Ar);
+	AIMODULE_API bool Serialize(FArchive& Ar);
 };
 
 template<>
@@ -102,60 +102,60 @@ struct TStructOpsTypeTraits<FAIDataProviderTypedValue> : public TStructOpsTypeTr
 };
 
 USTRUCT()
-struct AIMODULE_API FAIDataProviderStructValue : public FAIDataProviderValue
+struct FAIDataProviderStructValue : public FAIDataProviderValue
 {
 	GENERATED_USTRUCT_BODY()
 
 	/** name of UStruct type */
 	FString StructName;
 
-	virtual bool IsMatchingType(FProperty* PropType) const override;
+	AIMODULE_API virtual bool IsMatchingType(FProperty* PropType) const override;
 };
 
 USTRUCT()
-struct AIMODULE_API FAIDataProviderIntValue : public FAIDataProviderTypedValue
+struct FAIDataProviderIntValue : public FAIDataProviderTypedValue
 {
 	GENERATED_USTRUCT_BODY()
-	FAIDataProviderIntValue();
+	AIMODULE_API FAIDataProviderIntValue();
 
 	UPROPERTY(EditAnywhere, Category = Value)
 	int32 DefaultValue;
 
-	int32 GetValue() const;
-	virtual FString ValueToString() const override;
+	AIMODULE_API int32 GetValue() const;
+	AIMODULE_API virtual FString ValueToString() const override;
 };
 
 USTRUCT()
-struct AIMODULE_API FAIDataProviderFloatValue : public FAIDataProviderTypedValue
+struct FAIDataProviderFloatValue : public FAIDataProviderTypedValue
 {
 	GENERATED_USTRUCT_BODY()
-	FAIDataProviderFloatValue();
+	AIMODULE_API FAIDataProviderFloatValue();
 
 	UPROPERTY(EditAnywhere, Category = Value)
 	float DefaultValue;
 
-	float GetValue() const;
-	virtual FString ValueToString() const override;
+	AIMODULE_API float GetValue() const;
+	AIMODULE_API virtual FString ValueToString() const override;
 };
 
 USTRUCT()
-struct AIMODULE_API FAIDataProviderBoolValue : public FAIDataProviderTypedValue
+struct FAIDataProviderBoolValue : public FAIDataProviderTypedValue
 {
 	GENERATED_USTRUCT_BODY()
-	FAIDataProviderBoolValue();
+	AIMODULE_API FAIDataProviderBoolValue();
 
 	UPROPERTY(EditAnywhere, Category = Value)
 	bool DefaultValue;
 
-	bool GetValue() const;
-	virtual FString ValueToString() const override;
+	AIMODULE_API bool GetValue() const;
+	AIMODULE_API virtual FString ValueToString() const override;
 };
 
-UCLASS(EditInlineNew, Abstract, CollapseCategories, AutoExpandCategories=(Provider))
-class AIMODULE_API UAIDataProvider : public UObject
+UCLASS(EditInlineNew, Abstract, CollapseCategories, AutoExpandCategories=(Provider), MinimalAPI)
+class UAIDataProvider : public UObject
 {
 	GENERATED_UCLASS_BODY()
 
-	virtual void BindData(const UObject& Owner, int32 RequestId);
-	virtual FString ToString(FName PropName) const;
+	AIMODULE_API virtual void BindData(const UObject& Owner, int32 RequestId);
+	AIMODULE_API virtual FString ToString(FName PropName) const;
 };
diff --git a/Engine/Source/Runtime/AIModule/Classes/DataProviders/AIDataProvider_QueryParams.h b/Engine/Source/Runtime/AIModule/Classes/DataProviders/AIDataProvider_QueryParams.h
--- a/Engine/Source/Runtime/AIModule/Classes/DataProviders/AIDataProvider_QueryParams.h
+++ b/Engine/Source/Runtime/AIModule/Classes/DataProviders/AIDataProvider_QueryParams.h
@@ -15,14 +15,14 @@
  * for every query execution.
  */
 
-UCLASS(EditInlineNew, meta=(DisplayName="Query Params"))
-class AIMODULE_API UAIDataProvider_QueryParams : public UAIDataProvider
+UCLASS(EditInlineNew, meta=(DisplayName="Query Params"), MinimalAPI)
+class UAIDataProvider_QueryParams : public UAIDataProvider
 {
 	GENERATED_BODY()
 
 public:
-	virtual void BindData(const UObject& Owner, int32 RequestId) override;
-	virtual FString ToString(FName PropName) const override;
+	AIMODULE_API virtual void BindData(const UObject& Owner, int32 RequestId) override;
+	AIMODULE_API virtual FString ToString(FName PropName) const override;
 
 	/** Arbitrary name this query parameter will be exposed as to outside world (like BT nodes) */
 	UPROPERTY(EditAnywhere, Category = Provider)
diff --git a/Engine/Source/Runtime/AIModule/Classes/EnvironmentQuery/Contexts/EnvQueryContext_Item.h b/Engine/Source/Runtime/AIModule/Classes/EnvironmentQuery/Contexts/EnvQueryContext_Item.h
--- a/Engine/Source/Runtime/AIModule/Classes/EnvironmentQuery/Contexts/EnvQueryContext_Item.h
+++ b/Engine/Source/Runtime/AIModule/Classes/EnvironmentQuery/Contexts/EnvQueryContext_Item.h
@@ -7,8 +7,8 @@
 #include "EnvironmentQuery/EnvQueryContext.h"
 #include "EnvQueryContext_Item.generated.h"
 
-UCLASS()
-class AIMODULE_API UEnvQueryContext_Item : public UEnvQueryContext
+UCLASS(MinimalAPI)
+class UEnvQueryContext_Item : public UEnvQueryContext
 {
 	GENERATED_UCLASS_BODY()
 };
diff --git a/Engine/Source/Runtime/AIModule/Classes/EnvironmentQuery/EQSRenderingComponent.h b/Engine/Source/Runtime/AIModule/Classes/EnvironmentQuery/EQSRenderingComponent.h
--- a/Engine/Source/Runtime/AIModule/Classes/EnvironmentQuery/EQSRenderingComponent.h
+++ b/Engine/Source/Runtime/AIModule/Classes/EnvironmentQuery/EQSRenderingComponent.h
@@ -14,19 +14,19 @@ class APlayerController;
 class IEQSQueryResultSourceInterface;
 class UCanvas;
 
-class AIMODULE_API FEQSSceneProxy final : public FDebugRenderSceneProxy
+class FEQSSceneProxy final : public FDebugRenderSceneProxy
 {
 	friend class FEQSRenderingDebugDrawDelegateHelper;
 public:
-	virtual SIZE_T GetTypeHash() const override;
+	AIMODULE_API virtual SIZE_T GetTypeHash() const override;
 
-	explicit FEQSSceneProxy(const UPrimitiveComponent& InComponent, const FString& ViewFlagName = TEXT("DebugAI"), const TArray<FSphere>& Spheres = TArray<FSphere>(), const TArray<FText3d>& Texts = TArray<FText3d>());
+	AIMODULE_API explicit FEQSSceneProxy(const UPrimitiveComponent& InComponent, const FString& ViewFlagName = TEXT("DebugAI"), const TArray<FSphere>& Spheres = TArray<FSphere>(), const TArray<FText3d>& Texts = TArray<FText3d>());
 	
-	virtual FPrimitiveViewRelevance GetViewRelevance(const FSceneView* View) const override;
+	AIMODULE_API virtual FPrimitiveViewRelevance GetViewRelevance(const FSceneView* View) const override;
 
 #if  USE_EQS_DEBUGGER 
-	static void CollectEQSData(const UPrimitiveComponent* InComponent, const IEQSQueryResultSourceInterface* QueryDataSource, TArray<FSphere>& Spheres, TArray<FText3d>& Texts, TArray<EQSDebug::FDebugHelper>& DebugItems);
-	static void CollectEQSData(const FEnvQueryResult* ResultItems, const FEnvQueryInstance* QueryInstance, float HighlightRangePct, bool ShouldDrawFailedItems, TArray<FSphere>& Spheres, TArray<FText3d>& Texts, TArray<EQSDebug::FDebugHelper>& DebugItems);
+	static AIMODULE_API void CollectEQSData(const UPrimitiveComponent* InComponent, const IEQSQueryResultSourceInterface* QueryDataSource, TArray<FSphere>& Spheres, TArray<FText3d>& Texts, TArray<EQSDebug::FDebugHelper>& DebugItems);
+	static AIMODULE_API void CollectEQSData(const FEnvQueryResult* ResultItems, const FEnvQueryInstance* QueryInstance, float HighlightRangePct, bool ShouldDrawFailedItems, TArray<FSphere>& Spheres, TArray<FText3d>& Texts, TArray<EQSDebug::FDebugHelper>& DebugItems);
 #endif
 private:
 	FEnvQueryResult QueryResult;	
@@ -35,9 +35,9 @@ private:
 	const IEQSQueryResultSourceInterface* QueryDataSource;
 	uint32 bDrawOnlyWhenSelected : 1;
 
-	static const FVector3f ItemDrawRadius;
+	static AIMODULE_API const FVector3f ItemDrawRadius;
 
-	bool SafeIsActorSelected() const;
+	AIMODULE_API bool SafeIsActorSelected() const;
 };
 
 #if  USE_EQS_DEBUGGER
@@ -76,24 +76,24 @@ private:
 };
 #endif
 
-UCLASS(ClassGroup = Debug)
-class AIMODULE_API UEQSRenderingComponent : public UDebugDrawComponent
+UCLASS(ClassGroup = Debug, MinimalAPI)
+class UEQSRenderingComponent : public UDebugDrawComponent
 {
 	GENERATED_UCLASS_BODY()
 
 	FString DrawFlagName;
 	uint32 bDrawOnlyWhenSelected : 1;
 
-	void ClearStoredDebugData();
+	AIMODULE_API void ClearStoredDebugData();
 #if  USE_EQS_DEBUGGER || ENABLE_VISUAL_LOG
-	void StoreDebugData(const EQSDebug::FQueryData& DebugData);
+	AIMODULE_API void StoreDebugData(const EQSDebug::FQueryData& DebugData);
 #endif
 
 protected:
-	virtual FBoxSphereBounds CalcBounds(const FTransform &LocalToWorld) const override;
+	AIMODULE_API virtual FBoxSphereBounds CalcBounds(const FTransform &LocalToWorld) const override;
 
 #if UE_ENABLE_DEBUG_DRAWING && USE_EQS_DEBUGGER
-	virtual FDebugRenderSceneProxy* CreateDebugSceneProxy() override;
+	AIMODULE_API virtual FDebugRenderSceneProxy* CreateDebugSceneProxy() override;
 	virtual FDebugDrawDelegateHelper& GetDebugDrawDelegateHelper() override { return EQSRenderingDebugDrawDelegateHelper; }
 	FEQSRenderingDebugDrawDelegateHelper EQSRenderingDebugDrawDelegateHelper;
 #endif
diff --git a/Engine/Source/Runtime/AIModule/Classes/EnvironmentQuery/EQSTestingPawn.h b/Engine/Source/Runtime/AIModule/Classes/EnvironmentQuery/EQSTestingPawn.h
--- a/Engine/Source/Runtime/AIModule/Classes/EnvironmentQuery/EQSTestingPawn.h
+++ b/Engine/Source/Runtime/AIModule/Classes/EnvironmentQuery/EQSTestingPawn.h
@@ -21,8 +21,8 @@ enum class EEnvQueryHightlightMode : uint8
 	Best25Pct UMETA(DisplayName = "Best 25%"),
 };
 
-UCLASS(hidedropdown, hidecategories=(Advanced, Attachment, Mesh, Animation, Clothing, Physics, Rendering, Lighting, Activation, CharacterMovement, AgentPhysics, Avoidance, MovementComponent, Velocity, Shape, Camera, Input, Layers, SkeletalMesh, Optimization, Pawn, Replication, Actor))
-class AIMODULE_API AEQSTestingPawn : public ACharacter, public IEQSQueryResultSourceInterface
+UCLASS(hidedropdown, hidecategories=(Advanced, Attachment, Mesh, Animation, Clothing, Physics, Rendering, Lighting, Activation, CharacterMovement, AgentPhysics, Avoidance, MovementComponent, Velocity, Shape, Camera, Input, Layers, SkeletalMesh, Optimization, Pawn, Replication, Actor), MinimalAPI)
+class AEQSTestingPawn : public ACharacter, public IEQSQueryResultSourceInterface
 {
 	GENERATED_UCLASS_BODY()
 	
@@ -81,39 +81,39 @@ protected:
 
 public:
 	/** This pawn class spawns its controller in PostInitProperties to have it available in editor mode*/
-	virtual void TickActor( float DeltaTime, enum ELevelTick TickType, FActorTickFunction& ThisTickFunction ) override;
-	virtual void PostLoad() override;
+	AIMODULE_API virtual void TickActor( float DeltaTime, enum ELevelTick TickType, FActorTickFunction& ThisTickFunction ) override;
+	AIMODULE_API virtual void PostLoad() override;
 
 #if WITH_EDITOR
-	virtual void PostEditChangeProperty(struct FPropertyChangedEvent& PropertyChangedEvent) override;
-	virtual void PostEditMove(bool bFinished) override;
-	virtual void PostTransacted(const FTransactionObjectEvent& TransactionEvent) override;
-	void OnPropertyChanged(const FName PropName);
+	AIMODULE_API virtual void PostEditChangeProperty(struct FPropertyChangedEvent& PropertyChangedEvent) override;
+	AIMODULE_API virtual void PostEditMove(bool bFinished) override;
+	AIMODULE_API virtual void PostTransacted(const FTransactionObjectEvent& TransactionEvent) override;
+	AIMODULE_API void OnPropertyChanged(const FName PropName);
 #endif // WITH_EDITOR
 
 	// IEQSQueryResultSourceInterface start
-	virtual const FEnvQueryResult* GetQueryResult() const override;
-	virtual const FEnvQueryInstance* GetQueryInstance() const  override;
+	AIMODULE_API virtual const FEnvQueryResult* GetQueryResult() const override;
+	AIMODULE_API virtual const FEnvQueryInstance* GetQueryInstance() const  override;
 
 	virtual bool GetShouldDebugDrawLabels() const override { return bDrawLabels; }
 	virtual bool GetShouldDrawFailedItems() const override{ return bDrawFailedItems; }
-	virtual float GetHighlightRangePct() const override;
+	AIMODULE_API virtual float GetHighlightRangePct() const override;
 	// IEQSQueryResultSourceInterface end
 
 	// INavAgentInterface begin
-	virtual const FNavAgentProperties& GetNavAgentPropertiesRef() const override;
+	AIMODULE_API virtual const FNavAgentProperties& GetNavAgentPropertiesRef() const override;
 	// INavAgentInterface end
 
-	void RunEQSQuery();
+	AIMODULE_API void RunEQSQuery();
 
 protected:	
-	void Reset() override;
-	void MakeOneStep();
+	AIMODULE_API void Reset() override;
+	AIMODULE_API void MakeOneStep();
 
-	void UpdateDrawing();
+	AIMODULE_API void UpdateDrawing();
 
 #if WITH_EDITOR
-	static void OnEditorSelectionChanged(UObject* NewSelection);
+	static AIMODULE_API void OnEditorSelectionChanged(UObject* NewSelection);
 #endif // WITH_EDITOR
 
 public:
diff --git a/Engine/Source/Runtime/AIModule/Classes/EnvironmentQuery/EnvQuery.h b/Engine/Source/Runtime/AIModule/Classes/EnvironmentQuery/EnvQuery.h
--- a/Engine/Source/Runtime/AIModule/Classes/EnvironmentQuery/EnvQuery.h
+++ b/Engine/Source/Runtime/AIModule/Classes/EnvironmentQuery/EnvQuery.h
@@ -15,8 +15,8 @@ class UEnvQueryOption;
 class UEdGraph;
 #endif // WITH_EDITORONLY_DATA
 
-UCLASS(BlueprintType)
-class AIMODULE_API UEnvQuery : public UDataAsset
+UCLASS(BlueprintType, MinimalAPI)
+class UEnvQuery : public UDataAsset
 {
 	GENERATED_UCLASS_BODY()
 
@@ -37,15 +37,15 @@ protected:
 
 public:
 	/** Gather all required named params */
-	void CollectQueryParams(UObject& QueryOwner, TArray<FAIDynamicParam>& NamedValues) const;
+	AIMODULE_API void CollectQueryParams(UObject& QueryOwner, TArray<FAIDynamicParam>& NamedValues) const;
 
-	virtual  void PostInitProperties() override;
+	AIMODULE_API virtual  void PostInitProperties() override;
 
 	/** QueryName patching up */
-	virtual void PostLoad() override;
+	AIMODULE_API virtual void PostLoad() override;
 #if WITH_EDITOR
-	virtual void PostRename(UObject* OldOuter, const FName OldName) override;
-	virtual void PostDuplicate(bool bDuplicateForPIE) override;
+	AIMODULE_API virtual void PostRename(UObject* OldOuter, const FName OldName) override;
+	AIMODULE_API virtual void PostDuplicate(bool bDuplicateForPIE) override;
 #endif // WITH_EDITOR
 
 	FName GetQueryName() const { return QueryName; }
diff --git a/Engine/Source/Runtime/AIModule/Classes/EnvironmentQuery/EnvQueryContext.h b/Engine/Source/Runtime/AIModule/Classes/EnvironmentQuery/EnvQueryContext.h
--- a/Engine/Source/Runtime/AIModule/Classes/EnvironmentQuery/EnvQueryContext.h
+++ b/Engine/Source/Runtime/AIModule/Classes/EnvironmentQuery/EnvQueryContext.h
@@ -10,10 +10,10 @@
 struct FEnvQueryContextData;
 struct FEnvQueryInstance;
 
-UCLASS(Abstract, EditInlineNew)
-class AIMODULE_API UEnvQueryContext : public UObject
+UCLASS(Abstract, EditInlineNew, MinimalAPI)
+class UEnvQueryContext : public UObject
 {
 	GENERATED_UCLASS_BODY()
 
-	virtual void ProvideContext(FEnvQueryInstance& QueryInstance, FEnvQueryContextData& ContextData) const;
+	AIMODULE_API virtual void ProvideContext(FEnvQueryInstance& QueryInstance, FEnvQueryContextData& ContextData) const;
 };
diff --git a/Engine/Source/Runtime/AIModule/Classes/EnvironmentQuery/EnvQueryDebugHelpers.h b/Engine/Source/Runtime/AIModule/Classes/EnvironmentQuery/EnvQueryDebugHelpers.h
--- a/Engine/Source/Runtime/AIModule/Classes/EnvironmentQuery/EnvQueryDebugHelpers.h
+++ b/Engine/Source/Runtime/AIModule/Classes/EnvironmentQuery/EnvQueryDebugHelpers.h
@@ -154,21 +154,21 @@ FArchive& operator<<(FArchive& Ar, EQSDebug::FQueryData& Data)
 #	define UE_VLOG_EQS(Query, CategoryName, Verbosity)
 #endif //ENABLE_VISUAL_LOG && USE_EQS_DEBUGGER
 
-UCLASS(Abstract)
-class AIMODULE_API UEnvQueryDebugHelpers : public UObject
+UCLASS(Abstract, MinimalAPI)
+class UEnvQueryDebugHelpers : public UObject
 {
 	GENERATED_BODY()
 public:
 #if USE_EQS_DEBUGGER
-	static void QueryToDebugData(FEnvQueryInstance& Query, EQSDebug::FQueryData& EQSLocalData, int32 MaxItemsToStore = 10);
-	static void QueryToBlobArray(FEnvQueryInstance& Query, TArray<uint8>& BlobArray, bool bUseCompression = false);
-	static void DebugDataToBlobArray(EQSDebug::FQueryData& QueryData, TArray<uint8>& BlobArray, bool bUseCompression = false);
-	static void BlobArrayToDebugData(const TArray<uint8>& BlobArray, EQSDebug::FQueryData& EQSLocalData, bool bUseCompression = false);
+	static AIMODULE_API void QueryToDebugData(FEnvQueryInstance& Query, EQSDebug::FQueryData& EQSLocalData, int32 MaxItemsToStore = 10);
+	static AIMODULE_API void QueryToBlobArray(FEnvQueryInstance& Query, TArray<uint8>& BlobArray, bool bUseCompression = false);
+	static AIMODULE_API void DebugDataToBlobArray(EQSDebug::FQueryData& QueryData, TArray<uint8>& BlobArray, bool bUseCompression = false);
+	static AIMODULE_API void BlobArrayToDebugData(const TArray<uint8>& BlobArray, EQSDebug::FQueryData& EQSLocalData, bool bUseCompression = false);
 #endif
 
 #if ENABLE_VISUAL_LOG && USE_EQS_DEBUGGER
-	static void LogQuery(FEnvQueryInstance& Query, const FLogCategoryBase& Category, ELogVerbosity::Type Verbosity);
-	static void LogQuery(FEnvQueryInstance& Query, const FName& CategoryName, ELogVerbosity::Type Verbosity);
+	static AIMODULE_API void LogQuery(FEnvQueryInstance& Query, const FLogCategoryBase& Category, ELogVerbosity::Type Verbosity);
+	static AIMODULE_API void LogQuery(FEnvQueryInstance& Query, const FName& CategoryName, ELogVerbosity::Type Verbosity);
 
 private:
 	static void LogQueryInternal(FEnvQueryInstance& Query, const FName& CategoryName, ELogVerbosity::Type Verbosity, double TimeSeconds, FVisualLogEntry *CurrentEntry);
diff --git a/Engine/Source/Runtime/AIModule/Classes/EnvironmentQuery/EnvQueryGenerator.h b/Engine/Source/Runtime/AIModule/Classes/EnvironmentQuery/EnvQueryGenerator.h
--- a/Engine/Source/Runtime/AIModule/Classes/EnvironmentQuery/EnvQueryGenerator.h
+++ b/Engine/Source/Runtime/AIModule/Classes/EnvironmentQuery/EnvQueryGenerator.h
@@ -17,8 +17,8 @@ namespace EnvQueryGeneratorVersion
 	inline const int32 Latest = DataProviders;
 }
 
-UCLASS(EditInlineNew, Abstract, meta = (Category = "Generators"))
-class AIMODULE_API UEnvQueryGenerator : public UEnvQueryNode
+UCLASS(EditInlineNew, Abstract, meta = (Category = "Generators"), MinimalAPI)
+class UEnvQueryGenerator : public UEnvQueryNode
 {
 	GENERATED_UCLASS_BODY()
 
@@ -36,6 +36,6 @@ class AIMODULE_API UEnvQueryGenerator : public UEnvQueryNode
 	virtual void GenerateItems(FEnvQueryInstance& QueryInstance) const { checkNoEntry(); }
 	virtual bool IsValidGenerator() const { return ItemType != nullptr; }
 
-	virtual void PostLoad() override;
-	void UpdateNodeVersion() override;
+	AIMODULE_API virtual void PostLoad() override;
+	AIMODULE_API void UpdateNodeVersion() override;
 };
diff --git a/Engine/Source/Runtime/AIModule/Classes/EnvironmentQuery/EnvQueryInstanceBlueprintWrapper.h b/Engine/Source/Runtime/AIModule/Classes/EnvironmentQuery/EnvQueryInstanceBlueprintWrapper.h
--- a/Engine/Source/Runtime/AIModule/Classes/EnvironmentQuery/EnvQueryInstanceBlueprintWrapper.h
+++ b/Engine/Source/Runtime/AIModule/Classes/EnvironmentQuery/EnvQueryInstanceBlueprintWrapper.h
@@ -15,8 +15,8 @@
 class AActor;
 struct FEnvQueryRequest;
 
-UCLASS(Blueprintable, BlueprintType, meta = (DisplayName = "EQS Query Instance"))
-class AIMODULE_API UEnvQueryInstanceBlueprintWrapper : public UObject, public IEQSQueryResultSourceInterface
+UCLASS(Blueprintable, BlueprintType, meta = (DisplayName = "EQS Query Instance"), MinimalAPI)
+class UEnvQueryInstanceBlueprintWrapper : public UObject, public IEQSQueryResultSourceInterface
 {
 	GENERATED_BODY()
 
@@ -53,40 +53,40 @@ protected:
 #endif // !UE_BUILD_SHIPPING
 
 public:
-	UEnvQueryInstanceBlueprintWrapper(const FObjectInitializer& ObjectInitializer);
+	AIMODULE_API UEnvQueryInstanceBlueprintWrapper(const FObjectInitializer& ObjectInitializer);
 	
-	void RunQuery(const EEnvQueryRunMode::Type InRunMode, FEnvQueryRequest& QueryRequest);
+	AIMODULE_API void RunQuery(const EEnvQueryRunMode::Type InRunMode, FEnvQueryRequest& QueryRequest);
 
 	UFUNCTION(BlueprintPure, Category = "AI|EQS")
-	float GetItemScore(int32 ItemIndex) const;
+	AIMODULE_API float GetItemScore(int32 ItemIndex) const;
 
 	/** Outputs an array filled with resulting actors.  Note that it makes sense only if ItemType is a EnvQueryItemType_ActorBase-derived type.  Returns false if there is no valid result. */
 	UFUNCTION(BlueprintCallable, BlueprintPure = false, Category = "AI|EQS")
-	bool GetQueryResultsAsActors(TArray<AActor*>& ResultActors) const;
+	AIMODULE_API bool GetQueryResultsAsActors(TArray<AActor*>& ResultActors) const;
 
 	/** Outputs an array of locations generated by the query. If the query generated Actors the the array is filled with their locations.  Returns false if there is no valid result. */
 	UFUNCTION(BlueprintCallable, BlueprintPure = false, Category = "AI|EQS")
-	bool GetQueryResultsAsLocations(TArray<FVector>& ResultLocations) const;
+	AIMODULE_API bool GetQueryResultsAsLocations(TArray<FVector>& ResultLocations) const;
 
 	/** DEPRECATED!  Use GetQueryResultsAsActors instead.  Returns an array filled with resulting actors. Note that it makes sense only if ItemType is a EnvQueryItemType_ActorBase-derived type. */
 	UFUNCTION(BlueprintPure, Category = "AI|EQS", Meta=(DeprecatedFunction,
 		DeprecationMessage="Use GetQueryResultsAsActors instead, which is more efficient itself and protects against very bad perf issues in some blueprints."))
-	TArray<AActor*> GetResultsAsActors() const;
+	AIMODULE_API TArray<AActor*> GetResultsAsActors() const;
 
 	/** DEPRECATED!  Use GetQueryResultsAsLocations instead.  Returns an array of locations generated by the query. If the query generated Actors the the array is filled with their locations. */
 	UFUNCTION(BlueprintPure, Category = "AI|EQS", Meta=(DeprecatedFunction,
 		DeprecationMessage="GetQueryResultsAsLocations instead, which is more efficient itself and protects against very bad perf issues in some blueprints."))
-	TArray<FVector> GetResultsAsLocations() const;
+	AIMODULE_API TArray<FVector> GetResultsAsLocations() const;
 
 	UFUNCTION(BlueprintCallable, Category = "AI|EQS")
-	void SetNamedParam(FName ParamName, float Value);
+	AIMODULE_API void SetNamedParam(FName ParamName, float Value);
 
-	void SetInstigator(const UObject* Object);
+	AIMODULE_API void SetInstigator(const UObject* Object);
 
 	FEQSQueryDoneSignature& GetOnQueryFinishedEvent() { return OnQueryFinishedEvent; }
 
 protected:
-	void OnQueryFinished(TSharedPtr<FEnvQueryResult> Result);
+	AIMODULE_API void OnQueryFinished(TSharedPtr<FEnvQueryResult> Result);
 
-	virtual bool IsSupportedForNetworking() const override;
+	AIMODULE_API virtual bool IsSupportedForNetworking() const override;
 };
diff --git a/Engine/Source/Runtime/AIModule/Classes/EnvironmentQuery/EnvQueryManager.h b/Engine/Source/Runtime/AIModule/Classes/EnvironmentQuery/EnvQueryManager.h
--- a/Engine/Source/Runtime/AIModule/Classes/EnvironmentQuery/EnvQueryManager.h
+++ b/Engine/Source/Runtime/AIModule/Classes/EnvironmentQuery/EnvQueryManager.h
@@ -67,7 +67,7 @@ public:
 
 /** wrapper for easy query execution */
 USTRUCT()
-struct AIMODULE_API FEnvQueryRequest
+struct FEnvQueryRequest
 {
 	GENERATED_USTRUCT_BODY()
 
@@ -81,14 +81,14 @@ struct AIMODULE_API FEnvQueryRequest
 
 	// set names param indicated by Param. If Param is configured to read values from a blackboard then BlackboardComponent
 	// is expected to be non-null (the function will fail a check otherwise).
-	FEnvQueryRequest& SetDynamicParam(const FAIDynamicParam& Param, const UBlackboardComponent* BlackboardComponent = nullptr);
+	AIMODULE_API FEnvQueryRequest& SetDynamicParam(const FAIDynamicParam& Param, const UBlackboardComponent* BlackboardComponent = nullptr);
 
 	// set named params
 	FORCEINLINE FEnvQueryRequest& SetFloatParam(FName ParamName, float Value) { NamedParams.Add(ParamName, Value); return *this; }
 	FORCEINLINE FEnvQueryRequest& SetIntParam(FName ParamName, int32 Value) { NamedParams.Add(ParamName, *((float*)&Value)); return *this; }
 	FORCEINLINE FEnvQueryRequest& SetBoolParam(FName ParamName, bool Value) { NamedParams.Add(ParamName, Value ? 1.0f : -1.0f); return *this; }
 	FORCEINLINE FEnvQueryRequest& SetNamedParam(const FEnvNamedValue& ParamData) { NamedParams.Add(ParamData.ParamName, ParamData.Value); return *this; }
-	FEnvQueryRequest& SetNamedParams(const TArray<FEnvNamedValue>& Params);
+	AIMODULE_API FEnvQueryRequest& SetNamedParams(const TArray<FEnvNamedValue>& Params);
 
 	// set world (for accessing query manager) when owner can't provide it
 	FORCEINLINE FEnvQueryRequest& SetWorldOverride(UWorld* InWorld) { World = InWorld; return *this; }
@@ -103,7 +103,7 @@ struct AIMODULE_API FEnvQueryRequest
 	{
 		return Execute(Mode, FQueryFinishedSignature::CreateUObject(InObj, InMethod));
 	}
-	int32 Execute(EEnvQueryRunMode::Type RunMode, FQueryFinishedSignature const& FinishDelegate);
+	AIMODULE_API int32 Execute(EEnvQueryRunMode::Type RunMode, FQueryFinishedSignature const& FinishDelegate);
 
 protected:
 
@@ -143,7 +143,7 @@ struct FEnvQueryInstanceCache
 };
 
 #if USE_EQS_DEBUGGER
-struct AIMODULE_API FEQSDebugger
+struct FEQSDebugger
 {
 	struct FEnvQueryInfo
 	{
@@ -177,14 +177,14 @@ struct AIMODULE_API FEQSDebugger
 		{}
 	};
 
-	void StoreStats(const FEnvQueryInstance& QueryInstance);
-	void StoreTickTime(const FEnvQueryInstance& QueryInstance, double TickTime, double MaxTickTime);
-	void StoreQuery(const TSharedPtr<FEnvQueryInstance>& QueryInstance);
+	AIMODULE_API void StoreStats(const FEnvQueryInstance& QueryInstance);
+	AIMODULE_API void StoreTickTime(const FEnvQueryInstance& QueryInstance, double TickTime, double MaxTickTime);
+	AIMODULE_API void StoreQuery(const TSharedPtr<FEnvQueryInstance>& QueryInstance);
 	
-	static void SaveStats(const FString& FileName);
-	static void LoadStats(const FString& FileName);
+	static AIMODULE_API void SaveStats(const FString& FileName);
+	static AIMODULE_API void LoadStats(const FString& FileName);
 
-	const TArray<FEnvQueryInfo>& GetAllQueriesForOwner(const UObject* Owner);
+	AIMODULE_API const TArray<FEnvQueryInfo>& GetAllQueriesForOwner(const UObject* Owner);
 
 	// map query name with profiler data
 	TMap<FName, FStatsInfo> StoredStats;
@@ -202,25 +202,25 @@ FORCEINLINE bool operator== (const FEQSDebugger::FEnvQueryInfo & Left, const FEQ
 }
 #endif // USE_EQS_DEBUGGER
 
-UCLASS(config = Game, defaultconfig, Transient)
-class AIMODULE_API UEnvQueryManager : public UAISubsystem, public FSelfRegisteringExec
+UCLASS(config = Game, defaultconfig, Transient, MinimalAPI)
+class UEnvQueryManager : public UAISubsystem, public FSelfRegisteringExec
 {
 	GENERATED_UCLASS_BODY()
 
 	// makes sure we don't have any UEnvQueryManager instances serialized in. 
 	// Any loaded instance will get marked as PendingKill
-	virtual void PostLoad() override;
-	virtual void PostInitProperties() override;
+	AIMODULE_API virtual void PostLoad() override;
+	AIMODULE_API virtual void PostInitProperties() override;
 
 	// FTickableGameObject begin
-	virtual void Tick(float DeltaTime) override;
-	virtual TStatId GetStatId() const override;
+	AIMODULE_API virtual void Tick(float DeltaTime) override;
+	AIMODULE_API virtual TStatId GetStatId() const override;
 	virtual bool IsTickableInEditor() const override { return true; }
 	// FTickableGameObject end
 
 	/** execute query */
-	int32 RunQuery(const FEnvQueryRequest& Request, EEnvQueryRunMode::Type RunMode, FQueryFinishedSignature const& FinishDelegate);
-	int32 RunQuery(const TSharedPtr<FEnvQueryInstance>& QueryInstance, FQueryFinishedSignature const& FinishDelegate);
+	AIMODULE_API int32 RunQuery(const FEnvQueryRequest& Request, EEnvQueryRunMode::Type RunMode, FQueryFinishedSignature const& FinishDelegate);
+	AIMODULE_API int32 RunQuery(const TSharedPtr<FEnvQueryInstance>& QueryInstance, FQueryFinishedSignature const& FinishDelegate);
 
 	/** Removed all active queries asked by Querier. No "on finished" notifications are being sent, call this function when
 	 *	you no longer care about Querier's queries, like when it's "dead" */
@@ -229,74 +229,74 @@ class AIMODULE_API UEnvQueryManager : public UAISubsystem, public FSelfRegisteri
 		RemoveAllQueriesByQuerier(Querier, /*bExecuteFinishDelegate=*/false);
 	}
 
-	void RemoveAllQueriesByQuerier(const UObject& Querier, bool bExecuteFinishDelegate);
+	AIMODULE_API void RemoveAllQueriesByQuerier(const UObject& Querier, bool bExecuteFinishDelegate);
 
 	/** alternative way to run queries. Do not use for anything other than testing
 	*  or when you know exactly what you're doing! Bypasses all EQS perf controlling
 	*  and time slicing mechanics. */
-	TSharedPtr<FEnvQueryResult> RunInstantQuery(const FEnvQueryRequest& Request, EEnvQueryRunMode::Type RunMode);
-	void RunInstantQuery(const TSharedPtr<FEnvQueryInstance>& QueryInstance);
+	AIMODULE_API TSharedPtr<FEnvQueryResult> RunInstantQuery(const FEnvQueryRequest& Request, EEnvQueryRunMode::Type RunMode);
+	AIMODULE_API void RunInstantQuery(const TSharedPtr<FEnvQueryInstance>& QueryInstance);
 
 	/** Creates a query instance configured for execution */
-	TSharedPtr<FEnvQueryInstance> PrepareQueryInstance(const FEnvQueryRequest& Request, EEnvQueryRunMode::Type RunMode);
+	AIMODULE_API TSharedPtr<FEnvQueryInstance> PrepareQueryInstance(const FEnvQueryRequest& Request, EEnvQueryRunMode::Type RunMode);
 
 	/** finds UEnvQuery matching QueryName by first looking at instantiated queries (from InstanceCache)
 	 *	falling back to looking up all UEnvQuery and testing their name */
-	UEnvQuery* FindQueryTemplate(const FString& QueryName) const;
+	AIMODULE_API UEnvQuery* FindQueryTemplate(const FString& QueryName) const;
 
 	/** creates local context object */
-	UEnvQueryContext* PrepareLocalContext(TSubclassOf<UEnvQueryContext> ContextClass);
+	AIMODULE_API UEnvQueryContext* PrepareLocalContext(TSubclassOf<UEnvQueryContext> ContextClass);
 
 	/** find value of named param stored with active query */
-	float FindNamedParam(int32 QueryId, FName ParamName) const;
+	AIMODULE_API float FindNamedParam(int32 QueryId, FName ParamName) const;
 
 	/** aborts specific query */
-	bool AbortQuery(int32 RequestID);
+	AIMODULE_API bool AbortQuery(int32 RequestID);
 
 	/** outputs active queries to log */
-	void PrintActiveQueryInfo() const;
+	AIMODULE_API void PrintActiveQueryInfo() const;
 
 	/** fail all running queries on cleaning the world */
-	virtual void OnWorldCleanup();
+	AIMODULE_API virtual void OnWorldCleanup();
 
 	/** cleanup hooks for map loading */
-	virtual void FinishDestroy() override;
+	AIMODULE_API virtual void FinishDestroy() override;
 
 	/** add information for data providers about query instance run independently */
-	void RegisterExternalQuery(const TSharedPtr<FEnvQueryInstance>& QueryInstance);
+	AIMODULE_API void RegisterExternalQuery(const TSharedPtr<FEnvQueryInstance>& QueryInstance);
 
 	/** clear information about query instance run independently */
-	void UnregisterExternalQuery(const TSharedPtr<FEnvQueryInstance>& QueryInstance);
+	AIMODULE_API void UnregisterExternalQuery(const TSharedPtr<FEnvQueryInstance>& QueryInstance);
 
 	/** list of all known item types */
-	static TArray<TSubclassOf<UEnvQueryItemType> > RegisteredItemTypes;
+	static AIMODULE_API TArray<TSubclassOf<UEnvQueryItemType> > RegisteredItemTypes;
 
-	static UEnvQueryManager* GetCurrent(UWorld* World);
-	static UEnvQueryManager* GetCurrent(const UObject* WorldContextObject);
+	static AIMODULE_API UEnvQueryManager* GetCurrent(UWorld* World);
+	static AIMODULE_API UEnvQueryManager* GetCurrent(const UObject* WorldContextObject);
 	
 	UFUNCTION(BlueprintCallable, Category = "AI|EQS", meta = (WorldContext = "WorldContextObject", AdvancedDisplay = "WrapperClass"))
-	static UEnvQueryInstanceBlueprintWrapper* RunEQSQuery(UObject* WorldContextObject, UEnvQuery* QueryTemplate, UObject* Querier, TEnumAsByte<EEnvQueryRunMode::Type> RunMode, TSubclassOf<UEnvQueryInstanceBlueprintWrapper> WrapperClass);
+	static AIMODULE_API UEnvQueryInstanceBlueprintWrapper* RunEQSQuery(UObject* WorldContextObject, UEnvQuery* QueryTemplate, UObject* Querier, TEnumAsByte<EEnvQueryRunMode::Type> RunMode, TSubclassOf<UEnvQueryInstanceBlueprintWrapper> WrapperClass);
 
-	void RegisterActiveWrapper(UEnvQueryInstanceBlueprintWrapper& Wrapper);
-	void UnregisterActiveWrapper(UEnvQueryInstanceBlueprintWrapper& Wrapper);
+	AIMODULE_API void RegisterActiveWrapper(UEnvQueryInstanceBlueprintWrapper& Wrapper);
+	AIMODULE_API void UnregisterActiveWrapper(UEnvQueryInstanceBlueprintWrapper& Wrapper);
 
-	static void SetAllowTimeSlicing(bool bAllowTimeSlicing);
+	static AIMODULE_API void SetAllowTimeSlicing(bool bAllowTimeSlicing);
 
 	/** Configure config variables during runtime */
-	void Configure(const FEnvQueryManagerConfig& NewConfig);
+	AIMODULE_API void Configure(const FEnvQueryManagerConfig& NewConfig);
 
 protected:
 	//~ Begin FExec Interface
-	virtual bool Exec_Dev(UWorld* Inworld,const TCHAR* Cmd,FOutputDevice& Ar) override;
+	AIMODULE_API virtual bool Exec_Dev(UWorld* Inworld,const TCHAR* Cmd,FOutputDevice& Ar) override;
 	//~ End FExec Interface
 
 	friend UEnvQueryInstanceBlueprintWrapper;
-	TSharedPtr<FEnvQueryInstance> FindQueryInstance(const int32 QueryID);
+	AIMODULE_API TSharedPtr<FEnvQueryInstance> FindQueryInstance(const int32 QueryID);
 
 #if USE_EQS_DEBUGGER
 public:
-	static void NotifyAssetUpdate(UEnvQuery* Query);
-	static TMap<FName, FEQSDebugger::FStatsInfo> DebuggerStats;
+	static AIMODULE_API void NotifyAssetUpdate(UEnvQuery* Query);
+	static AIMODULE_API TMap<FName, FEQSDebugger::FStatsInfo> DebuggerStats;
 
 	FEQSDebugger& GetDebugger() { return EQSDebugger; }
 
@@ -333,7 +333,7 @@ protected:
 	int32 NextQueryID;
 
 	/** create new instance, using cached data is possible */
-	TSharedPtr<FEnvQueryInstance> CreateQueryInstance(const UEnvQuery* Template, EEnvQueryRunMode::Type RunMode);
+	AIMODULE_API TSharedPtr<FEnvQueryInstance> CreateQueryInstance(const UEnvQuery* Template, EEnvQueryRunMode::Type RunMode);
 
 	/** how long are we allowed to test per update, in seconds. */
 	UPROPERTY(config)
@@ -367,22 +367,22 @@ protected:
 private:
 
 	/** create and bind delegates in instance */
-	void CreateOptionInstance(UEnvQueryOption* OptionTemplate, int32 SourceOptionIndex, const TArray<UEnvQueryTest*>& SortedTests, FEnvQueryInstance& Instance);
+	AIMODULE_API void CreateOptionInstance(UEnvQueryOption* OptionTemplate, int32 SourceOptionIndex, const TArray<UEnvQueryTest*>& SortedTests, FEnvQueryInstance& Instance);
 
 #if !(UE_BUILD_SHIPPING || UE_BUILD_TEST)
-	static bool bAllowEQSTimeSlicing;
+	static AIMODULE_API bool bAllowEQSTimeSlicing;
 #endif
 #if !(UE_BUILD_SHIPPING)
 	mutable double LastQueryCountWarningThresholdTime;
 
-	void CheckQueryCount() const;
-	void LogQueryInfo(bool bDisplayThresholdWarning) const;
+	AIMODULE_API void CheckQueryCount() const;
+	AIMODULE_API void LogQueryInfo(bool bDisplayThresholdWarning) const;
 #endif
 
 #if WITH_EDITOR
 	FDelegateHandle OnBlueprintCompiledHandle;
 
 	// used to reset LocalContextMap in case we had BP contexts stored
-	void OnBlueprintCompiled();
+	AIMODULE_API void OnBlueprintCompiled();
 #endif // WITH_EDITOR
 };
diff --git a/Engine/Source/Runtime/AIModule/Classes/EnvironmentQuery/EnvQueryNode.h b/Engine/Source/Runtime/AIModule/Classes/EnvironmentQuery/EnvQueryNode.h
--- a/Engine/Source/Runtime/AIModule/Classes/EnvironmentQuery/EnvQueryNode.h
+++ b/Engine/Source/Runtime/AIModule/Classes/EnvironmentQuery/EnvQueryNode.h
@@ -10,8 +10,8 @@
 
 struct FPropertyChangedEvent;
 
-UCLASS(Abstract)
-class AIMODULE_API UEnvQueryNode : public UObject
+UCLASS(Abstract, MinimalAPI)
+class UEnvQueryNode : public UObject
 {
 	GENERATED_UCLASS_BODY()
 
@@ -19,12 +19,12 @@ class AIMODULE_API UEnvQueryNode : public UObject
 	UPROPERTY()
 	int32 VerNum;
 
-	virtual void UpdateNodeVersion();
+	AIMODULE_API virtual void UpdateNodeVersion();
 
-	virtual FText GetDescriptionTitle() const;
-	virtual FText GetDescriptionDetails() const;
+	AIMODULE_API virtual FText GetDescriptionTitle() const;
+	AIMODULE_API virtual FText GetDescriptionDetails() const;
 
 #if WITH_EDITOR
-	virtual void PostEditChangeProperty(FPropertyChangedEvent& PropertyChangedEvent) override;
+	AIMODULE_API virtual void PostEditChangeProperty(FPropertyChangedEvent& PropertyChangedEvent) override;
 #endif //WITH_EDITOR
 };
diff --git a/Engine/Source/Runtime/AIModule/Classes/EnvironmentQuery/EnvQueryOption.h b/Engine/Source/Runtime/AIModule/Classes/EnvironmentQuery/EnvQueryOption.h
--- a/Engine/Source/Runtime/AIModule/Classes/EnvironmentQuery/EnvQueryOption.h
+++ b/Engine/Source/Runtime/AIModule/Classes/EnvironmentQuery/EnvQueryOption.h
@@ -10,8 +10,8 @@
 class UEnvQueryGenerator;
 class UEnvQueryTest;
 
-UCLASS()
-class AIMODULE_API UEnvQueryOption : public UObject
+UCLASS(MinimalAPI)
+class UEnvQueryOption : public UObject
 {
 	GENERATED_UCLASS_BODY()
 
@@ -21,6 +21,6 @@ class AIMODULE_API UEnvQueryOption : public UObject
 	UPROPERTY()
 	TArray<TObjectPtr<UEnvQueryTest>> Tests;
 
-	FText GetDescriptionTitle() const;
-	FText GetDescriptionDetails() const;
+	AIMODULE_API FText GetDescriptionTitle() const;
+	AIMODULE_API FText GetDescriptionDetails() const;
 };
diff --git a/Engine/Source/Runtime/AIModule/Classes/EnvironmentQuery/EnvQueryTest.h b/Engine/Source/Runtime/AIModule/Classes/EnvironmentQuery/EnvQueryTest.h
--- a/Engine/Source/Runtime/AIModule/Classes/EnvironmentQuery/EnvQueryTest.h
+++ b/Engine/Source/Runtime/AIModule/Classes/EnvironmentQuery/EnvQueryTest.h
@@ -57,8 +57,8 @@ struct FEnvQueryTestScoringPreview
 };
 #endif
 
-UCLASS(Abstract)
-class AIMODULE_API UEnvQueryTest : public UEnvQueryNode
+UCLASS(Abstract, MinimalAPI)
+class UEnvQueryTest : public UEnvQueryNode
 {
 	GENERATED_UCLASS_BODY()
 
@@ -150,7 +150,7 @@ class AIMODULE_API UEnvQueryTest : public UEnvQueryNode
 	FEnvQueryTestScoringPreview PreviewData;
 #endif
 
-	void SetWorkOnFloatValues(bool bWorkOnFloats);
+	AIMODULE_API void SetWorkOnFloatValues(bool bWorkOnFloats);
 	FORCEINLINE bool GetWorkOnFloatValues() const { return bWorkOnFloatValues; }
 
 	FORCEINLINE bool CanRunAsFinalCondition() const 
@@ -172,13 +172,13 @@ public:
 	virtual void RunTest(FEnvQueryInstance& QueryInstance) const { checkNoEntry(); }
 
 	/** check if test supports item type */
-	bool IsSupportedItem(TSubclassOf<UEnvQueryItemType> ItemType) const;
+	AIMODULE_API bool IsSupportedItem(TSubclassOf<UEnvQueryItemType> ItemType) const;
 
 	/** check if context needs to be updated for every item */
-	bool IsContextPerItem(TSubclassOf<UEnvQueryContext> CheckContext) const;
+	AIMODULE_API bool IsContextPerItem(TSubclassOf<UEnvQueryContext> CheckContext) const;
 
 	/** helper: get location of item */
-	FVector GetItemLocation(FEnvQueryInstance& QueryInstance, int32 ItemIndex) const;
+	AIMODULE_API FVector GetItemLocation(FEnvQueryInstance& QueryInstance, int32 ItemIndex) const;
 
 	/** helper: get location of item */
 	FORCEINLINE FVector GetItemLocation(FEnvQueryInstance& QueryInstance, const FEnvQueryInstance::ItemIterator& Iterator) const
@@ -187,7 +187,7 @@ public:
 	}
 
 	/** helper: get location of item */
-	FRotator GetItemRotation(FEnvQueryInstance& QueryInstance, int32 ItemIndex) const;
+	AIMODULE_API FRotator GetItemRotation(FEnvQueryInstance& QueryInstance, int32 ItemIndex) const;
 
 	/** helper: get location of item */
 	FORCEINLINE FRotator GetItemRotation(FEnvQueryInstance& QueryInstance, const FEnvQueryInstance::ItemIterator& Iterator) const
@@ -196,7 +196,7 @@ public:
 	}
 
 	/** helper: get actor from item */
-	AActor* GetItemActor(FEnvQueryInstance& QueryInstance, int32 ItemIndex) const;
+	AIMODULE_API AActor* GetItemActor(FEnvQueryInstance& QueryInstance, int32 ItemIndex) const;
 		
 	/** helper: get actor from item */
 	FORCEINLINE AActor* GetItemActor(FEnvQueryInstance& QueryInstance, const FEnvQueryInstance::ItemIterator& Iterator) const
@@ -205,25 +205,25 @@ public:
 	}
 
 	/** normalize scores in range */
-	void NormalizeItemScores(FEnvQueryInstance& QueryInstance);
+	AIMODULE_API void NormalizeItemScores(FEnvQueryInstance& QueryInstance);
 
 	FORCEINLINE bool IsScoring() const { return (TestPurpose != EEnvTestPurpose::Filter); } 
 	FORCEINLINE bool IsFiltering() const { return (TestPurpose != EEnvTestPurpose::Score); }
 
-	FText DescribeFloatTestParams() const;
-	FText DescribeBoolTestParams(const FString& ConditionDesc) const;
+	AIMODULE_API FText DescribeFloatTestParams() const;
+	AIMODULE_API FText DescribeBoolTestParams(const FString& ConditionDesc) const;
 
-	virtual void PostLoad() override;
+	AIMODULE_API virtual void PostLoad() override;
 
 	/** update to latest version after spawning */
-	virtual void UpdateNodeVersion() override;
+	AIMODULE_API virtual void UpdateNodeVersion() override;
 
 #if WITH_EDITOR
-	virtual void PostEditChangeProperty(FPropertyChangedEvent& PropertyChangedEvent) override;
+	AIMODULE_API virtual void PostEditChangeProperty(FPropertyChangedEvent& PropertyChangedEvent) override;
 #endif
 
 	/** update preview list */
-	void UpdatePreviewData();
+	AIMODULE_API void UpdatePreviewData();
 };
 
 //////////////////////////////////////////////////////////////////////////
diff --git a/Engine/Source/Runtime/AIModule/Classes/EnvironmentQuery/EnvQueryTypes.h b/Engine/Source/Runtime/AIModule/Classes/EnvironmentQuery/EnvQueryTypes.h
--- a/Engine/Source/Runtime/AIModule/Classes/EnvironmentQuery/EnvQueryTypes.h
+++ b/Engine/Source/Runtime/AIModule/Classes/EnvironmentQuery/EnvQueryTypes.h
@@ -269,7 +269,7 @@ namespace EEnvQueryTestClamping
 }
 
 USTRUCT(BlueprintType)
-struct AIMODULE_API FEnvNamedValue
+struct FEnvNamedValue
 {
 	GENERATED_USTRUCT_BODY()
 
@@ -288,7 +288,7 @@ struct AIMODULE_API FEnvNamedValue
 };
 
 USTRUCT()
-struct AIMODULE_API FEnvDirection
+struct FEnvDirection
 {
 	GENERATED_USTRUCT_BODY()
 
@@ -308,11 +308,11 @@ struct AIMODULE_API FEnvDirection
 	UPROPERTY(EditDefaultsOnly, Category=Direction, meta=(DisplayName="Mode"))
 	TEnumAsByte<EEnvDirection::Type> DirMode = EEnvDirection::Type::TwoPoints;
 
-	FText ToText() const;
+	AIMODULE_API FText ToText() const;
 };
 
 USTRUCT()
-struct AIMODULE_API FEnvTraceData
+struct FEnvTraceData
 {
 	GENERATED_USTRUCT_BODY()
 
@@ -411,16 +411,16 @@ struct AIMODULE_API FEnvTraceData
 	UPROPERTY(EditDefaultsOnly, Category=Trace)
 	uint32 bCanProjectDown : 1;
 
-	FText ToText(EDescriptionMode DescMode) const;
+	AIMODULE_API FText ToText(EDescriptionMode DescMode) const;
 
-	void SetGeometryOnly();
-	void SetNavmeshOnly();
+	AIMODULE_API void SetGeometryOnly();
+	AIMODULE_API void SetNavmeshOnly();
 	
-	void OnPostLoad();
+	AIMODULE_API void OnPostLoad();
 };
 
 USTRUCT()
-struct AIMODULE_API FEnvOverlapData
+struct FEnvOverlapData
 {
 	GENERATED_USTRUCT_BODY()
 
@@ -477,7 +477,7 @@ struct AIMODULE_API FEnvOverlapData
 //////////////////////////////////////////////////////////////////////////
 // Returned results
 
-struct AIMODULE_API FEnvQueryItem
+struct FEnvQueryItem
 {
 	/** total score of item */
 	float Score;
@@ -511,7 +511,7 @@ struct AIMODULE_API FEnvQueryItem
 template <> struct TIsZeroConstructType<FEnvQueryItem> { enum { Value = true }; };
 
 USTRUCT(BlueprintType)
-struct AIMODULE_API FEnvQueryResult
+struct FEnvQueryResult
 {
 	GENERATED_USTRUCT_BODY()
 
@@ -543,8 +543,8 @@ public:
 	FORCEINLINE float GetItemScore(int32 Index) const { return Items.IsValidIndex(Index) ? Items[Index].Score : 0.0f; }
 
 	/** item accessors for basic types */
-	AActor* GetItemAsActor(int32 Index) const;
-	FVector GetItemAsLocation(int32 Index) const;
+	AIMODULE_API AActor* GetItemAsActor(int32 Index) const;
+	AIMODULE_API FVector GetItemAsLocation(int32 Index) const;
 
 	template<typename TItemType>
 	const typename TItemType::FValueType& GetItemAsTypeChecked(const int32 Index) const
@@ -557,8 +557,8 @@ public:
 	}
 
 	/** note that this function does not strip out the null-actors to not mess up results of GetItemScore(Index) calls*/
-	void GetAllAsActors(TArray<AActor*>& OutActors) const;
-	void GetAllAsLocations(TArray<FVector>& OutLocations) const;
+	AIMODULE_API void GetAllAsActors(TArray<AActor*>& OutActors) const;
+	AIMODULE_API void GetAllAsLocations(TArray<FVector>& OutLocations) const;
 
 	const uint8* GetItemRawMemory(const int32 Index) const
 	{
@@ -590,14 +590,14 @@ public:
 
 DECLARE_DELEGATE_OneParam(FQueryFinishedSignature, TSharedPtr<FEnvQueryResult>);
 
-struct AIMODULE_API FEnvQuerySpatialData
+struct FEnvQuerySpatialData
 {
 	FVector Location;
 	FRotator Rotation;
 };
 
 /** Detailed information about item, used by tests */
-struct AIMODULE_API FEnvQueryItemDetails
+struct FEnvQueryItemDetails
 {
 	/** Results assigned by option's tests, before any modifications */
 	TArray<float> TestResults;
@@ -632,7 +632,7 @@ struct AIMODULE_API FEnvQueryItemDetails
 	}
 };
 
-struct AIMODULE_API FEnvQueryContextData
+struct FEnvQueryContextData
 {
 	/** type of context values */
 	TSubclassOf<UEnvQueryItemType> ValueType;
@@ -650,7 +650,7 @@ struct AIMODULE_API FEnvQueryContextData
 	FORCEINLINE uint32 GetAllocatedSize() const { return sizeof(*this) + RawData.GetAllocatedSize(); }
 };
 
-struct AIMODULE_API FEnvQueryOptionInstance
+struct FEnvQueryOptionInstance
 {
 	/** generator object, raw pointer can be used safely because it will be always referenced by EnvQueryManager */
 	UEnvQueryGenerator* Generator;
@@ -742,23 +742,23 @@ struct FEQSQueryDebugData : public FEnvQueryDebugData
 {
 };
 
-UCLASS(Abstract)
-class AIMODULE_API UEnvQueryTypes : public UObject
+UCLASS(Abstract, MinimalAPI)
+class UEnvQueryTypes : public UObject
 {
 	GENERATED_BODY()
 
 public:
 	/** special test value assigned to items skipped by condition check */
-	static const float SkippedItemValue;
+	static AIMODULE_API const float SkippedItemValue;
 
 	/** special value used for executing query steps to prevent them from being time sliced */
-	static float UnlimitedStepTime;
+	static AIMODULE_API float UnlimitedStepTime;
 
-	static FText GetShortTypeName(const UObject* Ob);
-	static FText DescribeContext(TSubclassOf<UEnvQueryContext> ContextClass);
+	static AIMODULE_API FText GetShortTypeName(const UObject* Ob);
+	static AIMODULE_API FText DescribeContext(TSubclassOf<UEnvQueryContext> ContextClass);
 };
 
-struct AIMODULE_API FEnvQueryInstance : public FEnvQueryResult
+struct FEnvQueryInstance : public FEnvQueryResult
 {
 	typedef float FNamedParamValueType;
 
@@ -839,29 +839,29 @@ struct AIMODULE_API FEnvQueryInstance : public FEnvQueryResult
 	/** item type's CDO for actor tests */
 	UEnvQueryItemType_ActorBase* ItemTypeActorCDO;
 
-	FEnvQueryInstance();
-	FEnvQueryInstance(const FEnvQueryInstance& Other);
-	~FEnvQueryInstance();
+	AIMODULE_API FEnvQueryInstance();
+	AIMODULE_API FEnvQueryInstance(const FEnvQueryInstance& Other);
+	AIMODULE_API ~FEnvQueryInstance();
 
 	/** execute single step of query */
-	void ExecuteOneStep(double TimeLimit);
+	AIMODULE_API void ExecuteOneStep(double TimeLimit);
 
 	/** update context cache */
-	bool PrepareContext(UClass* Context, FEnvQueryContextData& ContextData);
+	AIMODULE_API bool PrepareContext(UClass* Context, FEnvQueryContextData& ContextData);
 
 	/** helpers for reading spatial data from context */
-	bool PrepareContext(UClass* Context, TArray<FEnvQuerySpatialData>& Data);
-	bool PrepareContext(UClass* Context, TArray<FVector>& Data);
-	bool PrepareContext(UClass* Context, TArray<FRotator>& Data);
+	AIMODULE_API bool PrepareContext(UClass* Context, TArray<FEnvQuerySpatialData>& Data);
+	AIMODULE_API bool PrepareContext(UClass* Context, TArray<FVector>& Data);
+	AIMODULE_API bool PrepareContext(UClass* Context, TArray<FRotator>& Data);
 	/** helpers for reading actor data from context */
-	bool PrepareContext(UClass* Context, TArray<AActor*>& Data);
+	AIMODULE_API bool PrepareContext(UClass* Context, TArray<AActor*>& Data);
 	
 	bool IsInSingleItemFinalSearch() const { return !!bPassOnSingleResult; }
 	/** check if current test can batch its calculations */
 	bool CanBatchTest() const { return !IsInSingleItemFinalSearch(); }
 
 	/** raw data operations */
-	void ReserveItemData(int32 NumAdditionalItems);
+	AIMODULE_API void ReserveItemData(int32 NumAdditionalItems);
 
 	template<typename TypeItem>
 	void AddItemData(typename TypeItem::FValueType ItemValue)
@@ -925,30 +925,30 @@ struct AIMODULE_API FEnvQueryInstance : public FEnvQueryResult
 protected:
 
 	/** prepare item data after generator has finished */
-	void FinalizeGeneration();
+	AIMODULE_API void FinalizeGeneration();
 
 	/** update costs and flags after test has finished */
-	void FinalizeTest();
+	AIMODULE_API void FinalizeTest();
 	
 	/** final pass on items of finished query */
-	void FinalizeQuery();
+	AIMODULE_API void FinalizeQuery();
 
 	/** normalize total score in range 0..1 */
-	void NormalizeScores();
+	AIMODULE_API void NormalizeScores();
 
 	/** sort all scores, from highest to lowest */
-	void SortScores();
+	AIMODULE_API void SortScores();
 
 	/** pick one of items with score equal or higher than specified */
-	void PickRandomItemOfScoreAtLeast(float MinScore);
+	AIMODULE_API void PickRandomItemOfScoreAtLeast(float MinScore);
 
 	/** discard all items but one */
-	void PickSingleItem(int32 ItemIndex);
+	AIMODULE_API void PickSingleItem(int32 ItemIndex);
 
 public:
 
 	/** removes all runtime data that can be used for debugging (not a part of actual query result) */
-	void StripRedundantData();
+	AIMODULE_API void StripRedundantData();
 
 #if STATS
 	FORCEINLINE void IncStats()
@@ -963,8 +963,8 @@ public:
 		DEC_DWORD_STAT_BY(STAT_AI_EQS_NumItems, Items.Num());
 	}
 
-	uint32 GetAllocatedSize() const;
-	uint32 GetContextAllocatedSize() const;
+	AIMODULE_API uint32 GetAllocatedSize() const;
+	AIMODULE_API uint32 GetContextAllocatedSize() const;
 #else
 	FORCEINLINE uint32 GetAllocatedSize() const { return 0; }
 	FORCEINLINE uint32 GetContextAllocatedSize() const { return 0; }
@@ -973,7 +973,7 @@ public:
 #endif // STATS
 
 #if !NO_LOGGING
-	void Log(const FString Msg) const;
+	AIMODULE_API void Log(const FString Msg) const;
 #endif // #if !NO_LOGGING
 
 #if USE_EQS_DEBUGGER
@@ -992,12 +992,12 @@ public:
 #endif
 
 	/** describe for logging purposes what the query spent time on */
-	FString GetExecutionTimeDescription() const;
+	AIMODULE_API FString GetExecutionTimeDescription() const;
 
 #if CPP || UE_BUILD_DOCS
 	/** Note that this iterator is for read-only purposes. Please use FItemIterator for regular item iteration 
 	 *	while performing EQS testing and scoring */
-	struct AIMODULE_API FConstItemIterator
+	struct FConstItemIterator
 	{
 		FConstItemIterator(FEnvQueryInstance& QueryInstance, int32 StartingItemIndex = INDEX_NONE)
 			: Instance(QueryInstance)
@@ -1274,10 +1274,10 @@ public:
 
 #if USE_EQS_DEBUGGER
 	FEnvQueryDebugData DebugData;
-	static bool bDebuggingInfoEnabled;
+	static AIMODULE_API bool bDebuggingInfoEnabled;
 #endif // USE_EQS_DEBUGGER
 
-	FBox GetBoundingBox() const;
+	AIMODULE_API FBox GetBoundingBox() const;
 };
 
 namespace FEQSHelpers
@@ -1286,7 +1286,7 @@ namespace FEQSHelpers
 }
 
 USTRUCT(BlueprintType)
-struct AIMODULE_API FAIDynamicParam
+struct FAIDynamicParam
 {
 	GENERATED_USTRUCT_BODY();
 
@@ -1309,17 +1309,17 @@ struct AIMODULE_API FAIDynamicParam
 		BBKey.AllowNoneAsValue(true);
 	}
 
-	void ConfigureBBKey(UObject &QueryOwner);
+	AIMODULE_API void ConfigureBBKey(UObject &QueryOwner);
 
-	static void GenerateConfigurableParamsFromNamedValues(UObject &QueryOwner, TArray<FAIDynamicParam>& OutQueryConfig, TArray<FEnvNamedValue>& InQueryParams);
+	static AIMODULE_API void GenerateConfigurableParamsFromNamedValues(UObject &QueryOwner, TArray<FAIDynamicParam>& OutQueryConfig, TArray<FEnvNamedValue>& InQueryParams);
 };
 
 USTRUCT()
-struct AIMODULE_API FEQSParametrizedQueryExecutionRequest
+struct FEQSParametrizedQueryExecutionRequest
 {
 	GENERATED_USTRUCT_BODY()
 
-	FEQSParametrizedQueryExecutionRequest();
+	AIMODULE_API FEQSParametrizedQueryExecutionRequest();
 
 	UPROPERTY(Category = Node, EditAnywhere, meta = (EditCondition = "!bUseBBKeyForQueryTemplate"))
 	TObjectPtr<UEnvQuery> QueryTemplate;
@@ -1340,13 +1340,13 @@ struct AIMODULE_API FEQSParametrizedQueryExecutionRequest
 
 	uint32 bInitialized : 1;
 
-	void InitForOwnerAndBlackboard(UObject& Owner, UBlackboardData* BBAsset);
+	AIMODULE_API void InitForOwnerAndBlackboard(UObject& Owner, UBlackboardData* BBAsset);
 
 	bool IsValid() const { return bInitialized; }
 
-	int32 Execute(UObject& QueryOwner, const UBlackboardComponent* BlackboardComponent, FQueryFinishedSignature& QueryFinishedDelegate);
+	AIMODULE_API int32 Execute(UObject& QueryOwner, const UBlackboardComponent* BlackboardComponent, FQueryFinishedSignature& QueryFinishedDelegate);
 
 #if WITH_EDITOR
-	void PostEditChangeProperty(UObject& Owner, struct FPropertyChangedEvent& PropertyChangedEvent);
+	AIMODULE_API void PostEditChangeProperty(UObject& Owner, struct FPropertyChangedEvent& PropertyChangedEvent);
 #endif // WITH_EDITOR
 };
diff --git a/Engine/Source/Runtime/AIModule/Classes/EnvironmentQuery/Generators/EnvQueryGenerator_ActorsOfClass.h b/Engine/Source/Runtime/AIModule/Classes/EnvironmentQuery/Generators/EnvQueryGenerator_ActorsOfClass.h
--- a/Engine/Source/Runtime/AIModule/Classes/EnvironmentQuery/Generators/EnvQueryGenerator_ActorsOfClass.h
+++ b/Engine/Source/Runtime/AIModule/Classes/EnvironmentQuery/Generators/EnvQueryGenerator_ActorsOfClass.h
@@ -11,8 +11,8 @@
 #include "EnvironmentQuery/EnvQueryGenerator.h"
 #include "EnvQueryGenerator_ActorsOfClass.generated.h"
 
-UCLASS(meta = (DisplayName = "Actors Of Class"))
-class AIMODULE_API UEnvQueryGenerator_ActorsOfClass : public UEnvQueryGenerator
+UCLASS(meta = (DisplayName = "Actors Of Class"), MinimalAPI)
+class UEnvQueryGenerator_ActorsOfClass : public UEnvQueryGenerator
 {
 	GENERATED_UCLASS_BODY()
 
@@ -33,9 +33,9 @@ class AIMODULE_API UEnvQueryGenerator_ActorsOfClass : public UEnvQueryGenerator
 	UPROPERTY(EditAnywhere, Category=Generator)
 	TSubclassOf<UEnvQueryContext> SearchCenter;
 
-	virtual void GenerateItems(FEnvQueryInstance& QueryInstance) const override;
+	AIMODULE_API virtual void GenerateItems(FEnvQueryInstance& QueryInstance) const override;
 	virtual void ProcessItems(FEnvQueryInstance& QueryInstance, TArray<AActor*>& MatchingActors) const {}
 
-	virtual FText GetDescriptionTitle() const override;
-	virtual FText GetDescriptionDetails() const override;
+	AIMODULE_API virtual FText GetDescriptionTitle() const override;
+	AIMODULE_API virtual FText GetDescriptionDetails() const override;
 };
diff --git a/Engine/Source/Runtime/AIModule/Classes/EnvironmentQuery/Generators/EnvQueryGenerator_BlueprintBase.h b/Engine/Source/Runtime/AIModule/Classes/EnvironmentQuery/Generators/EnvQueryGenerator_BlueprintBase.h
--- a/Engine/Source/Runtime/AIModule/Classes/EnvironmentQuery/Generators/EnvQueryGenerator_BlueprintBase.h
+++ b/Engine/Source/Runtime/AIModule/Classes/EnvironmentQuery/Generators/EnvQueryGenerator_BlueprintBase.h
@@ -13,8 +13,8 @@
 
 class AActor;
 
-UCLASS(Abstract, Blueprintable)
-class AIMODULE_API UEnvQueryGenerator_BlueprintBase : public UEnvQueryGenerator
+UCLASS(Abstract, Blueprintable, MinimalAPI)
+class UEnvQueryGenerator_BlueprintBase : public UEnvQueryGenerator
 {
 	GENERATED_UCLASS_BODY()
 
@@ -32,28 +32,28 @@ class AIMODULE_API UEnvQueryGenerator_BlueprintBase : public UEnvQueryGenerator
 	UPROPERTY(EditDefaultsOnly, Category = Generator)
 	TSubclassOf<UEnvQueryItemType> GeneratedItemType;
 
-	virtual void PostInitProperties() override;
-	virtual UWorld* GetWorld() const override;
+	AIMODULE_API virtual void PostInitProperties() override;
+	AIMODULE_API virtual UWorld* GetWorld() const override;
 
 	UFUNCTION(BlueprintImplementableEvent, Category = Generator)
-	void DoItemGeneration(const TArray<FVector>& ContextLocations) const;
+	AIMODULE_API void DoItemGeneration(const TArray<FVector>& ContextLocations) const;
 
 	UFUNCTION(BlueprintImplementableEvent, Category = Generator)
-	void DoItemGenerationFromActors(const TArray<AActor*>& ContextActors) const;
+	AIMODULE_API void DoItemGenerationFromActors(const TArray<AActor*>& ContextActors) const;
 
-	virtual void GenerateItems(FEnvQueryInstance& QueryInstance) const override;
+	AIMODULE_API virtual void GenerateItems(FEnvQueryInstance& QueryInstance) const override;
 
-	virtual FText GetDescriptionTitle() const override;
-	virtual FText GetDescriptionDetails() const override;
+	AIMODULE_API virtual FText GetDescriptionTitle() const override;
+	AIMODULE_API virtual FText GetDescriptionDetails() const override;
 	
 	UFUNCTION(BlueprintCallable, Category = "EQS")
-	void AddGeneratedVector(FVector GeneratedVector) const;
+	AIMODULE_API void AddGeneratedVector(FVector GeneratedVector) const;
 
 	UFUNCTION(BlueprintCallable, Category = "EQS")
-	void AddGeneratedActor(AActor* GeneratedActor) const;
+	AIMODULE_API void AddGeneratedActor(AActor* GeneratedActor) const;
 
 	UFUNCTION(BlueprintCallable, Category = "EQS")
-	UObject* GetQuerier() const;
+	AIMODULE_API UObject* GetQuerier() const;
 
 private:
 	enum class ECallMode
diff --git a/Engine/Source/Runtime/AIModule/Classes/EnvironmentQuery/Generators/EnvQueryGenerator_Composite.h b/Engine/Source/Runtime/AIModule/Classes/EnvironmentQuery/Generators/EnvQueryGenerator_Composite.h
--- a/Engine/Source/Runtime/AIModule/Classes/EnvironmentQuery/Generators/EnvQueryGenerator_Composite.h
+++ b/Engine/Source/Runtime/AIModule/Classes/EnvironmentQuery/Generators/EnvQueryGenerator_Composite.h
@@ -12,8 +12,8 @@
  * All child generators must produce exactly the same item type!
  */
 
-UCLASS(meta = (DisplayName = "Composite"))
-class AIMODULE_API UEnvQueryGenerator_Composite : public UEnvQueryGenerator
+UCLASS(meta = (DisplayName = "Composite"), MinimalAPI)
+class UEnvQueryGenerator_Composite : public UEnvQueryGenerator
 {
 	GENERATED_UCLASS_BODY()
 
@@ -37,10 +37,10 @@ class AIMODULE_API UEnvQueryGenerator_Composite : public UEnvQueryGenerator
 	UPROPERTY(EditDefaultsOnly, Category = Generator, AdvancedDisplay)
 	TSubclassOf<UEnvQueryItemType> ForcedItemType;
 
-	void VerifyItemTypes();
+	AIMODULE_API void VerifyItemTypes();
 
 protected:
-	virtual void GenerateItems(FEnvQueryInstance& QueryInstance) const override;
-	virtual FText GetDescriptionTitle() const override;
-	virtual bool IsValidGenerator() const override;
+	AIMODULE_API virtual void GenerateItems(FEnvQueryInstance& QueryInstance) const override;
+	AIMODULE_API virtual FText GetDescriptionTitle() const override;
+	AIMODULE_API virtual bool IsValidGenerator() const override;
 };
diff --git a/Engine/Source/Runtime/AIModule/Classes/EnvironmentQuery/Generators/EnvQueryGenerator_Cone.h b/Engine/Source/Runtime/AIModule/Classes/EnvironmentQuery/Generators/EnvQueryGenerator_Cone.h
--- a/Engine/Source/Runtime/AIModule/Classes/EnvironmentQuery/Generators/EnvQueryGenerator_Cone.h
+++ b/Engine/Source/Runtime/AIModule/Classes/EnvironmentQuery/Generators/EnvQueryGenerator_Cone.h
@@ -11,25 +11,25 @@
 #include "EnvironmentQuery/Generators/EnvQueryGenerator_ProjectedPoints.h"
 #include "EnvQueryGenerator_Cone.generated.h"
 
-UCLASS(meta = (DisplayName = "Points: Cone"))
-class AIMODULE_API UEnvQueryGenerator_Cone : public UEnvQueryGenerator_ProjectedPoints
+UCLASS(meta = (DisplayName = "Points: Cone"), MinimalAPI)
+class UEnvQueryGenerator_Cone : public UEnvQueryGenerator_ProjectedPoints
 {
 	GENERATED_BODY()
 public:
-	UEnvQueryGenerator_Cone(const FObjectInitializer& ObjectInitializer = FObjectInitializer::Get());
+	AIMODULE_API UEnvQueryGenerator_Cone(const FObjectInitializer& ObjectInitializer = FObjectInitializer::Get());
 
 	/** Binds data to data providers 
 	 *  @param QueryInstance - the instance of the query
 	 */
-	void BindDataToDataProviders(FEnvQueryInstance& QueryInstance) const;
+	AIMODULE_API void BindDataToDataProviders(FEnvQueryInstance& QueryInstance) const;
 	
-	virtual void GenerateItems(FEnvQueryInstance& QueryInstance) const override;
+	AIMODULE_API virtual void GenerateItems(FEnvQueryInstance& QueryInstance) const override;
 
 	/** Returns the title of the generator on the corresponding node in the EQS Editor window */
-	virtual FText GetDescriptionTitle() const override;
+	AIMODULE_API virtual FText GetDescriptionTitle() const override;
 
 	/** Returns the details of the generator on the corresponding node in the EQS Editor window */
-	virtual FText GetDescriptionDetails() const override;
+	AIMODULE_API virtual FText GetDescriptionDetails() const override;
 
 protected:
 	/** Distance between each point of the same angle */
diff --git a/Engine/Source/Runtime/AIModule/Classes/EnvironmentQuery/Generators/EnvQueryGenerator_CurrentLocation.h b/Engine/Source/Runtime/AIModule/Classes/EnvironmentQuery/Generators/EnvQueryGenerator_CurrentLocation.h
--- a/Engine/Source/Runtime/AIModule/Classes/EnvironmentQuery/Generators/EnvQueryGenerator_CurrentLocation.h
+++ b/Engine/Source/Runtime/AIModule/Classes/EnvironmentQuery/Generators/EnvQueryGenerator_CurrentLocation.h
@@ -9,8 +9,8 @@
 #include "EnvironmentQuery/EnvQueryGenerator.h"
 #include "EnvQueryGenerator_CurrentLocation.generated.h"
 
-UCLASS(meta = (DisplayName = "Current Location"))
-class AIMODULE_API UEnvQueryGenerator_CurrentLocation : public UEnvQueryGenerator
+UCLASS(meta = (DisplayName = "Current Location"), MinimalAPI)
+class UEnvQueryGenerator_CurrentLocation : public UEnvQueryGenerator
 {
 	GENERATED_BODY()
 
@@ -21,10 +21,10 @@ protected:
 
 public:
 
-	UEnvQueryGenerator_CurrentLocation(const FObjectInitializer& ObjectInitializer);
+	AIMODULE_API UEnvQueryGenerator_CurrentLocation(const FObjectInitializer& ObjectInitializer);
 	
-	virtual void GenerateItems(FEnvQueryInstance& QueryInstance) const override;
+	AIMODULE_API virtual void GenerateItems(FEnvQueryInstance& QueryInstance) const override;
 
-	virtual FText GetDescriptionTitle() const override;
-	virtual FText GetDescriptionDetails() const override;
+	AIMODULE_API virtual FText GetDescriptionTitle() const override;
+	AIMODULE_API virtual FText GetDescriptionDetails() const override;
 };
diff --git a/Engine/Source/Runtime/AIModule/Classes/EnvironmentQuery/Generators/EnvQueryGenerator_Donut.h b/Engine/Source/Runtime/AIModule/Classes/EnvironmentQuery/Generators/EnvQueryGenerator_Donut.h
--- a/Engine/Source/Runtime/AIModule/Classes/EnvironmentQuery/Generators/EnvQueryGenerator_Donut.h
+++ b/Engine/Source/Runtime/AIModule/Classes/EnvironmentQuery/Generators/EnvQueryGenerator_Donut.h
@@ -10,8 +10,8 @@
 #include "EnvironmentQuery/Generators/EnvQueryGenerator_ProjectedPoints.h"
 #include "EnvQueryGenerator_Donut.generated.h"
 
-UCLASS(meta = (DisplayName = "Points: Donut"))
-class AIMODULE_API UEnvQueryGenerator_Donut : public UEnvQueryGenerator_ProjectedPoints
+UCLASS(meta = (DisplayName = "Points: Donut"), MinimalAPI)
+class UEnvQueryGenerator_Donut : public UEnvQueryGenerator_ProjectedPoints
 {
 	GENERATED_UCLASS_BODY()
 
@@ -51,16 +51,16 @@ class AIMODULE_API UEnvQueryGenerator_Donut : public UEnvQueryGenerator_Projecte
 	UPROPERTY(EditAnywhere, Category = Generator, meta=(InlineEditConditionToggle))
 	uint32 bDefineArc : 1;
 
-	virtual void GenerateItems(FEnvQueryInstance& QueryInstance) const override;
+	AIMODULE_API virtual void GenerateItems(FEnvQueryInstance& QueryInstance) const override;
 
-	virtual FText GetDescriptionTitle() const override;
-	virtual FText GetDescriptionDetails() const override;
+	AIMODULE_API virtual FText GetDescriptionTitle() const override;
+	AIMODULE_API virtual FText GetDescriptionDetails() const override;
 
 #if WITH_EDITOR
-	virtual void PostEditChangeProperty(struct FPropertyChangedEvent& PropertyChangedEvent) override;
+	AIMODULE_API virtual void PostEditChangeProperty(struct FPropertyChangedEvent& PropertyChangedEvent) override;
 #endif // WITH_EDITOR
 
 protected:
-	FVector::FReal GetArcBisectorAngle(FEnvQueryInstance& QueryInstance) const;
-	bool IsAngleAllowed(FVector::FReal TestAngleRad, FVector::FReal BisectAngleDeg, FVector::FReal AngleRangeDeg, bool bConstrainAngle) const;
+	AIMODULE_API FVector::FReal GetArcBisectorAngle(FEnvQueryInstance& QueryInstance) const;
+	AIMODULE_API bool IsAngleAllowed(FVector::FReal TestAngleRad, FVector::FReal BisectAngleDeg, FVector::FReal AngleRangeDeg, bool bConstrainAngle) const;
 };
diff --git a/Engine/Source/Runtime/AIModule/Classes/EnvironmentQuery/Generators/EnvQueryGenerator_OnCircle.h b/Engine/Source/Runtime/AIModule/Classes/EnvironmentQuery/Generators/EnvQueryGenerator_OnCircle.h
--- a/Engine/Source/Runtime/AIModule/Classes/EnvironmentQuery/Generators/EnvQueryGenerator_OnCircle.h
+++ b/Engine/Source/Runtime/AIModule/Classes/EnvironmentQuery/Generators/EnvQueryGenerator_OnCircle.h
@@ -24,8 +24,8 @@ enum class EPointOnCircleSpacingMethod :uint8
 	ByNumberOfPoints
 };
 
-UCLASS(meta = (DisplayName = "Points: Circle"))
-class AIMODULE_API UEnvQueryGenerator_OnCircle : public UEnvQueryGenerator_ProjectedPoints
+UCLASS(meta = (DisplayName = "Points: Circle"), MinimalAPI)
+class UEnvQueryGenerator_OnCircle : public UEnvQueryGenerator_ProjectedPoints
 {
 	GENERATED_UCLASS_BODY()
 
@@ -75,25 +75,25 @@ class AIMODULE_API UEnvQueryGenerator_OnCircle : public UEnvQueryGenerator_Proje
 	UPROPERTY(EditAnywhere, Category=Generator, meta=(InlineEditConditionToggle))
 	uint32 bDefineArc:1;
 
-	virtual void PostLoad() override;
+	AIMODULE_API virtual void PostLoad() override;
 
-	virtual void GenerateItems(FEnvQueryInstance& QueryInstance) const override;
+	AIMODULE_API virtual void GenerateItems(FEnvQueryInstance& QueryInstance) const override;
 
-	virtual FText GetDescriptionTitle() const override;
-	virtual FText GetDescriptionDetails() const override;
+	AIMODULE_API virtual FText GetDescriptionTitle() const override;
+	AIMODULE_API virtual FText GetDescriptionDetails() const override;
 
 #if WITH_EDITOR
-	virtual void PostEditChangeProperty( struct FPropertyChangedEvent& PropertyChangedEvent) override;
+	AIMODULE_API virtual void PostEditChangeProperty( struct FPropertyChangedEvent& PropertyChangedEvent) override;
 #endif // WITH_EDITOR
 
 protected:
-	FVector CalcDirection(FEnvQueryInstance& QueryInstance) const;
+	AIMODULE_API FVector CalcDirection(FEnvQueryInstance& QueryInstance) const;
 
-	void GenerateItemsForCircle(uint8* ContextRawData, UEnvQueryItemType* ContextItemType,
+	AIMODULE_API void GenerateItemsForCircle(uint8* ContextRawData, UEnvQueryItemType* ContextItemType,
 		const FVector& CenterLocation, const FVector& StartDirection,
 		const TArray<AActor*>& IgnoredActors,
 		int32 StepsCount, float AngleStep, FEnvQueryInstance& OutQueryInstance) const;
 
-	virtual void AddItemDataForCircle(uint8* ContextRawData, UEnvQueryItemType* ContextItemType, 
+	AIMODULE_API virtual void AddItemDataForCircle(uint8* ContextRawData, UEnvQueryItemType* ContextItemType, 
 		const TArray<FNavLocation>& Locations, FEnvQueryInstance& OutQueryInstance) const;
 };
diff --git a/Engine/Source/Runtime/AIModule/Classes/EnvironmentQuery/Generators/EnvQueryGenerator_PathingGrid.h b/Engine/Source/Runtime/AIModule/Classes/EnvironmentQuery/Generators/EnvQueryGenerator_PathingGrid.h
--- a/Engine/Source/Runtime/AIModule/Classes/EnvironmentQuery/Generators/EnvQueryGenerator_PathingGrid.h
+++ b/Engine/Source/Runtime/AIModule/Classes/EnvironmentQuery/Generators/EnvQueryGenerator_PathingGrid.h
@@ -16,8 +16,8 @@
  *  with paths to/from context no further than given limit
  */
 
-UCLASS(meta = (DisplayName = "Points: Pathing Grid"))
-class AIMODULE_API UEnvQueryGenerator_PathingGrid : public UEnvQueryGenerator_SimpleGrid
+UCLASS(meta = (DisplayName = "Points: Pathing Grid"), MinimalAPI)
+class UEnvQueryGenerator_PathingGrid : public UEnvQueryGenerator_SimpleGrid
 {
 	GENERATED_UCLASS_BODY()
 
@@ -33,5 +33,5 @@ class AIMODULE_API UEnvQueryGenerator_PathingGrid : public UEnvQueryGenerator_Si
 	UPROPERTY(EditDefaultsOnly, Category = Pathfinding, AdvancedDisplay)
 	FAIDataProviderFloatValue ScanRangeMultiplier;
 
-	virtual void ProjectAndFilterNavPoints(TArray<FNavLocation>& Points, FEnvQueryInstance& QueryInstance) const override;
+	AIMODULE_API virtual void ProjectAndFilterNavPoints(TArray<FNavLocation>& Points, FEnvQueryInstance& QueryInstance) const override;
 };
diff --git a/Engine/Source/Runtime/AIModule/Classes/EnvironmentQuery/Generators/EnvQueryGenerator_PerceivedActors.h b/Engine/Source/Runtime/AIModule/Classes/EnvironmentQuery/Generators/EnvQueryGenerator_PerceivedActors.h
--- a/Engine/Source/Runtime/AIModule/Classes/EnvironmentQuery/Generators/EnvQueryGenerator_PerceivedActors.h
+++ b/Engine/Source/Runtime/AIModule/Classes/EnvironmentQuery/Generators/EnvQueryGenerator_PerceivedActors.h
@@ -9,8 +9,8 @@
 
 
 /** Gathers actors perceived by context */
-UCLASS(meta = (DisplayName = "Perceived Actors"))
-class AIMODULE_API UEnvQueryGenerator_PerceivedActors : public UEnvQueryGenerator
+UCLASS(meta = (DisplayName = "Perceived Actors"), MinimalAPI)
+class UEnvQueryGenerator_PerceivedActors : public UEnvQueryGenerator
 {
 	GENERATED_UCLASS_BODY()
 
@@ -39,8 +39,8 @@ protected:
 	UPROPERTY(EditAnywhere, Category=Generator)
 	bool bIncludeKnownActors = true;
 
-	virtual void GenerateItems(FEnvQueryInstance& QueryInstance) const override;
+	AIMODULE_API virtual void GenerateItems(FEnvQueryInstance& QueryInstance) const override;
 
-	virtual FText GetDescriptionTitle() const override;
-	virtual FText GetDescriptionDetails() const override;
+	AIMODULE_API virtual FText GetDescriptionTitle() const override;
+	AIMODULE_API virtual FText GetDescriptionDetails() const override;
 };
