[Engine]

diff --git a/Engine/Source/Runtime/AIModule/Classes/EnvironmentQuery/Generators/EnvQueryGenerator_ProjectedPoints.h b/Engine/Source/Runtime/AIModule/Classes/EnvironmentQuery/Generators/EnvQueryGenerator_ProjectedPoints.h
--- a/Engine/Source/Runtime/AIModule/Classes/EnvironmentQuery/Generators/EnvQueryGenerator_ProjectedPoints.h
+++ b/Engine/Source/Runtime/AIModule/Classes/EnvironmentQuery/Generators/EnvQueryGenerator_ProjectedPoints.h
@@ -9,8 +9,8 @@
 #include "EnvironmentQuery/EnvQueryGenerator.h"
 #include "EnvQueryGenerator_ProjectedPoints.generated.h"
 
-UCLASS(Abstract)
-class AIMODULE_API UEnvQueryGenerator_ProjectedPoints : public UEnvQueryGenerator
+UCLASS(Abstract, MinimalAPI)
+class UEnvQueryGenerator_ProjectedPoints : public UEnvQueryGenerator
 {
 	GENERATED_UCLASS_BODY()
 
@@ -31,10 +31,10 @@ class AIMODULE_API UEnvQueryGenerator_ProjectedPoints : public UEnvQueryGenerato
 	};
 
 	/** project all points in array and remove those outside navmesh */
-	virtual void ProjectAndFilterNavPoints(TArray<FNavLocation>& Points, FEnvQueryInstance& QueryInstance) const;
+	AIMODULE_API virtual void ProjectAndFilterNavPoints(TArray<FNavLocation>& Points, FEnvQueryInstance& QueryInstance) const;
 
 	/** store points as generator's result */
-	virtual void StoreNavPoints(const TArray<FNavLocation>& Points, FEnvQueryInstance& QueryInstance) const;
+	AIMODULE_API virtual void StoreNavPoints(const TArray<FNavLocation>& Points, FEnvQueryInstance& QueryInstance) const;
 
-	virtual void PostLoad() override;
+	AIMODULE_API virtual void PostLoad() override;
 };
diff --git a/Engine/Source/Runtime/AIModule/Classes/EnvironmentQuery/Generators/EnvQueryGenerator_SimpleGrid.h b/Engine/Source/Runtime/AIModule/Classes/EnvironmentQuery/Generators/EnvQueryGenerator_SimpleGrid.h
--- a/Engine/Source/Runtime/AIModule/Classes/EnvironmentQuery/Generators/EnvQueryGenerator_SimpleGrid.h
+++ b/Engine/Source/Runtime/AIModule/Classes/EnvironmentQuery/Generators/EnvQueryGenerator_SimpleGrid.h
@@ -14,8 +14,8 @@
  *  Simple grid, generates points in 2D square around context
  */
 
-UCLASS(meta = (DisplayName = "Points: Grid"))
-class AIMODULE_API UEnvQueryGenerator_SimpleGrid : public UEnvQueryGenerator_ProjectedPoints
+UCLASS(meta = (DisplayName = "Points: Grid"), MinimalAPI)
+class UEnvQueryGenerator_SimpleGrid : public UEnvQueryGenerator_ProjectedPoints
 {
 	GENERATED_UCLASS_BODY()
 
@@ -31,8 +31,8 @@ class AIMODULE_API UEnvQueryGenerator_SimpleGrid : public UEnvQueryGenerator_Pro
 	UPROPERTY(EditDefaultsOnly, Category=Generator)
 	TSubclassOf<UEnvQueryContext> GenerateAround;
 
-	virtual void GenerateItems(FEnvQueryInstance& QueryInstance) const override;
+	AIMODULE_API virtual void GenerateItems(FEnvQueryInstance& QueryInstance) const override;
 
-	virtual FText GetDescriptionTitle() const override;
-	virtual FText GetDescriptionDetails() const override;
+	AIMODULE_API virtual FText GetDescriptionTitle() const override;
+	AIMODULE_API virtual FText GetDescriptionDetails() const override;
 };
diff --git a/Engine/Source/Runtime/AIModule/Classes/EnvironmentQuery/Items/EnvQueryItemType.h b/Engine/Source/Runtime/AIModule/Classes/EnvironmentQuery/Items/EnvQueryItemType.h
--- a/Engine/Source/Runtime/AIModule/Classes/EnvironmentQuery/Items/EnvQueryItemType.h
+++ b/Engine/Source/Runtime/AIModule/Classes/EnvironmentQuery/Items/EnvQueryItemType.h
@@ -11,8 +11,8 @@
 class UBlackboardComponent;
 struct FBlackboardKeySelector;
 
-UCLASS(Abstract)
-class AIMODULE_API UEnvQueryItemType : public UObject
+UCLASS(Abstract, MinimalAPI)
+class UEnvQueryItemType : public UObject
 {
 	GENERATED_BODY()
 
@@ -22,22 +22,22 @@ public:
 	  *	not set in a defived class*/
 	typedef void FValueType;
 
-	UEnvQueryItemType(const FObjectInitializer& ObjectInitializer = FObjectInitializer::Get());
+	AIMODULE_API UEnvQueryItemType(const FObjectInitializer& ObjectInitializer = FObjectInitializer::Get());
 
 	/** get ValueSize */
 	FORCEINLINE uint16 GetValueSize() const { return ValueSize; }
 
 	/** add filters for blackboard key selector */
-	virtual void AddBlackboardFilters(FBlackboardKeySelector& KeySelector, UObject* FilterOwner) const;
+	AIMODULE_API virtual void AddBlackboardFilters(FBlackboardKeySelector& KeySelector, UObject* FilterOwner) const;
 
 	/** store value in blackboard entry */
-	virtual bool StoreInBlackboard(FBlackboardKeySelector& KeySelector, UBlackboardComponent* Blackboard, const uint8* RawData) const;
+	AIMODULE_API virtual bool StoreInBlackboard(FBlackboardKeySelector& KeySelector, UBlackboardComponent* Blackboard, const uint8* RawData) const;
 
 	/** unregister from known types in EnvQueryManager */
-	virtual void FinishDestroy() override;
+	AIMODULE_API virtual void FinishDestroy() override;
 
 	/** describe item */
-	virtual FString GetDescription(const uint8* RawData) const;
+	AIMODULE_API virtual FString GetDescription(const uint8* RawData) const;
 
 	/** helper function for reading typed data from memory block */
 	template<typename T>
diff --git a/Engine/Source/Runtime/AIModule/Classes/EnvironmentQuery/Items/EnvQueryItemType_Actor.h b/Engine/Source/Runtime/AIModule/Classes/EnvironmentQuery/Items/EnvQueryItemType_Actor.h
--- a/Engine/Source/Runtime/AIModule/Classes/EnvironmentQuery/Items/EnvQueryItemType_Actor.h
+++ b/Engine/Source/Runtime/AIModule/Classes/EnvironmentQuery/Items/EnvQueryItemType_Actor.h
@@ -12,23 +12,23 @@ class AActor;
 struct FEnvQueryContextData;
 struct FWeakObjectPtr;
 
-UCLASS()
-class AIMODULE_API UEnvQueryItemType_Actor : public UEnvQueryItemType_ActorBase
+UCLASS(MinimalAPI)
+class UEnvQueryItemType_Actor : public UEnvQueryItemType_ActorBase
 {
 	GENERATED_BODY()
 public:
 	typedef const FWeakObjectPtr& FValueType;
 
-	UEnvQueryItemType_Actor(const FObjectInitializer& ObjectInitializer = FObjectInitializer::Get());
+	AIMODULE_API UEnvQueryItemType_Actor(const FObjectInitializer& ObjectInitializer = FObjectInitializer::Get());
 
-	static AActor* GetValue(const uint8* RawData);
-	static void SetValue(uint8* RawData, const FWeakObjectPtr& Value);
+	static AIMODULE_API AActor* GetValue(const uint8* RawData);
+	static AIMODULE_API void SetValue(uint8* RawData, const FWeakObjectPtr& Value);
 
-	static void SetContextHelper(FEnvQueryContextData& ContextData, const AActor* SingleActor);
-	static void SetContextHelper(FEnvQueryContextData& ContextData, const TArray<const AActor*>& MultipleActors);
-	static void SetContextHelper(FEnvQueryContextData& ContextData, const TArray<AActor*>& MultipleActors);
+	static AIMODULE_API void SetContextHelper(FEnvQueryContextData& ContextData, const AActor* SingleActor);
+	static AIMODULE_API void SetContextHelper(FEnvQueryContextData& ContextData, const TArray<const AActor*>& MultipleActors);
+	static AIMODULE_API void SetContextHelper(FEnvQueryContextData& ContextData, const TArray<AActor*>& MultipleActors);
 
-	virtual FVector GetItemLocation(const uint8* RawData) const override;
-	virtual FRotator GetItemRotation(const uint8* RawData) const override;
-	virtual AActor* GetActor(const uint8* RawData) const override;
+	AIMODULE_API virtual FVector GetItemLocation(const uint8* RawData) const override;
+	AIMODULE_API virtual FRotator GetItemRotation(const uint8* RawData) const override;
+	AIMODULE_API virtual AActor* GetActor(const uint8* RawData) const override;
 };
diff --git a/Engine/Source/Runtime/AIModule/Classes/EnvironmentQuery/Items/EnvQueryItemType_ActorBase.h b/Engine/Source/Runtime/AIModule/Classes/EnvironmentQuery/Items/EnvQueryItemType_ActorBase.h
--- a/Engine/Source/Runtime/AIModule/Classes/EnvironmentQuery/Items/EnvQueryItemType_ActorBase.h
+++ b/Engine/Source/Runtime/AIModule/Classes/EnvironmentQuery/Items/EnvQueryItemType_ActorBase.h
@@ -11,15 +11,15 @@ class AActor;
 class UBlackboardComponent;
 struct FBlackboardKeySelector;
 
-UCLASS(Abstract)
-class AIMODULE_API UEnvQueryItemType_ActorBase : public UEnvQueryItemType_VectorBase
+UCLASS(Abstract, MinimalAPI)
+class UEnvQueryItemType_ActorBase : public UEnvQueryItemType_VectorBase
 {
 	GENERATED_BODY()
 
 public:
-	virtual AActor* GetActor(const uint8* RawData) const;
+	AIMODULE_API virtual AActor* GetActor(const uint8* RawData) const;
 
-	virtual void AddBlackboardFilters(FBlackboardKeySelector& KeySelector, UObject* FilterOwner) const override;
-	virtual bool StoreInBlackboard(FBlackboardKeySelector& KeySelector, UBlackboardComponent* Blackboard, const uint8* RawData) const override;
-	virtual FString GetDescription(const uint8* RawData) const override;
+	AIMODULE_API virtual void AddBlackboardFilters(FBlackboardKeySelector& KeySelector, UObject* FilterOwner) const override;
+	AIMODULE_API virtual bool StoreInBlackboard(FBlackboardKeySelector& KeySelector, UBlackboardComponent* Blackboard, const uint8* RawData) const override;
+	AIMODULE_API virtual FString GetDescription(const uint8* RawData) const override;
 };
diff --git a/Engine/Source/Runtime/AIModule/Classes/EnvironmentQuery/Items/EnvQueryItemType_Direction.h b/Engine/Source/Runtime/AIModule/Classes/EnvironmentQuery/Items/EnvQueryItemType_Direction.h
--- a/Engine/Source/Runtime/AIModule/Classes/EnvironmentQuery/Items/EnvQueryItemType_Direction.h
+++ b/Engine/Source/Runtime/AIModule/Classes/EnvironmentQuery/Items/EnvQueryItemType_Direction.h
@@ -10,26 +10,26 @@
 
 struct FEnvQueryContextData;
 
-UCLASS()
-class AIMODULE_API UEnvQueryItemType_Direction : public UEnvQueryItemType_VectorBase
+UCLASS(MinimalAPI)
+class UEnvQueryItemType_Direction : public UEnvQueryItemType_VectorBase
 {
 	GENERATED_BODY()
 public:
 	typedef FVector FValueType;
 
-	UEnvQueryItemType_Direction(const FObjectInitializer& ObjectInitializer = FObjectInitializer::Get());
+	AIMODULE_API UEnvQueryItemType_Direction(const FObjectInitializer& ObjectInitializer = FObjectInitializer::Get());
 
-	static FVector GetValue(const uint8* RawData);
-	static void SetValue(uint8* RawData, const FVector& Value);
+	static AIMODULE_API FVector GetValue(const uint8* RawData);
+	static AIMODULE_API void SetValue(uint8* RawData, const FVector& Value);
 
-	static FRotator GetValueRot(const uint8* RawData);
-	static void SetValueRot(uint8* RawData, const FRotator& Value);
+	static AIMODULE_API FRotator GetValueRot(const uint8* RawData);
+	static AIMODULE_API void SetValueRot(uint8* RawData, const FRotator& Value);
 
-	static void SetContextHelper(FEnvQueryContextData& ContextData, const FVector& SingleDirection);
-	static void SetContextHelper(FEnvQueryContextData& ContextData, const FRotator& SingleRotation);
-	static void SetContextHelper(FEnvQueryContextData& ContextData, const TArray<FVector>& MultipleDirections);
-	static void SetContextHelper(FEnvQueryContextData& ContextData, const TArray<FRotator>& MultipleRotations);
+	static AIMODULE_API void SetContextHelper(FEnvQueryContextData& ContextData, const FVector& SingleDirection);
+	static AIMODULE_API void SetContextHelper(FEnvQueryContextData& ContextData, const FRotator& SingleRotation);
+	static AIMODULE_API void SetContextHelper(FEnvQueryContextData& ContextData, const TArray<FVector>& MultipleDirections);
+	static AIMODULE_API void SetContextHelper(FEnvQueryContextData& ContextData, const TArray<FRotator>& MultipleRotations);
 
-	virtual FRotator GetItemRotation(const uint8* RawData) const override;
-	virtual FString GetDescription(const uint8* RawData) const override;
+	AIMODULE_API virtual FRotator GetItemRotation(const uint8* RawData) const override;
+	AIMODULE_API virtual FString GetDescription(const uint8* RawData) const override;
 };
diff --git a/Engine/Source/Runtime/AIModule/Classes/EnvironmentQuery/Items/EnvQueryItemType_Point.h b/Engine/Source/Runtime/AIModule/Classes/EnvironmentQuery/Items/EnvQueryItemType_Point.h
--- a/Engine/Source/Runtime/AIModule/Classes/EnvironmentQuery/Items/EnvQueryItemType_Point.h
+++ b/Engine/Source/Runtime/AIModule/Classes/EnvironmentQuery/Items/EnvQueryItemType_Point.h
@@ -10,26 +10,26 @@
 #include "EnvironmentQuery/Items/EnvQueryItemType_VectorBase.h"
 #include "EnvQueryItemType_Point.generated.h"
 
-UCLASS()
-class AIMODULE_API UEnvQueryItemType_Point : public UEnvQueryItemType_VectorBase
+UCLASS(MinimalAPI)
+class UEnvQueryItemType_Point : public UEnvQueryItemType_VectorBase
 {
 	GENERATED_BODY()
 public:
 	typedef const FNavLocation FValueType;
 
-	UEnvQueryItemType_Point(const FObjectInitializer& ObjectInitializer = FObjectInitializer::Get());
+	AIMODULE_API UEnvQueryItemType_Point(const FObjectInitializer& ObjectInitializer = FObjectInitializer::Get());
 
-	static FNavLocation GetValue(const uint8* RawData);
-	static void SetValue(uint8* RawData, const FNavLocation& Value);
+	static AIMODULE_API FNavLocation GetValue(const uint8* RawData);
+	static AIMODULE_API void SetValue(uint8* RawData, const FNavLocation& Value);
 
-	static void SetContextHelper(FEnvQueryContextData& ContextData, const FVector& SinglePoint);
-	static void SetContextHelper(FEnvQueryContextData& ContextData, const TArray<FVector>& MultiplePoints);
+	static AIMODULE_API void SetContextHelper(FEnvQueryContextData& ContextData, const FVector& SinglePoint);
+	static AIMODULE_API void SetContextHelper(FEnvQueryContextData& ContextData, const TArray<FVector>& MultiplePoints);
 
-	virtual FVector GetItemLocation(const uint8* RawData) const override;
-	virtual FNavLocation GetItemNavLocation(const uint8* RawData) const;
+	AIMODULE_API virtual FVector GetItemLocation(const uint8* RawData) const override;
+	AIMODULE_API virtual FNavLocation GetItemNavLocation(const uint8* RawData) const;
 
 	/** Update location data in existing item */
-	virtual void SetItemNavLocation(uint8* RawData, const FNavLocation& Value) const;
+	AIMODULE_API virtual void SetItemNavLocation(uint8* RawData, const FNavLocation& Value) const;
 };
 
 // a specialization to support saving locations with navigation data already gathered
diff --git a/Engine/Source/Runtime/AIModule/Classes/EnvironmentQuery/Items/EnvQueryItemType_VectorBase.h b/Engine/Source/Runtime/AIModule/Classes/EnvironmentQuery/Items/EnvQueryItemType_VectorBase.h
--- a/Engine/Source/Runtime/AIModule/Classes/EnvironmentQuery/Items/EnvQueryItemType_VectorBase.h
+++ b/Engine/Source/Runtime/AIModule/Classes/EnvironmentQuery/Items/EnvQueryItemType_VectorBase.h
@@ -10,16 +10,16 @@
 class UBlackboardComponent;
 struct FBlackboardKeySelector;
 
-UCLASS(Abstract)
-class AIMODULE_API UEnvQueryItemType_VectorBase : public UEnvQueryItemType
+UCLASS(Abstract, MinimalAPI)
+class UEnvQueryItemType_VectorBase : public UEnvQueryItemType
 {
 	GENERATED_BODY()
 
 public:
-	virtual FVector GetItemLocation(const uint8* RawData) const;
-	virtual FRotator GetItemRotation(const uint8* RawData) const;
+	AIMODULE_API virtual FVector GetItemLocation(const uint8* RawData) const;
+	AIMODULE_API virtual FRotator GetItemRotation(const uint8* RawData) const;
 
-	virtual void AddBlackboardFilters(FBlackboardKeySelector& KeySelector, UObject* FilterOwner) const override;
-	virtual bool StoreInBlackboard(FBlackboardKeySelector& KeySelector, UBlackboardComponent* Blackboard, const uint8* RawData) const override;
-	virtual FString GetDescription(const uint8* RawData) const override;
+	AIMODULE_API virtual void AddBlackboardFilters(FBlackboardKeySelector& KeySelector, UObject* FilterOwner) const override;
+	AIMODULE_API virtual bool StoreInBlackboard(FBlackboardKeySelector& KeySelector, UBlackboardComponent* Blackboard, const uint8* RawData) const override;
+	AIMODULE_API virtual FString GetDescription(const uint8* RawData) const override;
 };
diff --git a/Engine/Source/Runtime/AIModule/Classes/EnvironmentQuery/Tests/EnvQueryTest_Pathfinding.h b/Engine/Source/Runtime/AIModule/Classes/EnvironmentQuery/Tests/EnvQueryTest_Pathfinding.h
--- a/Engine/Source/Runtime/AIModule/Classes/EnvironmentQuery/Tests/EnvQueryTest_Pathfinding.h
+++ b/Engine/Source/Runtime/AIModule/Classes/EnvironmentQuery/Tests/EnvQueryTest_Pathfinding.h
@@ -27,8 +27,8 @@ namespace EEnvTestPathfinding
 	};
 }
 
-UCLASS()
-class AIMODULE_API UEnvQueryTest_Pathfinding : public UEnvQueryTest
+UCLASS(MinimalAPI)
+class UEnvQueryTest_Pathfinding : public UEnvQueryTest
 {
 	GENERATED_UCLASS_BODY()
 
@@ -52,30 +52,30 @@ class AIMODULE_API UEnvQueryTest_Pathfinding : public UEnvQueryTest
 	UPROPERTY(EditDefaultsOnly, Category=Pathfinding)
 	TSubclassOf<UNavigationQueryFilter> FilterClass;
 
-	virtual void RunTest(FEnvQueryInstance& QueryInstance) const override;
+	AIMODULE_API virtual void RunTest(FEnvQueryInstance& QueryInstance) const override;
 
-	virtual FText GetDescriptionTitle() const override;
-	virtual FText GetDescriptionDetails() const override;
+	AIMODULE_API virtual FText GetDescriptionTitle() const override;
+	AIMODULE_API virtual FText GetDescriptionDetails() const override;
 
 #if WITH_EDITOR
 	/** update test properties after changing mode */
-	virtual void PostEditChangeProperty(struct FPropertyChangedEvent& PropertyChangedEvent) override;
+	AIMODULE_API virtual void PostEditChangeProperty(struct FPropertyChangedEvent& PropertyChangedEvent) override;
 #endif
-	virtual void PostLoad() override;
+	AIMODULE_API virtual void PostLoad() override;
 
 protected:
 
 	DECLARE_DELEGATE_RetVal_SevenParams(bool, FTestPathSignature, const FVector&, const FVector&, EPathFindingMode::Type, const ANavigationData&, UNavigationSystemV1&, FSharedConstNavQueryFilter, const UObject*);
 	DECLARE_DELEGATE_RetVal_SevenParams(float, FFindPathSignature, const FVector&, const FVector&, EPathFindingMode::Type, const ANavigationData&, UNavigationSystemV1&, FSharedConstNavQueryFilter, const UObject*);
 
-	bool TestPathFrom(const FVector& ItemPos, const FVector& ContextPos, EPathFindingMode::Type Mode, const ANavigationData& NavData, UNavigationSystemV1& NavSys, FSharedConstNavQueryFilter NavFilter, const UObject* PathOwner) const;
-	bool TestPathTo(const FVector& ItemPos, const FVector& ContextPos, EPathFindingMode::Type Mode, const ANavigationData& NavData, UNavigationSystemV1& NavSys, FSharedConstNavQueryFilter NavFilter, const UObject* PathOwner) const;
-	float FindPathCostFrom(const FVector& ItemPos, const FVector& ContextPos, EPathFindingMode::Type Mode, const ANavigationData& NavData, UNavigationSystemV1& NavSys, FSharedConstNavQueryFilter NavFilter, const UObject* PathOwner) const;
-	float FindPathCostTo(const FVector& ItemPos, const FVector& ContextPos, EPathFindingMode::Type Mode, const ANavigationData& NavData, UNavigationSystemV1& NavSys, FSharedConstNavQueryFilter NavFilter, const UObject* PathOwner) const;
-	float FindPathLengthFrom(const FVector& ItemPos, const FVector& ContextPos, EPathFindingMode::Type Mode, const ANavigationData& NavData, UNavigationSystemV1& NavSys, FSharedConstNavQueryFilter NavFilter, const UObject* PathOwner) const;
-	float FindPathLengthTo(const FVector& ItemPos, const FVector& ContextPos, EPathFindingMode::Type Mode, const ANavigationData& NavData, UNavigationSystemV1& NavSys, FSharedConstNavQueryFilter NavFilter, const UObject* PathOwner) const;
+	AIMODULE_API bool TestPathFrom(const FVector& ItemPos, const FVector& ContextPos, EPathFindingMode::Type Mode, const ANavigationData& NavData, UNavigationSystemV1& NavSys, FSharedConstNavQueryFilter NavFilter, const UObject* PathOwner) const;
+	AIMODULE_API bool TestPathTo(const FVector& ItemPos, const FVector& ContextPos, EPathFindingMode::Type Mode, const ANavigationData& NavData, UNavigationSystemV1& NavSys, FSharedConstNavQueryFilter NavFilter, const UObject* PathOwner) const;
+	AIMODULE_API float FindPathCostFrom(const FVector& ItemPos, const FVector& ContextPos, EPathFindingMode::Type Mode, const ANavigationData& NavData, UNavigationSystemV1& NavSys, FSharedConstNavQueryFilter NavFilter, const UObject* PathOwner) const;
+	AIMODULE_API float FindPathCostTo(const FVector& ItemPos, const FVector& ContextPos, EPathFindingMode::Type Mode, const ANavigationData& NavData, UNavigationSystemV1& NavSys, FSharedConstNavQueryFilter NavFilter, const UObject* PathOwner) const;
+	AIMODULE_API float FindPathLengthFrom(const FVector& ItemPos, const FVector& ContextPos, EPathFindingMode::Type Mode, const ANavigationData& NavData, UNavigationSystemV1& NavSys, FSharedConstNavQueryFilter NavFilter, const UObject* PathOwner) const;
+	AIMODULE_API float FindPathLengthTo(const FVector& ItemPos, const FVector& ContextPos, EPathFindingMode::Type Mode, const ANavigationData& NavData, UNavigationSystemV1& NavSys, FSharedConstNavQueryFilter NavFilter, const UObject* PathOwner) const;
 
-	ANavigationData* FindNavigationData(UNavigationSystemV1& NavSys, UObject* Owner) const;
+	AIMODULE_API ANavigationData* FindNavigationData(UNavigationSystemV1& NavSys, UObject* Owner) const;
 
-	virtual TSubclassOf<UNavigationQueryFilter> GetNavFilterClass(FEnvQueryInstance& QueryInstance) const;
+	AIMODULE_API virtual TSubclassOf<UNavigationQueryFilter> GetNavFilterClass(FEnvQueryInstance& QueryInstance) const;
 };
diff --git a/Engine/Source/Runtime/AIModule/Classes/EnvironmentQuery/Tests/EnvQueryTest_PathfindingBatch.h b/Engine/Source/Runtime/AIModule/Classes/EnvironmentQuery/Tests/EnvQueryTest_PathfindingBatch.h
--- a/Engine/Source/Runtime/AIModule/Classes/EnvironmentQuery/Tests/EnvQueryTest_PathfindingBatch.h
+++ b/Engine/Source/Runtime/AIModule/Classes/EnvironmentQuery/Tests/EnvQueryTest_PathfindingBatch.h
@@ -8,8 +8,8 @@
 #include "EnvironmentQuery/Tests/EnvQueryTest_Pathfinding.h"
 #include "EnvQueryTest_PathfindingBatch.generated.h"
 
-UCLASS()
-class AIMODULE_API UEnvQueryTest_PathfindingBatch : public UEnvQueryTest_Pathfinding
+UCLASS(MinimalAPI)
+class UEnvQueryTest_PathfindingBatch : public UEnvQueryTest_Pathfinding
 {
 	GENERATED_UCLASS_BODY()
 
@@ -17,6 +17,6 @@ class AIMODULE_API UEnvQueryTest_PathfindingBatch : public UEnvQueryTest_Pathfin
 	UPROPERTY(EditDefaultsOnly, Category = Pathfinding, AdvancedDisplay)
 	FAIDataProviderFloatValue ScanRangeMultiplier;
 
-	virtual void RunTest(FEnvQueryInstance& QueryInstance) const override;
-	virtual FText GetDescriptionTitle() const override;
+	AIMODULE_API virtual void RunTest(FEnvQueryInstance& QueryInstance) const override;
+	AIMODULE_API virtual FText GetDescriptionTitle() const override;
 };
diff --git a/Engine/Source/Runtime/AIModule/Classes/EnvironmentQuery/Tests/EnvQueryTest_Trace.h b/Engine/Source/Runtime/AIModule/Classes/EnvironmentQuery/Tests/EnvQueryTest_Trace.h
--- a/Engine/Source/Runtime/AIModule/Classes/EnvironmentQuery/Tests/EnvQueryTest_Trace.h
+++ b/Engine/Source/Runtime/AIModule/Classes/EnvironmentQuery/Tests/EnvQueryTest_Trace.h
@@ -14,8 +14,8 @@
 class AActor;
 struct FCollisionQueryParams;
 
-UCLASS()
-class AIMODULE_API UEnvQueryTest_Trace : public UEnvQueryTest
+UCLASS(MinimalAPI)
+class UEnvQueryTest_Trace : public UEnvQueryTest
 {
 	GENERATED_UCLASS_BODY()
 
@@ -39,23 +39,23 @@ class AIMODULE_API UEnvQueryTest_Trace : public UEnvQueryTest
 	UPROPERTY(EditDefaultsOnly, Category=Trace)
 	TSubclassOf<UEnvQueryContext> Context;
 
-	virtual void RunTest(FEnvQueryInstance& QueryInstance) const override;
+	AIMODULE_API virtual void RunTest(FEnvQueryInstance& QueryInstance) const override;
 
-	virtual FText GetDescriptionTitle() const override;
-	virtual FText GetDescriptionDetails() const override;
+	AIMODULE_API virtual FText GetDescriptionTitle() const override;
+	AIMODULE_API virtual FText GetDescriptionDetails() const override;
 
-	virtual void PostLoad() override;
+	AIMODULE_API virtual void PostLoad() override;
 
 protected:
 
 	DECLARE_DELEGATE_RetVal_SevenParams(bool, FRunTraceSignature, const FVector&, const FVector&, AActor*, UWorld*, enum ECollisionChannel, const FCollisionQueryParams&, const FVector&);
 
-	bool RunLineTraceTo(const FVector& ItemPos, const FVector& ContextPos, AActor* ItemActor, UWorld* World, enum ECollisionChannel Channel, const FCollisionQueryParams& Params, const FVector& Extent);
-	bool RunLineTraceFrom(const FVector& ItemPos, const FVector& ContextPos, AActor* ItemActor, UWorld* World, enum ECollisionChannel Channel, const FCollisionQueryParams& Params, const FVector& Extent);
-	bool RunBoxTraceTo(const FVector& ItemPos, const FVector& ContextPos, AActor* ItemActor, UWorld* World, enum ECollisionChannel Channel, const FCollisionQueryParams& Params, const FVector& Extent);
-	bool RunBoxTraceFrom(const FVector& ItemPos, const FVector& ContextPos, AActor* ItemActor, UWorld* World, enum ECollisionChannel Channel, const FCollisionQueryParams& Params, const FVector& Extent);
-	bool RunSphereTraceTo(const FVector& ItemPos, const FVector& ContextPos, AActor* ItemActor, UWorld* World, enum ECollisionChannel Channel, const FCollisionQueryParams& Params, const FVector& Extent);
-	bool RunSphereTraceFrom(const FVector& ItemPos, const FVector& ContextPos, AActor* ItemActor, UWorld* World, enum ECollisionChannel Channel, const FCollisionQueryParams& Params, const FVector& Extent);
-	bool RunCapsuleTraceTo(const FVector& ItemPos, const FVector& ContextPos, AActor* ItemActor, UWorld* World, enum ECollisionChannel Channel, const FCollisionQueryParams& Params, const FVector& Extent);
-	bool RunCapsuleTraceFrom(const FVector& ItemPos, const FVector& ContextPos, AActor* ItemActor, UWorld* World, enum ECollisionChannel Channel, const FCollisionQueryParams& Params, const FVector& Extent);
+	AIMODULE_API bool RunLineTraceTo(const FVector& ItemPos, const FVector& ContextPos, AActor* ItemActor, UWorld* World, enum ECollisionChannel Channel, const FCollisionQueryParams& Params, const FVector& Extent);
+	AIMODULE_API bool RunLineTraceFrom(const FVector& ItemPos, const FVector& ContextPos, AActor* ItemActor, UWorld* World, enum ECollisionChannel Channel, const FCollisionQueryParams& Params, const FVector& Extent);
+	AIMODULE_API bool RunBoxTraceTo(const FVector& ItemPos, const FVector& ContextPos, AActor* ItemActor, UWorld* World, enum ECollisionChannel Channel, const FCollisionQueryParams& Params, const FVector& Extent);
+	AIMODULE_API bool RunBoxTraceFrom(const FVector& ItemPos, const FVector& ContextPos, AActor* ItemActor, UWorld* World, enum ECollisionChannel Channel, const FCollisionQueryParams& Params, const FVector& Extent);
+	AIMODULE_API bool RunSphereTraceTo(const FVector& ItemPos, const FVector& ContextPos, AActor* ItemActor, UWorld* World, enum ECollisionChannel Channel, const FCollisionQueryParams& Params, const FVector& Extent);
+	AIMODULE_API bool RunSphereTraceFrom(const FVector& ItemPos, const FVector& ContextPos, AActor* ItemActor, UWorld* World, enum ECollisionChannel Channel, const FCollisionQueryParams& Params, const FVector& Extent);
+	AIMODULE_API bool RunCapsuleTraceTo(const FVector& ItemPos, const FVector& ContextPos, AActor* ItemActor, UWorld* World, enum ECollisionChannel Channel, const FCollisionQueryParams& Params, const FVector& Extent);
+	AIMODULE_API bool RunCapsuleTraceFrom(const FVector& ItemPos, const FVector& ContextPos, AActor* ItemActor, UWorld* World, enum ECollisionChannel Channel, const FCollisionQueryParams& Params, const FVector& Extent);
 };
diff --git a/Engine/Source/Runtime/AIModule/Classes/GenericTeamAgentInterface.h b/Engine/Source/Runtime/AIModule/Classes/GenericTeamAgentInterface.h
--- a/Engine/Source/Runtime/AIModule/Classes/GenericTeamAgentInterface.h
+++ b/Engine/Source/Runtime/AIModule/Classes/GenericTeamAgentInterface.h
@@ -20,7 +20,7 @@ namespace ETeamAttitude
 }
 
 USTRUCT(BlueprintType)
-struct AIMODULE_API FGenericTeamId
+struct FGenericTeamId
 {
 	GENERATED_USTRUCT_BODY()
 
@@ -45,8 +45,8 @@ public:
 	FORCEINLINE uint8 GetId() const { return TeamID; }
 	//FORCEINLINE void SetId(uint8 NewId) { TeamID = NewId; }
 	
-	static FGenericTeamId GetTeamIdentifier(const AActor* TeamMember);
-	static ETeamAttitude::Type GetAttitude(const AActor* A, const AActor* B);
+	static AIMODULE_API FGenericTeamId GetTeamIdentifier(const AActor* TeamMember);
+	static AIMODULE_API ETeamAttitude::Type GetAttitude(const AActor* A, const AActor* B);
 	static ETeamAttitude::Type GetAttitude(FGenericTeamId TeamA, FGenericTeamId TeamB)
 	{
 		return AttitudeSolverImpl ? (AttitudeSolverImpl)(TeamA, TeamB) : ETeamAttitude::Neutral;
@@ -54,17 +54,17 @@ public:
 
 	typedef TFunction<ETeamAttitude::Type(FGenericTeamId, FGenericTeamId)> FAttitudeSolverFunction;
 	
-	static void SetAttitudeSolver(const FAttitudeSolverFunction& Solver);
-	static void ResetAttitudeSolver();
+	static AIMODULE_API void SetAttitudeSolver(const FAttitudeSolverFunction& Solver);
+	static AIMODULE_API void ResetAttitudeSolver();
 
 protected:
 	// the default implementation makes all teams hostile
 	// @note that for consistency IGenericTeamAgentInterface should be using the same function 
 	//	(by default it does)
-	static FAttitudeSolverFunction AttitudeSolverImpl;
+	static AIMODULE_API FAttitudeSolverFunction AttitudeSolverImpl;
 
 public:
-	static const FGenericTeamId NoTeam;
+	static AIMODULE_API const FGenericTeamId NoTeam;
 
 	friend FORCEINLINE uint32 GetTypeHash(const FGenericTeamId Value)
 	{
@@ -72,13 +72,13 @@ public:
 	}
 };
 
-UINTERFACE()
-class AIMODULE_API UGenericTeamAgentInterface : public UInterface
+UINTERFACE(MinimalAPI)
+class UGenericTeamAgentInterface : public UInterface
 {
 	GENERATED_UINTERFACE_BODY()
 };
 
-class AIMODULE_API IGenericTeamAgentInterface
+class IGenericTeamAgentInterface
 {
 	GENERATED_IINTERFACE_BODY()
 
diff --git a/Engine/Source/Runtime/AIModule/Classes/HotSpots/AIHotSpotManager.h b/Engine/Source/Runtime/AIModule/Classes/HotSpots/AIHotSpotManager.h
--- a/Engine/Source/Runtime/AIModule/Classes/HotSpots/AIHotSpotManager.h
+++ b/Engine/Source/Runtime/AIModule/Classes/HotSpots/AIHotSpotManager.h
@@ -6,8 +6,8 @@
 #include "UObject/Object.h"
 #include "AIHotSpotManager.generated.h"
 
-UCLASS()
-class AIMODULE_API UAIHotSpotManager : public UObject
+UCLASS(MinimalAPI)
+class UAIHotSpotManager : public UObject
 {
 	GENERATED_UCLASS_BODY()
 };
diff --git a/Engine/Source/Runtime/AIModule/Classes/Navigation/CrowdFollowingComponent.h b/Engine/Source/Runtime/AIModule/Classes/Navigation/CrowdFollowingComponent.h
--- a/Engine/Source/Runtime/AIModule/Classes/Navigation/CrowdFollowingComponent.h
+++ b/Engine/Source/Runtime/AIModule/Classes/Navigation/CrowdFollowingComponent.h
@@ -34,75 +34,75 @@ enum class ECrowdSimulationState : uint8
 	Disabled		UMETA(DisplayName="Disabled, ignored by others"),
 };
 
-UCLASS(BlueprintType)
-class AIMODULE_API UCrowdFollowingComponent : public UPathFollowingComponent, public ICrowdAgentInterface
+UCLASS(BlueprintType, MinimalAPI)
+class UCrowdFollowingComponent : public UPathFollowingComponent, public ICrowdAgentInterface
 {
 	GENERATED_UCLASS_BODY()
 
-	virtual void BeginDestroy() override;
+	AIMODULE_API virtual void BeginDestroy() override;
 
 	// ICrowdAgentInterface BEGIN
-	virtual FVector GetCrowdAgentLocation() const override;
-	virtual FVector GetCrowdAgentVelocity() const override;
-	virtual void GetCrowdAgentCollisions(float& CylinderRadius, float& CylinderHalfHeight) const override;
-	virtual float GetCrowdAgentMaxSpeed() const override;
-	virtual int32 GetCrowdAgentAvoidanceGroup() const override;
-	virtual int32 GetCrowdAgentGroupsToAvoid() const override;
-	virtual int32 GetCrowdAgentGroupsToIgnore() const override;
+	AIMODULE_API virtual FVector GetCrowdAgentLocation() const override;
+	AIMODULE_API virtual FVector GetCrowdAgentVelocity() const override;
+	AIMODULE_API virtual void GetCrowdAgentCollisions(float& CylinderRadius, float& CylinderHalfHeight) const override;
+	AIMODULE_API virtual float GetCrowdAgentMaxSpeed() const override;
+	AIMODULE_API virtual int32 GetCrowdAgentAvoidanceGroup() const override;
+	AIMODULE_API virtual int32 GetCrowdAgentGroupsToAvoid() const override;
+	AIMODULE_API virtual int32 GetCrowdAgentGroupsToIgnore() const override;
 	// ICrowdAgentInterface END
 
 	// PathFollowingComponent BEGIN
-	virtual void Initialize() override;
-	virtual void Cleanup() override;
-	virtual void AbortMove(const UObject& Instigator, FPathFollowingResultFlags::Type AbortFlags, FAIRequestID RequestID = FAIRequestID::CurrentRequest, EPathFollowingVelocityMode VelocityMode = EPathFollowingVelocityMode::Reset) override;
-	virtual void PauseMove(FAIRequestID RequestID = FAIRequestID::CurrentRequest, EPathFollowingVelocityMode VelocityMode = EPathFollowingVelocityMode::Reset) override;
-	virtual void ResumeMove(FAIRequestID RequestID = FAIRequestID::CurrentRequest) override;
-	virtual FVector GetMoveFocus(bool bAllowStrafe) const override;
-	virtual void OnLanded() override;
-	virtual void FinishUsingCustomLink(INavLinkCustomInterface* CustomNavLink) override;
-	virtual void OnPathFinished(const FPathFollowingResult& Result) override;
-	virtual void OnPathUpdated() override;
-	virtual void OnPathfindingQuery(FPathFindingQuery& Query) override;
+	AIMODULE_API virtual void Initialize() override;
+	AIMODULE_API virtual void Cleanup() override;
+	AIMODULE_API virtual void AbortMove(const UObject& Instigator, FPathFollowingResultFlags::Type AbortFlags, FAIRequestID RequestID = FAIRequestID::CurrentRequest, EPathFollowingVelocityMode VelocityMode = EPathFollowingVelocityMode::Reset) override;
+	AIMODULE_API virtual void PauseMove(FAIRequestID RequestID = FAIRequestID::CurrentRequest, EPathFollowingVelocityMode VelocityMode = EPathFollowingVelocityMode::Reset) override;
+	AIMODULE_API virtual void ResumeMove(FAIRequestID RequestID = FAIRequestID::CurrentRequest) override;
+	AIMODULE_API virtual FVector GetMoveFocus(bool bAllowStrafe) const override;
+	AIMODULE_API virtual void OnLanded() override;
+	AIMODULE_API virtual void FinishUsingCustomLink(INavLinkCustomInterface* CustomNavLink) override;
+	AIMODULE_API virtual void OnPathFinished(const FPathFollowingResult& Result) override;
+	AIMODULE_API virtual void OnPathUpdated() override;
+	AIMODULE_API virtual void OnPathfindingQuery(FPathFindingQuery& Query) override;
 	virtual int32 GetCurrentPathElement() const override { return LastPathPolyIndex; }
-	virtual void OnNavigationInitDone() override;
+	AIMODULE_API virtual void OnNavigationInitDone() override;
 	// PathFollowingComponent END
 
 	/** update params in crowd manager */
-	void UpdateCrowdAgentParams() const;
+	AIMODULE_API void UpdateCrowdAgentParams() const;
 
 	/** pass agent velocity to movement component */
-	virtual void ApplyCrowdAgentVelocity(const FVector& NewVelocity, const FVector& DestPathCorner, bool bTraversingLink, bool bIsNearEndOfPath);
+	AIMODULE_API virtual void ApplyCrowdAgentVelocity(const FVector& NewVelocity, const FVector& DestPathCorner, bool bTraversingLink, bool bIsNearEndOfPath);
 	
 	/** pass desired position to movement component (after resolving collisions between crowd agents) */
-	virtual void ApplyCrowdAgentPosition(const FVector& NewPosition);
+	AIMODULE_API virtual void ApplyCrowdAgentPosition(const FVector& NewPosition);
 
 	/** main switch for crowd steering & avoidance */
 	UFUNCTION(BlueprintCallable, Category = "Crowd")
-	virtual void SuspendCrowdSteering(bool bSuspend);
+	AIMODULE_API virtual void SuspendCrowdSteering(bool bSuspend);
 
 	/** switch between crowd simulation and parent implementation (following path segments) */
-	virtual void SetCrowdSimulationState(ECrowdSimulationState NewState);
+	AIMODULE_API virtual void SetCrowdSimulationState(ECrowdSimulationState NewState);
 
 	/** called when agent moved to next nav node (poly) */
-	virtual void OnNavNodeChanged(NavNodeRef NewPolyRef, NavNodeRef PrevPolyRef, int32 CorridorSize);
-
-	void SetCrowdAnticipateTurns(bool bEnable, bool bUpdateAgent = true);
-	void SetCrowdObstacleAvoidance(bool bEnable, bool bUpdateAgent = true);
-	void SetCrowdSeparation(bool bEnable, bool bUpdateAgent = true);
-	void SetCrowdOptimizeVisibility(bool bEnable, bool bUpdateAgent = true);
-	void SetCrowdOptimizeTopology(bool bEnable, bool bUpdateAgent = true);
-	void SetCrowdPathOffset(bool bEnable, bool bUpdateAgent = true);
-	void SetCrowdSlowdownAtGoal(bool bEnable, bool bUpdateAgent = true);
-	void SetCrowdSeparationWeight(float Weight, bool bUpdateAgent = true);
-	void SetCrowdCollisionQueryRange(float Range, bool bUpdateAgent = true);
-	void SetCrowdPathOptimizationRange(float Range, bool bUpdateAgent = true);
-	void SetCrowdAvoidanceQuality(ECrowdAvoidanceQuality::Type Quality, bool bUpdateAgent = true);
-	void SetCrowdAvoidanceRangeMultiplier(float Multiplier, bool bUpdateAgent = true);
-	void SetCrowdAffectFallingVelocity(bool bEnable);
-	void SetCrowdRotateToVelocity(bool bEnable);
-	void SetAvoidanceGroup(int32 GroupFlags, bool bUpdateAgent = true);
-	void SetGroupsToAvoid(int32 GroupFlags, bool bUpdateAgent = true);
-	void SetGroupsToIgnore(int32 GroupFlags, bool bUpdateAgent = true);
+	AIMODULE_API virtual void OnNavNodeChanged(NavNodeRef NewPolyRef, NavNodeRef PrevPolyRef, int32 CorridorSize);
+
+	AIMODULE_API void SetCrowdAnticipateTurns(bool bEnable, bool bUpdateAgent = true);
+	AIMODULE_API void SetCrowdObstacleAvoidance(bool bEnable, bool bUpdateAgent = true);
+	AIMODULE_API void SetCrowdSeparation(bool bEnable, bool bUpdateAgent = true);
+	AIMODULE_API void SetCrowdOptimizeVisibility(bool bEnable, bool bUpdateAgent = true);
+	AIMODULE_API void SetCrowdOptimizeTopology(bool bEnable, bool bUpdateAgent = true);
+	AIMODULE_API void SetCrowdPathOffset(bool bEnable, bool bUpdateAgent = true);
+	AIMODULE_API void SetCrowdSlowdownAtGoal(bool bEnable, bool bUpdateAgent = true);
+	AIMODULE_API void SetCrowdSeparationWeight(float Weight, bool bUpdateAgent = true);
+	AIMODULE_API void SetCrowdCollisionQueryRange(float Range, bool bUpdateAgent = true);
+	AIMODULE_API void SetCrowdPathOptimizationRange(float Range, bool bUpdateAgent = true);
+	AIMODULE_API void SetCrowdAvoidanceQuality(ECrowdAvoidanceQuality::Type Quality, bool bUpdateAgent = true);
+	AIMODULE_API void SetCrowdAvoidanceRangeMultiplier(float Multiplier, bool bUpdateAgent = true);
+	AIMODULE_API void SetCrowdAffectFallingVelocity(bool bEnable);
+	AIMODULE_API void SetCrowdRotateToVelocity(bool bEnable);
+	AIMODULE_API void SetAvoidanceGroup(int32 GroupFlags, bool bUpdateAgent = true);
+	AIMODULE_API void SetGroupsToAvoid(int32 GroupFlags, bool bUpdateAgent = true);
+	AIMODULE_API void SetGroupsToIgnore(int32 GroupFlags, bool bUpdateAgent = true);
 
 	FORCEINLINE bool IsCrowdSimulationEnabled() const { return SimulationState == ECrowdSimulationState::Enabled; }
 	FORCEINLINE bool IsCrowdSimulatioSuspended() const { return bSuspendCrowdSimulation; }
@@ -132,19 +132,19 @@ class AIMODULE_API UCrowdFollowingComponent : public UPathFollowingComponent, pu
 	FORCEINLINE float GetCrowdPathOptimizationRange() const { return PathOptimizationRange; }
 	FORCEINLINE ECrowdAvoidanceQuality::Type GetCrowdAvoidanceQuality() const { return AvoidanceQuality; }
 	FORCEINLINE float GetCrowdAvoidanceRangeMultiplier() const { return AvoidanceRangeMultiplier; }
-	int32 GetAvoidanceGroup() const;
-	int32 GetGroupsToAvoid() const;
-	int32 GetGroupsToIgnore() const;
+	AIMODULE_API int32 GetAvoidanceGroup() const;
+	AIMODULE_API int32 GetGroupsToAvoid() const;
+	AIMODULE_API int32 GetGroupsToIgnore() const;
 
-	virtual void GetDebugStringTokens(TArray<FString>& Tokens, TArray<EPathFollowingDebugTokens::Type>& Flags) const override;
+	AIMODULE_API virtual void GetDebugStringTokens(TArray<FString>& Tokens, TArray<EPathFollowingDebugTokens::Type>& Flags) const override;
 #if ENABLE_VISUAL_LOG
-	virtual void DescribeSelfToVisLog(struct FVisualLogEntry* Snapshot) const override;
+	AIMODULE_API virtual void DescribeSelfToVisLog(struct FVisualLogEntry* Snapshot) const override;
 #endif // ENABLE_VISUAL_LOG
 
 	UE_DEPRECATED_FORGAME(4.16, "Use ApplyCrowdAgentVelocity function with bIsNearEndOfPath param instead.")
 	virtual void ApplyCrowdAgentVelocity(const FVector& NewVelocity, const FVector& DestPathCorner, bool bTraversingLink) {}
 
-	void UpdateDestinationForMovingGoal(const FVector& NewDestination);
+	AIMODULE_API void UpdateDestinationForMovingGoal(const FVector& NewDestination);
 
 protected:
 
@@ -223,24 +223,24 @@ protected:
 	int32 LastPathPolyIndex;
 
 	// PathFollowingComponent BEGIN
-	virtual int32 DetermineStartingPathPoint(const FNavigationPath* ConsideredPath) const override;
-	virtual void SetMoveSegment(int32 SegmentStartIndex) override;
-	virtual void UpdatePathSegment() override;
-	virtual void FollowPathSegment(float DeltaTime) override;
-	virtual bool ShouldCheckPathOnResume() const override;
-	virtual bool IsOnPath() const override;
-	virtual bool UpdateMovementComponent(bool bForce) override;
-	virtual void Reset() override;
+	AIMODULE_API virtual int32 DetermineStartingPathPoint(const FNavigationPath* ConsideredPath) const override;
+	AIMODULE_API virtual void SetMoveSegment(int32 SegmentStartIndex) override;
+	AIMODULE_API virtual void UpdatePathSegment() override;
+	AIMODULE_API virtual void FollowPathSegment(float DeltaTime) override;
+	AIMODULE_API virtual bool ShouldCheckPathOnResume() const override;
+	AIMODULE_API virtual bool IsOnPath() const override;
+	AIMODULE_API virtual bool UpdateMovementComponent(bool bForce) override;
+	AIMODULE_API virtual void Reset() override;
 	// PathFollowingComponent END
 
-	void SwitchToNextPathPart();
-	bool ShouldSwitchPathPart(int32 CorridorSize) const;
-	bool HasMovedDuringPause() const;
-	void UpdateCachedDirections(const FVector& NewVelocity, const FVector& NextPathCorner, bool bTraversingLink);
+	AIMODULE_API void SwitchToNextPathPart();
+	AIMODULE_API bool ShouldSwitchPathPart(int32 CorridorSize) const;
+	AIMODULE_API bool HasMovedDuringPause() const;
+	AIMODULE_API void UpdateCachedDirections(const FVector& NewVelocity, const FVector& NextPathCorner, bool bTraversingLink);
 
-	virtual bool ShouldTrackMovingGoal(FVector& OutGoalLocation) const;
+	AIMODULE_API virtual bool ShouldTrackMovingGoal(FVector& OutGoalLocation) const;
 	
-	void RegisterCrowdAgent();
+	AIMODULE_API void RegisterCrowdAgent();
 
 	friend UCrowdManager;
 };
diff --git a/Engine/Source/Runtime/AIModule/Classes/Navigation/CrowdManager.h b/Engine/Source/Runtime/AIModule/Classes/Navigation/CrowdManager.h
--- a/Engine/Source/Runtime/AIModule/Classes/Navigation/CrowdManager.h
+++ b/Engine/Source/Runtime/AIModule/Classes/Navigation/CrowdManager.h
@@ -64,7 +64,7 @@ class dtObstacleAvoidanceDebugData;
  */
 
 USTRUCT()
-struct AIMODULE_API FCrowdAvoidanceConfig
+struct FCrowdAvoidanceConfig
 {
 	GENERATED_USTRUCT_BODY()
 
@@ -110,7 +110,7 @@ struct AIMODULE_API FCrowdAvoidanceConfig
 };
 
 USTRUCT()
-struct AIMODULE_API FCrowdAvoidanceSamplingPattern
+struct FCrowdAvoidanceSamplingPattern
 {
 	GENERATED_USTRUCT_BODY()
 
@@ -120,11 +120,11 @@ struct AIMODULE_API FCrowdAvoidanceSamplingPattern
 	UPROPERTY(EditAnywhere, Category = Crowd)
 	TArray<float> Radii;
 
-	void AddSample(float AngleInDegrees, float NormalizedRadius);
-	void AddSampleWithMirror(float AngleInDegrees, float NormalizedRadius);
+	AIMODULE_API void AddSample(float AngleInDegrees, float NormalizedRadius);
+	AIMODULE_API void AddSampleWithMirror(float AngleInDegrees, float NormalizedRadius);
 };
 
-struct AIMODULE_API FCrowdAgentData
+struct FCrowdAgentData
 {
 #if WITH_RECAST
 	/** special filter for checking offmesh links */
@@ -149,7 +149,7 @@ struct AIMODULE_API FCrowdAgentData
 	FCrowdAgentData() :	PrevPoly(0), AgentIndex(-1), PathOptRemainingTime(0), bIsSimulated(false), bWantsPathOptimization(false) {}
 
 	bool IsValid() const { return AgentIndex >= 0; }
-	void ClearFilter();
+	AIMODULE_API void ClearFilter();
 };
 
 struct FCrowdTickHelper : FTickableGameObject
@@ -164,106 +164,106 @@ struct FCrowdTickHelper : FTickableGameObject
 	virtual TStatId GetStatId() const;
 };
 
-UCLASS(config = Engine, defaultconfig)
-class AIMODULE_API UCrowdManager : public UCrowdManagerBase
+UCLASS(config = Engine, defaultconfig, MinimalAPI)
+class UCrowdManager : public UCrowdManagerBase
 {
 	GENERATED_BODY()
 
 public:
-	UCrowdManager(const FObjectInitializer& ObjectInitializer = FObjectInitializer::Get());
+	AIMODULE_API UCrowdManager(const FObjectInitializer& ObjectInitializer = FObjectInitializer::Get());
 
-	virtual void Tick(float DeltaTime) override;
-	virtual void BeginDestroy() override;
+	AIMODULE_API virtual void Tick(float DeltaTime) override;
+	AIMODULE_API virtual void BeginDestroy() override;
 
 	/** adds new agent to crowd */
-	void RegisterAgent(ICrowdAgentInterface* Agent);
+	AIMODULE_API void RegisterAgent(ICrowdAgentInterface* Agent);
 
 	/** removes agent from crowd */
-	void UnregisterAgent(const ICrowdAgentInterface* Agent);
+	AIMODULE_API void UnregisterAgent(const ICrowdAgentInterface* Agent);
 
 	/** updates agent data */
-	void UpdateAgentParams(const ICrowdAgentInterface* Agent) const;
+	AIMODULE_API void UpdateAgentParams(const ICrowdAgentInterface* Agent) const;
 
 	/** refresh agent state */
-	void UpdateAgentState(const ICrowdAgentInterface* Agent) const;
+	AIMODULE_API void UpdateAgentState(const ICrowdAgentInterface* Agent) const;
 
 	/** update agent after using custom link */
-	void OnAgentFinishedCustomLink(const ICrowdAgentInterface* Agent) const;
+	AIMODULE_API void OnAgentFinishedCustomLink(const ICrowdAgentInterface* Agent) const;
 
 	/** sets move target for crowd agent (only for fully simulated) */
-	bool SetAgentMoveTarget(const UCrowdFollowingComponent* AgentComponent, const FVector& MoveTarget, FSharedConstNavQueryFilter Filter) const;
+	AIMODULE_API bool SetAgentMoveTarget(const UCrowdFollowingComponent* AgentComponent, const FVector& MoveTarget, FSharedConstNavQueryFilter Filter) const;
 
 	/** sets move direction for crowd agent (only for fully simulated) */
-	bool SetAgentMoveDirection(const UCrowdFollowingComponent* AgentComponent, const FVector& MoveDirection) const;
+	AIMODULE_API bool SetAgentMoveDirection(const UCrowdFollowingComponent* AgentComponent, const FVector& MoveDirection) const;
 
 	/** sets move target using path (only for fully simulated) */
-	bool SetAgentMovePath(const UCrowdFollowingComponent* AgentComponent, const FNavMeshPath* Path, int32 PathSectionStart, int32 PathSectionEnd, const FVector& PathSectionEndLocation) const;
+	AIMODULE_API bool SetAgentMovePath(const UCrowdFollowingComponent* AgentComponent, const FNavMeshPath* Path, int32 PathSectionStart, int32 PathSectionEnd, const FVector& PathSectionEndLocation) const;
 
 	/** clears move target for crowd agent (only for fully simulated) */
-	void ClearAgentMoveTarget(const UCrowdFollowingComponent* AgentComponent) const;
+	AIMODULE_API void ClearAgentMoveTarget(const UCrowdFollowingComponent* AgentComponent) const;
 
 	/** switch agent to waiting state */
-	void PauseAgent(const UCrowdFollowingComponent* AgentComponent) const;
+	AIMODULE_API void PauseAgent(const UCrowdFollowingComponent* AgentComponent) const;
 
 	/** resumes agent movement */
-	void ResumeAgent(const UCrowdFollowingComponent* AgentComponent, bool bForceReplanPath = true) const;
+	AIMODULE_API void ResumeAgent(const UCrowdFollowingComponent* AgentComponent, bool bForceReplanPath = true) const;
 
 	/** check if object is a valid crowd agent */
-	bool IsAgentValid(const UCrowdFollowingComponent* AgentComponent) const;
-	bool IsAgentValid(const ICrowdAgentInterface* Agent) const;
+	AIMODULE_API bool IsAgentValid(const UCrowdFollowingComponent* AgentComponent) const;
+	AIMODULE_API bool IsAgentValid(const ICrowdAgentInterface* Agent) const;
 
 	/** returns number of nearby agents */
-	int32 GetNumNearbyAgents(const ICrowdAgentInterface* Agent) const;
+	AIMODULE_API int32 GetNumNearbyAgents(const ICrowdAgentInterface* Agent) const;
 
 	/** returns a list of locations of nearby agents */
-	int32 GetNearbyAgentLocations(const ICrowdAgentInterface* Agent, TArray<FVector>& OutLocations) const;
+	AIMODULE_API int32 GetNearbyAgentLocations(const ICrowdAgentInterface* Agent, TArray<FVector>& OutLocations) const;
 
 	/** reads existing avoidance config or returns false */
-	bool GetAvoidanceConfig(int32 Idx, FCrowdAvoidanceConfig& Data) const;
+	AIMODULE_API bool GetAvoidanceConfig(int32 Idx, FCrowdAvoidanceConfig& Data) const;
 
 	/** updates existing avoidance config or returns false */
-	bool SetAvoidanceConfig(int32 Idx, const FCrowdAvoidanceConfig& Data);
+	AIMODULE_API bool SetAvoidanceConfig(int32 Idx, const FCrowdAvoidanceConfig& Data);
 
 	/** remove started offmesh connections from corridor */
-	void SetOffmeshConnectionPruning(bool bRemoveFromCorridor);
+	AIMODULE_API void SetOffmeshConnectionPruning(bool bRemoveFromCorridor);
 
 	/** block path visibility raycasts when crossing different nav areas */
-	void SetSingleAreaVisibilityOptimization(bool bEnable);
+	AIMODULE_API void SetSingleAreaVisibilityOptimization(bool bEnable);
 
 	/** adjust current position in path's corridor, starting test from PathStartIdx */
-	void AdjustAgentPathStart(const UCrowdFollowingComponent* AgentComponent, const FNavMeshPath* Path, int32& PathStartIdx) const;
+	AIMODULE_API void AdjustAgentPathStart(const UCrowdFollowingComponent* AgentComponent, const FNavMeshPath* Path, int32& PathStartIdx) const;
 
 #if WITH_EDITOR
-	virtual void PostEditChangeProperty(struct FPropertyChangedEvent& PropertyChangedEvent) override;
+	AIMODULE_API virtual void PostEditChangeProperty(struct FPropertyChangedEvent& PropertyChangedEvent) override;
 
-	void DebugTick() const;
+	AIMODULE_API void DebugTick() const;
 #endif
 
 	/** notify called when detour navmesh is changed */
-	void OnNavMeshUpdate();
+	AIMODULE_API void OnNavMeshUpdate();
 
 	/** Tests if NavData is a suitable nav data type to be used by this CrowdManager
 	 *	instance. */
-	virtual bool IsSuitableNavData(const ANavigationData& NavData) const;
+	AIMODULE_API virtual bool IsSuitableNavData(const ANavigationData& NavData) const;
 
 	/** Called by the nav system when a new navigation data instance is registered. 
 	 *	If the CrowdManager instance had no nav data cached it will consider this
 	 *	NavDataInstance and update if necesary. */
-	virtual void OnNavDataRegistered(ANavigationData& NavData) override;
+	AIMODULE_API virtual void OnNavDataRegistered(ANavigationData& NavData) override;
 
 	/** Called by the nav system when a navigation data instance is removed. The 
 	 *	crowd manager will see if it's the nav data being used by it an if so try
 	 *	to find another one. If there's none the crowd manager will stop working. */
-	virtual void OnNavDataUnregistered(ANavigationData& NavData) override;
+	AIMODULE_API virtual void OnNavDataUnregistered(ANavigationData& NavData) override;
 
 	virtual void CleanUp(float DeltaTime) override {};
 
 	const ANavigationData* GetNavData() const { return MyNavData; }
 
-	UWorld* GetWorld() const override;
+	AIMODULE_API UWorld* GetWorld() const override;
 
-	static UCrowdManager* GetCurrent(UObject* WorldContextObject);
-	static UCrowdManager* GetCurrent(UWorld* World);
+	static AIMODULE_API UCrowdManager* GetCurrent(UObject* WorldContextObject);
+	static AIMODULE_API UCrowdManager* GetCurrent(UWorld* World);
 
 protected:
 
@@ -339,51 +339,51 @@ protected:
 #endif // WITH_EDITOR
 
 	/** try to initialize nav data from already existing ones */
-	virtual void UpdateNavData();
+	AIMODULE_API virtual void UpdateNavData();
 
 	/** setup params of crowd avoidance */
-	virtual void UpdateAvoidanceConfig();
+	AIMODULE_API virtual void UpdateAvoidanceConfig();
 
 	/** called from tick, just after updating agents proximity data */
-	virtual void PostProximityUpdate();
+	AIMODULE_API virtual void PostProximityUpdate();
 
 	/** called from tick, after move points were updated, before any steering/avoidance */
-	virtual void PostMovePointUpdate();
+	AIMODULE_API virtual void PostMovePointUpdate();
 
 	/** Sets NavData as MyNavData. If Null and bFindNewNavDataIfNull is true then
 	 *	the manager will search for a new NavData instance that meets the 
 	 *	IsSuitableNavData() condition. */
-	void SetNavData(ANavigationData* NavData, const bool bFindNewNavDataIfNull = true);
+	AIMODULE_API void SetNavData(ANavigationData* NavData, const bool bFindNewNavDataIfNull = true);
 
 #if WITH_RECAST
-	void AddAgent(const ICrowdAgentInterface* Agent, FCrowdAgentData& AgentData) const;
-	void RemoveAgent(const ICrowdAgentInterface* Agent, FCrowdAgentData* AgentData) const;
-	void GetAgentParams(const ICrowdAgentInterface* Agent, dtCrowdAgentParams& AgentParams) const;
+	AIMODULE_API void AddAgent(const ICrowdAgentInterface* Agent, FCrowdAgentData& AgentData) const;
+	AIMODULE_API void RemoveAgent(const ICrowdAgentInterface* Agent, FCrowdAgentData* AgentData) const;
+	AIMODULE_API void GetAgentParams(const ICrowdAgentInterface* Agent, dtCrowdAgentParams& AgentParams) const;
 
 	/** prepare agent for next step of simulation */
-	void PrepareAgentStep(const ICrowdAgentInterface* Agent, FCrowdAgentData& AgentData, float DeltaTime) const;
+	AIMODULE_API void PrepareAgentStep(const ICrowdAgentInterface* Agent, FCrowdAgentData& AgentData, float DeltaTime) const;
 
 	/** pass new velocity to movement components */
-	virtual void ApplyVelocity(UCrowdFollowingComponent* AgentComponent, int32 AgentIndex) const;
+	AIMODULE_API virtual void ApplyVelocity(UCrowdFollowingComponent* AgentComponent, int32 AgentIndex) const;
 
 	/** check changes in crowd simulation and adjust UE specific properties (smart links, poly updates) */
-	void UpdateAgentPaths();
+	AIMODULE_API void UpdateAgentPaths();
 
 	/** switch debugger to object selected in PIE */
-	void UpdateSelectedDebug(const ICrowdAgentInterface* Agent, int32 AgentIndex) const;
+	AIMODULE_API void UpdateSelectedDebug(const ICrowdAgentInterface* Agent, int32 AgentIndex) const;
 
-	void CreateCrowdManager();
-	void DestroyCrowdManager();
+	AIMODULE_API void CreateCrowdManager();
+	AIMODULE_API void DestroyCrowdManager();
 
 #if ENABLE_DRAW_DEBUG
-	UWorld* GetDebugDrawingWorld() const;
-	void DrawDebugCorners(const dtCrowdAgent* CrowdAgent) const;
-	void DrawDebugCollisionSegments(const dtCrowdAgent* CrowdAgent) const;
-	void DrawDebugPath(const dtCrowdAgent* CrowdAgent) const;
-	void DrawDebugVelocityObstacles(const dtCrowdAgent* CrowdAgent) const;
-	void DrawDebugPathOptimization(const dtCrowdAgent* CrowdAgent) const;
-	void DrawDebugNeighbors(const dtCrowdAgent* CrowdAgent) const;
-	void DrawDebugSharedBoundary() const;
+	AIMODULE_API UWorld* GetDebugDrawingWorld() const;
+	AIMODULE_API void DrawDebugCorners(const dtCrowdAgent* CrowdAgent) const;
+	AIMODULE_API void DrawDebugCollisionSegments(const dtCrowdAgent* CrowdAgent) const;
+	AIMODULE_API void DrawDebugPath(const dtCrowdAgent* CrowdAgent) const;
+	AIMODULE_API void DrawDebugVelocityObstacles(const dtCrowdAgent* CrowdAgent) const;
+	AIMODULE_API void DrawDebugPathOptimization(const dtCrowdAgent* CrowdAgent) const;
+	AIMODULE_API void DrawDebugNeighbors(const dtCrowdAgent* CrowdAgent) const;
+	AIMODULE_API void DrawDebugSharedBoundary() const;
 #endif // ENABLE_DRAW_DEBUG
 
 #endif
diff --git a/Engine/Source/Runtime/AIModule/Classes/Navigation/GridPathFollowingComponent.h b/Engine/Source/Runtime/AIModule/Classes/Navigation/GridPathFollowingComponent.h
--- a/Engine/Source/Runtime/AIModule/Classes/Navigation/GridPathFollowingComponent.h
+++ b/Engine/Source/Runtime/AIModule/Classes/Navigation/GridPathFollowingComponent.h
@@ -17,17 +17,17 @@ class UNavLocalGridManager;
  *
  *  Does not replace proper avoidance for dynamic obstacles!
  */
-UCLASS(BlueprintType, Experimental)
-class AIMODULE_API UGridPathFollowingComponent : public UPathFollowingComponent
+UCLASS(BlueprintType, Experimental, MinimalAPI)
+class UGridPathFollowingComponent : public UPathFollowingComponent
 {
 	GENERATED_UCLASS_BODY()
 public:
 
-	virtual void Initialize() override;
-	virtual void UpdatePathSegment() override;
-	virtual void Reset() override;
-	virtual void ResumeMove(FAIRequestID RequestID = FAIRequestID::CurrentRequest) override;
-	virtual void OnPathUpdated();
+	AIMODULE_API virtual void Initialize() override;
+	AIMODULE_API virtual void UpdatePathSegment() override;
+	AIMODULE_API virtual void Reset() override;
+	AIMODULE_API virtual void ResumeMove(FAIRequestID RequestID = FAIRequestID::CurrentRequest) override;
+	AIMODULE_API virtual void OnPathUpdated();
 
 	bool HasActiveGrid() const { return ActiveGridIdx != INDEX_NONE; }
 	int32 GetActiveGridIdx() const { return ActiveGridIdx; }
@@ -37,7 +37,7 @@ public:
 
 protected:
 
-	void UpdateActiveGrid(const FVector& CurrentLocation);
+	AIMODULE_API void UpdateActiveGrid(const FVector& CurrentLocation);
 
 	UPROPERTY(Transient)
 	TObjectPtr<UNavLocalGridManager> GridManager;
diff --git a/Engine/Source/Runtime/AIModule/Classes/Navigation/MetaNavMeshPath.h b/Engine/Source/Runtime/AIModule/Classes/Navigation/MetaNavMeshPath.h
--- a/Engine/Source/Runtime/AIModule/Classes/Navigation/MetaNavMeshPath.h
+++ b/Engine/Source/Runtime/AIModule/Classes/Navigation/MetaNavMeshPath.h
@@ -43,31 +43,31 @@ struct FMetaPathWayPoint : public FVector
  *
  * Since path updates itself for moving agent it really shouldn't be reused by others.
  */
-struct AIMODULE_API FMetaNavMeshPath : public FNavMeshPath
+struct FMetaNavMeshPath : public FNavMeshPath
 {
 	typedef FNavMeshPath Super;
-	static const FNavPathType Type;
+	static AIMODULE_API const FNavPathType Type;
 
-	FMetaNavMeshPath();
-	FMetaNavMeshPath(const TArray<FMetaPathWayPoint>& InWaypoints, const ANavigationData& NavData);
-	FMetaNavMeshPath(const TArray<FMetaPathWayPoint>& InWaypoints, const AController& Owner);
-	FMetaNavMeshPath(const TArray<FVector>& InWaypoints, const ANavigationData& NavData);
-	FMetaNavMeshPath(const TArray<FVector>& InWaypoints, const AController& Owner);
+	AIMODULE_API FMetaNavMeshPath();
+	AIMODULE_API FMetaNavMeshPath(const TArray<FMetaPathWayPoint>& InWaypoints, const ANavigationData& NavData);
+	AIMODULE_API FMetaNavMeshPath(const TArray<FMetaPathWayPoint>& InWaypoints, const AController& Owner);
+	AIMODULE_API FMetaNavMeshPath(const TArray<FVector>& InWaypoints, const ANavigationData& NavData);
+	AIMODULE_API FMetaNavMeshPath(const TArray<FVector>& InWaypoints, const AController& Owner);
 
 	/** initialize path for path following */
-	virtual void Initialize(const FVector& AgentLocation);
+	AIMODULE_API virtual void Initialize(const FVector& AgentLocation);
 
 	/** try switching to next waypoint, depends on WaypointSwitchRadius */
-	virtual bool ConditionalMoveToNextSection(const FVector& AgentLocation, EMetaPathUpdateReason Reason);
+	AIMODULE_API virtual bool ConditionalMoveToNextSection(const FVector& AgentLocation, EMetaPathUpdateReason Reason);
 
 	/** force switching to next waypoint */
-	bool ForceMoveToNextSection(const FVector& AgentLocation);
+	AIMODULE_API bool ForceMoveToNextSection(const FVector& AgentLocation);
 
 	/** updates underlying navmesh path for current target waypoint */
-	virtual bool UpdatePath(const FVector& AgentLocation);
+	AIMODULE_API virtual bool UpdatePath(const FVector& AgentLocation);
 
 	/** copy properties of other meta path */
-	virtual void CopyFrom(const FMetaNavMeshPath& Other);
+	AIMODULE_API virtual void CopyFrom(const FMetaNavMeshPath& Other);
 
 	/** returns true if path at last waypoint */
 	bool IsLastSection() const { return (TargetWaypointIdx == (Waypoints.Num() - 1)) && (Waypoints.Num() > 0); }
@@ -85,10 +85,10 @@ struct AIMODULE_API FMetaNavMeshPath : public FNavMeshPath
 	AActor* GetMetaPathGoal() const { return PathGoal.Get(); }
 
 	/** tries to set waypoints, fails when path is ready being followed */
-	bool SetWaypoints(const TArray<FMetaPathWayPoint>& InWaypoints);
+	AIMODULE_API bool SetWaypoints(const TArray<FMetaPathWayPoint>& InWaypoints);
 
 	/** tries to set waypoints, fails when path is ready being followed */
-	bool SetWaypoints(const TArray<FVector>& InWaypoints);
+	AIMODULE_API bool SetWaypoints(const TArray<FVector>& InWaypoints);
 	
 	/** returns radius for switching to next waypoint during path following */
 	float GetWaypointSwitchRadius() const { return WaypointSwitchRadius; }
@@ -97,15 +97,15 @@ struct AIMODULE_API FMetaNavMeshPath : public FNavMeshPath
 	void SetWaypointSwitchRadius(float InSwitchRadius) { WaypointSwitchRadius = InSwitchRadius; }
 
 	/** returns approximate length of path, ignores parameters */
-	virtual FVector::FReal GetLengthFromPosition(FVector SegmentStart, uint32 NextPathPointIndex) const override;
+	AIMODULE_API virtual FVector::FReal GetLengthFromPosition(FVector SegmentStart, uint32 NextPathPointIndex) const override;
 	
 	/** returns approximate cost of path, ignores parameter */
-	virtual FVector::FReal GetCostFromIndex(int32 PathPointIndex) const override;
+	AIMODULE_API virtual FVector::FReal GetCostFromIndex(int32 PathPointIndex) const override;
 
 #if ENABLE_VISUAL_LOG
-	virtual void DescribeSelfToVisLog(FVisualLogEntry* Snapshot) const override;
+	AIMODULE_API virtual void DescribeSelfToVisLog(FVisualLogEntry* Snapshot) const override;
 #endif
-	virtual void DebugDraw(const ANavigationData* NavData, const FColor PathColor, UCanvas* Canvas, const bool bPersistent, const float LifeTime, const uint32 NextPathPointIndex = 0) const override;
+	AIMODULE_API virtual void DebugDraw(const ANavigationData* NavData, const FColor PathColor, UCanvas* Canvas, const bool bPersistent, const float LifeTime, const uint32 NextPathPointIndex = 0) const override;
 
 protected:
 
@@ -128,5 +128,5 @@ protected:
 	float PathGoalTetherDistance;
 
 	/** switch to next waypoint  */
-	bool MoveToNextSection(const FVector& AgentLocation);
+	AIMODULE_API bool MoveToNextSection(const FVector& AgentLocation);
 };
diff --git a/Engine/Source/Runtime/AIModule/Classes/Navigation/NavFilter_AIControllerDefault.h b/Engine/Source/Runtime/AIModule/Classes/Navigation/NavFilter_AIControllerDefault.h
--- a/Engine/Source/Runtime/AIModule/Classes/Navigation/NavFilter_AIControllerDefault.h
+++ b/Engine/Source/Runtime/AIModule/Classes/Navigation/NavFilter_AIControllerDefault.h
@@ -9,12 +9,12 @@
 #include "NavFilters/NavigationQueryFilter.h"
 #include "NavFilter_AIControllerDefault.generated.h"
 
-UCLASS()
-class AIMODULE_API UNavFilter_AIControllerDefault : public UNavigationQueryFilter
+UCLASS(MinimalAPI)
+class UNavFilter_AIControllerDefault : public UNavigationQueryFilter
 {
 	GENERATED_BODY()
 public:
-	UNavFilter_AIControllerDefault(const FObjectInitializer& ObjectInitializer = FObjectInitializer::Get());
+	AIMODULE_API UNavFilter_AIControllerDefault(const FObjectInitializer& ObjectInitializer = FObjectInitializer::Get());
 
-	virtual TSubclassOf<UNavigationQueryFilter> GetSimpleFilterForAgent(const UObject& Querier) const;
+	AIMODULE_API virtual TSubclassOf<UNavigationQueryFilter> GetSimpleFilterForAgent(const UObject& Querier) const;
 };
diff --git a/Engine/Source/Runtime/AIModule/Classes/Navigation/NavLinkProxy.h b/Engine/Source/Runtime/AIModule/Classes/Navigation/NavLinkProxy.h
--- a/Engine/Source/Runtime/AIModule/Classes/Navigation/NavLinkProxy.h
+++ b/Engine/Source/Runtime/AIModule/Classes/Navigation/NavLinkProxy.h
@@ -30,8 +30,8 @@ DECLARE_DYNAMIC_MULTICAST_DELEGATE_TwoParams( FSmartLinkReachedSignature, AActor
  *  There can only be at most one smart link per ANavLinkProxy instance.
  *  Both simple and smart links on a single ANavLinkProxy instance, can be set / enabled at once, as well as either or neither of them.
  */
-UCLASS(Blueprintable, autoCollapseCategories=(SmartLink, Actor), hideCategories=(Input))
-class AIMODULE_API ANavLinkProxy : public AActor, public INavLinkHostInterface, public INavRelevantInterface
+UCLASS(Blueprintable, autoCollapseCategories=(SmartLink, Actor), hideCategories=(Input), MinimalAPI)
+class ANavLinkProxy : public AActor, public INavLinkHostInterface, public INavRelevantInterface
 {
 	GENERATED_UCLASS_BODY()
 
@@ -66,64 +66,64 @@ public:
 #endif // WITH_EDITORONLY_DATA
 
 	// BEGIN INavRelevantInterface
-	virtual void GetNavigationData(FNavigationRelevantData& Data) const override;
-	virtual FBox GetNavigationBounds() const override;
-	virtual bool IsNavigationRelevant() const override;
+	AIMODULE_API virtual void GetNavigationData(FNavigationRelevantData& Data) const override;
+	AIMODULE_API virtual FBox GetNavigationBounds() const override;
+	AIMODULE_API virtual bool IsNavigationRelevant() const override;
 	// END INavRelevantInterface
 
 	// BEGIN INavLinkHostInterface
-	virtual bool GetNavigationLinksClasses(TArray<TSubclassOf<UNavLinkDefinition> >& OutClasses) const override;
-	virtual bool GetNavigationLinksArray(TArray<FNavigationLink>& OutLink, TArray<FNavigationSegmentLink>& OutSegments) const override;
+	AIMODULE_API virtual bool GetNavigationLinksClasses(TArray<TSubclassOf<UNavLinkDefinition> >& OutClasses) const override;
+	AIMODULE_API virtual bool GetNavigationLinksArray(TArray<FNavigationLink>& OutLink, TArray<FNavigationSegmentLink>& OutSegments) const override;
 	// END INavLinkHostInterface
 
 #if WITH_EDITOR
-	virtual void PostEditChangeProperty(FPropertyChangedEvent& PropertyChangedEvent) override;
-	virtual void PostEditUndo() override;
-	virtual void PostEditImport() override;
+	AIMODULE_API virtual void PostEditChangeProperty(FPropertyChangedEvent& PropertyChangedEvent) override;
+	AIMODULE_API virtual void PostEditUndo() override;
+	AIMODULE_API virtual void PostEditImport() override;
 #endif // WITH_EDITOR
-	virtual void PostInitProperties() override;
-	virtual void BeginDestroy() override;
+	AIMODULE_API virtual void PostInitProperties() override;
+	AIMODULE_API virtual void BeginDestroy() override;
 
-	virtual void PostRegisterAllComponents() override;
-	virtual void PostLoad() override;
+	AIMODULE_API virtual void PostRegisterAllComponents() override;
+	AIMODULE_API virtual void PostLoad() override;
 
 #if ENABLE_VISUAL_LOG
 protected:
-	virtual void BeginPlay() override;
+	AIMODULE_API virtual void BeginPlay() override;
 public:
 #endif // ENABLE_VISUAL_LOG
 
-	virtual FBox GetComponentsBoundingBox(bool bNonColliding = false, bool bIncludeFromChildActors = false) const override;
+	AIMODULE_API virtual FBox GetComponentsBoundingBox(bool bNonColliding = false, bool bIncludeFromChildActors = false) const override;
 
 	//////////////////////////////////////////////////////////////////////////
 	// Blueprint interface for smart links
 
 	/** called when agent reaches smart link during path following, use ResumePathFollowing() to give control back */
 	UFUNCTION(BlueprintImplementableEvent)
-	void ReceiveSmartLinkReached(AActor* Agent, const FVector& Destination);
+	AIMODULE_API void ReceiveSmartLinkReached(AActor* Agent, const FVector& Destination);
 
 	/** resume normal path following */
 	UFUNCTION(BlueprintCallable, Category="AI|Navigation")
-	void ResumePathFollowing(AActor* Agent);
+	AIMODULE_API void ResumePathFollowing(AActor* Agent);
 
 	/** check if smart link is enabled */
 	UFUNCTION(BlueprintCallable, Category="AI|Navigation")
-	bool IsSmartLinkEnabled() const;
+	AIMODULE_API bool IsSmartLinkEnabled() const;
 
 	/** change state of smart link */
 	UFUNCTION(BlueprintCallable, Category="AI|Navigation")
-	void SetSmartLinkEnabled(bool bEnabled);
+	AIMODULE_API void SetSmartLinkEnabled(bool bEnabled);
 
 	/** check if any agent is moving through smart link right now */
 	UFUNCTION(BlueprintCallable, Category="AI|Navigation")
-	bool HasMovingAgents() const;
+	AIMODULE_API bool HasMovingAgents() const;
 
 #if WITH_EDITOR
 	/** Copies navlink end points from the first entry in PointLinks array. This 
 	 *	function is a helper function making up for smart links not drawing
 	 *	the FVector widgets in the editor. */
 	UFUNCTION(CallInEditor, Category = SmartLink, meta = (DisplayName="CopyEndPointsFromSimpleLink"))
-	void CopyEndPointsFromSimpleLinkToSmartLink();
+	AIMODULE_API void CopyEndPointsFromSimpleLinkToSmartLink();
 #endif // WITH_EDITOR
 
 protected:
@@ -131,11 +131,11 @@ protected:
 	UPROPERTY(BlueprintAssignable)
 	FSmartLinkReachedSignature OnSmartLinkReached;
 
-	void NotifySmartLinkReached(UNavLinkCustomComponent* LinkComp, UObject* PathingAgent, const FVector& DestPoint);
+	AIMODULE_API void NotifySmartLinkReached(UNavLinkCustomComponent* LinkComp, UObject* PathingAgent, const FVector& DestPoint);
 
 #if WITH_EDITOR
-	void OnNavAreaRegistered(const UWorld& World, const UClass* NavAreaClass);
-	void OnNavAreaUnregistered(const UWorld& World, const UClass* NavAreaClass);
+	AIMODULE_API void OnNavAreaRegistered(const UWorld& World, const UClass* NavAreaClass);
+	AIMODULE_API void OnNavAreaUnregistered(const UWorld& World, const UClass* NavAreaClass);
 #endif // WITH_EDITOR
 
 public:
diff --git a/Engine/Source/Runtime/AIModule/Classes/Navigation/NavLocalGridData.h b/Engine/Source/Runtime/AIModule/Classes/Navigation/NavLocalGridData.h
--- a/Engine/Source/Runtime/AIModule/Classes/Navigation/NavLocalGridData.h
+++ b/Engine/Source/Runtime/AIModule/Classes/Navigation/NavLocalGridData.h
@@ -20,24 +20,24 @@ struct FNavigationPath;
  *  TODO: serialization (with maps?)
  *  TODO: FNavigationPath support?
  */
-struct AIMODULE_API FNavLocalGridData : public TSimpleCellGrid<uint8, MAX_uint8>
+struct FNavLocalGridData : public TSimpleCellGrid<uint8, MAX_uint8>
 {
 	FNavLocalGridData() : GridId(0) {}
-	FNavLocalGridData(const FVector& Center, float Extent2D);
-	FNavLocalGridData(const FVector& Center, const FVector2D& Extent2D);
-	FNavLocalGridData(const TArray<FNavLocalGridData>& SourceGrids);
+	AIMODULE_API FNavLocalGridData(const FVector& Center, float Extent2D);
+	AIMODULE_API FNavLocalGridData(const FVector& Center, const FVector2D& Extent2D);
+	AIMODULE_API FNavLocalGridData(const TArray<FNavLocalGridData>& SourceGrids);
 
 	/** mark single cell as obstacle */
-	void MarkPointObstacle(const FVector& Center);
+	AIMODULE_API void MarkPointObstacle(const FVector& Center);
 
 	/** mark box (AABB or rotated) shape as obstacle */
-	void MarkBoxObstacle(const FVector& Center, const FVector& Extent, const FQuat& Quat = FQuat::Identity);
+	AIMODULE_API void MarkBoxObstacle(const FVector& Center, const FVector& Extent, const FQuat& Quat = FQuat::Identity);
 
 	/** mark capsule shape as obstacle */
-	void MarkCapsuleObstacle(const FVector& Center, float Radius, float HalfHeight);
+	AIMODULE_API void MarkCapsuleObstacle(const FVector& Center, float Radius, float HalfHeight);
 
 	/** set height of bounds, if not set: ProjectCells will use height of default query box */
-	void SetHeight(float ExtentZ);
+	AIMODULE_API void SetHeight(float ExtentZ);
 
 	/** get unique Id of grid */
 	const int32 GetGridId() const
@@ -58,7 +58,7 @@ struct AIMODULE_API FNavLocalGridData : public TSimpleCellGrid<uint8, MAX_uint8>
 	}
 
 	/** convert global world coords to cell index, return -1 if outside */
-	int32 GetCellIndexFromGlobalCoords2D(const FIntVector& WorldCoords) const;
+	AIMODULE_API int32 GetCellIndexFromGlobalCoords2D(const FIntVector& WorldCoords) const;
 
 	/** convert cell index to world location using projected heights */
 	FVector GetProjectedCellCenter(int32 CellIdx) const 
@@ -80,13 +80,13 @@ struct AIMODULE_API FNavLocalGridData : public TSimpleCellGrid<uint8, MAX_uint8>
 	 *  @param PathPointsInside - [out] points inside grid
 	 *  @param NextSegmentStart - [out] next move segment on path after leaving grid or -1 if path ends inside
 	 */
-	void FindPathForMovingAgent(const FNavigationPath& SourcePath, const FVector& EntryLocation, int32 EntrySegmentStart, TArray<FVector>& PathPointsInside, int32& NextSegmentStart) const;
+	AIMODULE_API void FindPathForMovingAgent(const FNavigationPath& SourcePath, const FVector& EntryLocation, int32 EntrySegmentStart, TArray<FVector>& PathPointsInside, int32& NextSegmentStart) const;
 
 	/** create path points from StartCoords to EndCoord, returns false when failed */
-	bool FindPath(const FIntVector& StartCoords, const FIntVector& EndCoords, TArray<FIntVector>& PathCoords) const;
+	AIMODULE_API bool FindPath(const FIntVector& StartCoords, const FIntVector& EndCoords, TArray<FIntVector>& PathCoords) const;
 
 	/** project cells on navigation data and marks failed ones as obstacles */
-	void ProjectCells(const ANavigationData& NavData);
+	AIMODULE_API void ProjectCells(const ANavigationData& NavData);
 
 	//////////////////////////////////////////////////////////////////////////
 	// FGraphAStar: TGraph
@@ -94,7 +94,7 @@ struct AIMODULE_API FNavLocalGridData : public TSimpleCellGrid<uint8, MAX_uint8>
 
 	int32 GetNeighbourCount(FNodeRef NodeRef) const { return 8; }
 	bool IsValidRef(FNodeRef NodeRef) const { return IsValidIndex(NodeRef); }
-	FNodeRef GetNeighbour(const FNodeRef NodeRef, const int32 NeiIndex) const;
+	AIMODULE_API FNodeRef GetNeighbour(const FNodeRef NodeRef, const int32 NeiIndex) const;
 	//////////////////////////////////////////////////////////////////////////
 
 protected:
@@ -103,13 +103,13 @@ protected:
 	double LastAccessTime;
 
 	/** convert PathIndices into pruned PathCoords */
-	void PostProcessPath(const FIntVector& StartCoords, const FIntVector& EndCoords, const TArray<int32>& PathIndices, TArray<FIntVector>& PathCoords) const;
+	AIMODULE_API void PostProcessPath(const FIntVector& StartCoords, const FIntVector& EndCoords, const TArray<int32>& PathIndices, TArray<FIntVector>& PathCoords) const;
 
 	/** check if line trace between local coords on grid hits any obstacles, doesn't validate coords! */
-	bool IsLineObstructed(const FIntVector& StartCoords, const FIntVector& EndCoords) const;
+	AIMODULE_API bool IsLineObstructed(const FIntVector& StartCoords, const FIntVector& EndCoords) const;
 
 	/** set unique Id of grid */
-	void SetGridId(int32 NewId);
+	AIMODULE_API void SetGridId(int32 NewId);
 
 private:
 	int32 GridId;
diff --git a/Engine/Source/Runtime/AIModule/Classes/Navigation/NavLocalGridManager.h b/Engine/Source/Runtime/AIModule/Classes/Navigation/NavLocalGridManager.h
--- a/Engine/Source/Runtime/AIModule/Classes/Navigation/NavLocalGridManager.h
+++ b/Engine/Source/Runtime/AIModule/Classes/Navigation/NavLocalGridManager.h
@@ -23,32 +23,32 @@ struct FCombinedNavGridData : public FNavLocalGridData
  *  Check also: UGridPathFollowingComponent, FNavLocalGridData
  */
 
-UCLASS(Experimental)
-class AIMODULE_API UNavLocalGridManager : public UObject
+UCLASS(Experimental, MinimalAPI)
+class UNavLocalGridManager : public UObject
 {
 	GENERATED_UCLASS_BODY()
 public:
 
 	/** adds new grid */
-	int32 AddGridData(const FNavLocalGridData& GridData, bool bUpdate = true);
+	AIMODULE_API int32 AddGridData(const FNavLocalGridData& GridData, bool bUpdate = true);
 
 	/** removes grid */
-	void RemoveGridData(int32 GridId, bool bUpdate = true);
+	AIMODULE_API void RemoveGridData(int32 GridId, bool bUpdate = true);
 
 	/** rebuild overlapping grids if needed */
-	void RebuildGrids();
+	AIMODULE_API void RebuildGrids();
 
 	/** tries to find a path using grids, returns false when start and end locations are not on the same grid */
-	bool FindPath(const FVector& Start, const FVector& End, TArray<FVector>& PathPoints) const;
+	AIMODULE_API bool FindPath(const FVector& Start, const FVector& End, TArray<FVector>& PathPoints) const;
 
 	/** set shared size of grid cell, allowed only when there's no grid added */
-	bool SetCellSize(float CellSize);
+	AIMODULE_API bool SetCellSize(float CellSize);
 
 	/** get shared size of grid cell, static but there's only one active manager */
 	static float GetCellSize() { return GridCellSize; }
 
 	/** set limit of source grids, 0 or negative means unlimited */
-	void SetMaxActiveSources(int32 NumActiveSources);
+	AIMODULE_API void SetMaxActiveSources(int32 NumActiveSources);
 
 	/** get limit of source grids */
 	int32 GetMaxActiveSources() const { return MaxActiveSourceGrids; }
@@ -57,7 +57,7 @@ public:
 	bool HasSourceGridLimit() const { return MaxActiveSourceGrids > 0; }
 
 	/** updates LastAccessTime in all source grids */
-	void UpdateAccessTime(int32 CombinedGridIdx);
+	AIMODULE_API void UpdateAccessTime(int32 CombinedGridIdx);
 
 	/** get number of known source grids */
 	int32 GetNumSources() const { return SourceGrids.Num(); }
@@ -72,42 +72,42 @@ public:
 	const FNavLocalGridData& GetGridData(int32 GridIdx) const { return CombinedGrids[GridIdx]; }
 
 	/** find combined grid value at world location */
-	uint8 GetGridValueAt(const FVector& WorldLocation) const;
+	AIMODULE_API uint8 GetGridValueAt(const FVector& WorldLocation) const;
 
 	/** find combined grid at location */
-	int32 GetGridIndex(const FVector& WorldLocation) const;
+	AIMODULE_API int32 GetGridIndex(const FVector& WorldLocation) const;
 
 	/** get version of grid data, incremented with each rebuild */
 	int32 GetVersion() const { return VersionNum; }
 
 	UFUNCTION(BlueprintCallable, Category = "AI|Navigation", meta = (WorldContext = "WorldContextObject"))
-	static bool SetLocalNavigationGridDensity(UObject* WorldContextObject, float CellSize);
+	static AIMODULE_API bool SetLocalNavigationGridDensity(UObject* WorldContextObject, float CellSize);
 
 	/** creates new grid data for single point */
 	UFUNCTION(BlueprintCallable, Category = "AI|Navigation", meta = (WorldContext = "WorldContextObject"))
-	static int32 AddLocalNavigationGridForPoint(UObject* WorldContextObject, const FVector& Location, const int32 Radius2D = 5, const float Height = 100.0f, bool bRebuildGrids = true);
+	static AIMODULE_API int32 AddLocalNavigationGridForPoint(UObject* WorldContextObject, const FVector& Location, const int32 Radius2D = 5, const float Height = 100.0f, bool bRebuildGrids = true);
 
 	/** creates single grid data for set of points */
 	UFUNCTION(BlueprintCallable, Category = "AI|Navigation", meta = (WorldContext = "WorldContextObject"))
-	static int32 AddLocalNavigationGridForPoints(UObject* WorldContextObject, const TArray<FVector>& Locations, const int32 Radius2D = 5, const float Height = 100.0f, bool bRebuildGrids = true);
+	static AIMODULE_API int32 AddLocalNavigationGridForPoints(UObject* WorldContextObject, const TArray<FVector>& Locations, const int32 Radius2D = 5, const float Height = 100.0f, bool bRebuildGrids = true);
 
 	UFUNCTION(BlueprintCallable, Category = "AI|Navigation", meta = (WorldContext = "WorldContextObject"))
-	static int32 AddLocalNavigationGridForBox(UObject* WorldContextObject, const FVector& Location, FVector Extent = FVector(1,1,1), FRotator Rotation = FRotator::ZeroRotator, const int32 Radius2D = 5, const float Height = 100.0f, bool bRebuildGrids = true);
+	static AIMODULE_API int32 AddLocalNavigationGridForBox(UObject* WorldContextObject, const FVector& Location, FVector Extent = FVector(1,1,1), FRotator Rotation = FRotator::ZeroRotator, const int32 Radius2D = 5, const float Height = 100.0f, bool bRebuildGrids = true);
 
 	UFUNCTION(BlueprintCallable, Category = "AI|Navigation", meta = (WorldContext = "WorldContextObject"))
-	static int32 AddLocalNavigationGridForCapsule(UObject* WorldContextObject, const FVector& Location, float CapsuleRadius, float CapsuleHalfHeight, const int32 Radius2D = 5, const float Height = 100.0f, bool bRebuildGrids = true);
+	static AIMODULE_API int32 AddLocalNavigationGridForCapsule(UObject* WorldContextObject, const FVector& Location, float CapsuleRadius, float CapsuleHalfHeight, const int32 Radius2D = 5, const float Height = 100.0f, bool bRebuildGrids = true);
 
 	UFUNCTION(BlueprintCallable, Category = "AI|Navigation", meta = (WorldContext = "WorldContextObject"))
-	static void RemoveLocalNavigationGrid(UObject* WorldContextObject, int32 GridId, bool bRebuildGrids = true);
+	static AIMODULE_API void RemoveLocalNavigationGrid(UObject* WorldContextObject, int32 GridId, bool bRebuildGrids = true);
 
 	UFUNCTION(BlueprintCallable, Category = "AI|Navigation", meta = (WorldContext = "WorldContextObject"))
-	static bool FindLocalNavigationGridPath(UObject* WorldContextObject, const FVector& Start, const FVector& End, TArray<FVector>& PathPoints);
+	static AIMODULE_API bool FindLocalNavigationGridPath(UObject* WorldContextObject, const FVector& Start, const FVector& End, TArray<FVector>& PathPoints);
 
-	static UNavLocalGridManager* GetCurrent(UWorld* World);
-	static UNavLocalGridManager* GetCurrent(const UObject* WorldContextObject);
+	static AIMODULE_API UNavLocalGridManager* GetCurrent(UWorld* World);
+	static AIMODULE_API UNavLocalGridManager* GetCurrent(const UObject* WorldContextObject);
 
 #if WITH_ENGINE
-	virtual UWorld* GetWorld() const override;
+	AIMODULE_API virtual UWorld* GetWorld() const override;
 #endif
 
 protected:
@@ -115,7 +115,7 @@ protected:
 	TArray<FCombinedNavGridData> CombinedGrids;
 	TArray<FNavLocalGridData> SourceGrids;
 	
-	static float GridCellSize;
+	static AIMODULE_API float GridCellSize;
 
 	int32 VersionNum;
 	int32 NextGridId;
@@ -123,8 +123,8 @@ protected:
 	uint32 bNeedsRebuilds : 1;
 
 	/** projects combined grids to navigation data */
-	virtual void ProjectGrids(const TArray<int32>& GridIndices);
+	AIMODULE_API virtual void ProjectGrids(const TArray<int32>& GridIndices);
 
 	/** ensures limit of source grids, removing oldest entries (LastAccessTime) */
-	bool UpdateSourceGrids();
+	AIMODULE_API bool UpdateSourceGrids();
 };
