[Engine]

diff --git a/Engine/Source/Runtime/AIModule/Classes/Navigation/PathFollowingComponent.h b/Engine/Source/Runtime/AIModule/Classes/Navigation/PathFollowingComponent.h
--- a/Engine/Source/Runtime/AIModule/Classes/Navigation/PathFollowingComponent.h
+++ b/Engine/Source/Runtime/AIModule/Classes/Navigation/PathFollowingComponent.h
@@ -118,14 +118,14 @@ namespace FPathFollowingResultFlags
 	FString ToString(uint16 Value);
 }
 
-struct AIMODULE_API FPathFollowingResult
+struct FPathFollowingResult
 {
 	FPathFollowingResultFlags::Type Flags;
 	TEnumAsByte<EPathFollowingResult::Type> Code;
 
 	FPathFollowingResult() : Flags(0), Code(EPathFollowingResult::Invalid)  {}
-	FPathFollowingResult(FPathFollowingResultFlags::Type InFlags);
-	FPathFollowingResult(EPathFollowingResult::Type ResultCode, FPathFollowingResultFlags::Type ExtraFlags);
+	AIMODULE_API FPathFollowingResult(FPathFollowingResultFlags::Type InFlags);
+	AIMODULE_API FPathFollowingResult(EPathFollowingResult::Type ResultCode, FPathFollowingResultFlags::Type ExtraFlags);
 
 	bool HasFlag(FPathFollowingResultFlags::Type Flag) const { return (Flags & Flag) != 0; }
 
@@ -133,7 +133,7 @@ struct AIMODULE_API FPathFollowingResult
 	bool IsFailure() const { return !HasFlag(FPathFollowingResultFlags::Success); }
 	bool IsInterrupted() const { return HasFlag(FPathFollowingResultFlags::UserAbort | FPathFollowingResultFlags::NewRequest); }
 	
-	FString ToString() const;
+	AIMODULE_API FString ToString() const;
 };
 
 // DEPRECATED, will be removed with GetPathActionType function
@@ -161,7 +161,7 @@ namespace EPathFollowingRequestResult
 	};
 }
 
-struct AIMODULE_API FPathFollowingRequestResult
+struct FPathFollowingRequestResult
 {
 	FAIRequestID MoveId;
 	TEnumAsByte<EPathFollowingRequestResult::Type> Code;
@@ -212,8 +212,8 @@ enum class EPathFollowingReachMode : uint8
 	OverlapAgentAndGoal,
 };
 
-UCLASS(config=Engine)
-class AIMODULE_API UPathFollowingComponent : public UActorComponent, public IAIResourceInterface, public IPathFollowingAgentInterface
+UCLASS(config=Engine, MinimalAPI)
+class UPathFollowingComponent : public UActorComponent, public IAIResourceInterface, public IPathFollowingAgentInterface
 {
 	GENERATED_UCLASS_BODY()
 
@@ -229,63 +229,63 @@ class AIMODULE_API UPathFollowingComponent : public UActorComponent, public IAIR
 	FMoveComplete OnRequestFinished;
 
 	//~ Begin UActorComponent Interface
-	virtual void OnRegister() override;
-	virtual void TickComponent(float DeltaTime, enum ELevelTick TickType, FActorComponentTickFunction *ThisTickFunction) override;
+	AIMODULE_API virtual void OnRegister() override;
+	AIMODULE_API virtual void TickComponent(float DeltaTime, enum ELevelTick TickType, FActorComponentTickFunction *ThisTickFunction) override;
 	//~ End UActorComponent Interface
 
 	/** initialize component to use */
-	virtual void Initialize();
+	AIMODULE_API virtual void Initialize();
 
 	/** cleanup component before destroying */
-	virtual void Cleanup();
+	AIMODULE_API virtual void Cleanup();
 
 	/** updates cached pointers to relevant owner's components */
-	virtual void UpdateCachedComponents();
+	AIMODULE_API virtual void UpdateCachedComponents();
 
 	/** start movement along path
 	  * @return MoveId of requested move
 	  */
-	virtual FAIRequestID RequestMove(const FAIMoveRequest& RequestData, FNavPathSharedPtr InPath);
+	AIMODULE_API virtual FAIRequestID RequestMove(const FAIMoveRequest& RequestData, FNavPathSharedPtr InPath);
 
 	/** aborts following path */
-	virtual void AbortMove(const UObject& Instigator, FPathFollowingResultFlags::Type AbortFlags, FAIRequestID RequestID = FAIRequestID::CurrentRequest, EPathFollowingVelocityMode VelocityMode = EPathFollowingVelocityMode::Reset);
+	AIMODULE_API virtual void AbortMove(const UObject& Instigator, FPathFollowingResultFlags::Type AbortFlags, FAIRequestID RequestID = FAIRequestID::CurrentRequest, EPathFollowingVelocityMode VelocityMode = EPathFollowingVelocityMode::Reset);
 
 	/** create new request and finish it immediately (e.g. already at goal)
 	 *  @return MoveId of requested (and already finished) move
 	 */
-	FAIRequestID RequestMoveWithImmediateFinish(EPathFollowingResult::Type Result, EPathFollowingVelocityMode VelocityMode = EPathFollowingVelocityMode::Reset);
+	AIMODULE_API FAIRequestID RequestMoveWithImmediateFinish(EPathFollowingResult::Type Result, EPathFollowingVelocityMode VelocityMode = EPathFollowingVelocityMode::Reset);
 
 	/** pause path following
 	*  @param RequestID - request to pause, FAIRequestID::CurrentRequest means pause current request, regardless of its ID */
-	virtual void PauseMove(FAIRequestID RequestID = FAIRequestID::CurrentRequest, EPathFollowingVelocityMode VelocityMode = EPathFollowingVelocityMode::Reset);
+	AIMODULE_API virtual void PauseMove(FAIRequestID RequestID = FAIRequestID::CurrentRequest, EPathFollowingVelocityMode VelocityMode = EPathFollowingVelocityMode::Reset);
 
 	/** resume path following
 	*  @param RequestID - request to resume, FAIRequestID::CurrentRequest means restor current request, regardless of its ID */
-	virtual void ResumeMove(FAIRequestID RequestID = FAIRequestID::CurrentRequest);
+	AIMODULE_API virtual void ResumeMove(FAIRequestID RequestID = FAIRequestID::CurrentRequest);
 
 	/** notify about finished movement */
-	virtual void OnPathFinished(const FPathFollowingResult& Result);
+	AIMODULE_API virtual void OnPathFinished(const FPathFollowingResult& Result);
 
 	FORCEINLINE void OnPathFinished(EPathFollowingResult::Type ResultCode, uint16 ExtraResultFlags) { OnPathFinished(FPathFollowingResult(ResultCode, ExtraResultFlags)); }
 
 	/** notify about finishing move along current path segment */
-	virtual void OnSegmentFinished();
+	AIMODULE_API virtual void OnSegmentFinished();
 
 	/** notify about changing current path: new pointer or update from path event */
-	virtual void OnPathUpdated();
+	AIMODULE_API virtual void OnPathUpdated();
 
 	/** set associated movement component */
-	virtual void SetMovementComponent(UNavMovementComponent* MoveComp);
+	AIMODULE_API virtual void SetMovementComponent(UNavMovementComponent* MoveComp);
 
 	/** get current focal point of movement */
-	virtual FVector GetMoveFocus(bool bAllowStrafe) const;
+	AIMODULE_API virtual FVector GetMoveFocus(bool bAllowStrafe) const;
 
 	/** simple test for stationary agent (used as early finish condition), check if reached given point
 	 *  @param TestPoint - point to test
 	 *  @param AcceptanceRadius - allowed 2D distance
 	 *  @param ReachMode - modifiers for AcceptanceRadius
 	 */
-	bool HasReached(const FVector& TestPoint, EPathFollowingReachMode ReachMode, float AcceptanceRadius = UPathFollowingComponent::DefaultAcceptanceRadius) const;
+	AIMODULE_API bool HasReached(const FVector& TestPoint, EPathFollowingReachMode ReachMode, float AcceptanceRadius = UPathFollowingComponent::DefaultAcceptanceRadius) const;
 
 	/** simple test for stationary agent (used as early finish condition), check if reached given goal
 	 *  @param TestGoal - actor to test
@@ -293,19 +293,19 @@ class AIMODULE_API UPathFollowingComponent : public UActorComponent, public IAIR
 	 *  @param ReachMode - modifiers for AcceptanceRadius
 	 *  @param bUseNavAgentGoalLocation - true: if the goal is a nav agent, we will use their nav agent location rather than their actual location
 	 */
-	bool HasReached(const AActor& TestGoal, EPathFollowingReachMode ReachMode, float AcceptanceRadius = UPathFollowingComponent::DefaultAcceptanceRadius, bool bUseNavAgentGoalLocation = true) const;
+	AIMODULE_API bool HasReached(const AActor& TestGoal, EPathFollowingReachMode ReachMode, float AcceptanceRadius = UPathFollowingComponent::DefaultAcceptanceRadius, bool bUseNavAgentGoalLocation = true) const;
 
 	/** simple test for stationary agent (used as early finish condition), check if reached target specified in move request */
-	bool HasReached(const FAIMoveRequest& MoveRequest) const;
+	AIMODULE_API bool HasReached(const FAIMoveRequest& MoveRequest) const;
 
 	/** update state of block detection */
-	void SetBlockDetectionState(bool bEnable);
+	AIMODULE_API void SetBlockDetectionState(bool bEnable);
 
 	/** @returns state of block detection */
 	bool IsBlockDetectionActive() const { return bUseBlockDetection; }
 
 	/** set block detection params */
-	void SetBlockDetection(float DistanceThreshold, float Interval, int32 NumSamples);
+	AIMODULE_API void SetBlockDetection(float DistanceThreshold, float Interval, int32 NumSamples);
 
 	/** Returns true if pathfollowing is doing deceleration at the end of the path. */
 	bool IsDecelerating() const { return bIsDecelerating; };
@@ -317,22 +317,22 @@ class AIMODULE_API UPathFollowingComponent : public UActorComponent, public IAIR
 	FORCEINLINE void SetStopMovementOnFinish(bool bEnable) { bStopMovementOnFinish = bEnable; }
 
 	/** set threshold for precise reach tests in intermediate goals (minimal test radius)  */
-	void SetPreciseReachThreshold(float AgentRadiusMultiplier, float AgentHalfHeightMultiplier);
+	AIMODULE_API void SetPreciseReachThreshold(float AgentRadiusMultiplier, float AgentHalfHeightMultiplier);
 
 	/** set status of last requested move, works only in Idle state */
-	void SetLastMoveAtGoal(bool bFinishedAtGoal);
+	AIMODULE_API void SetLastMoveAtGoal(bool bFinishedAtGoal);
 
 	/** @returns estimated cost of unprocessed path segments
 	 *	@NOTE 0 means, that component is following final path segment or doesn't move */
-	FVector::FReal GetRemainingPathCost() const;
+	AIMODULE_API FVector::FReal GetRemainingPathCost() const;
 	
 	/** Returns current location on navigation data */
-	FNavLocation GetCurrentNavLocation() const;
+	AIMODULE_API FNavLocation GetCurrentNavLocation() const;
 
 	FORCEINLINE EPathFollowingStatus::Type GetStatus() const { return Status; }
 	FORCEINLINE float GetAcceptanceRadius() const { return AcceptanceRadius; }
 	FORCEINLINE float GetDefaultAcceptanceRadius() const { return MyDefaultAcceptanceRadius; }
-	void SetAcceptanceRadius(const float InAcceptanceRadius);
+	AIMODULE_API void SetAcceptanceRadius(const float InAcceptanceRadius);
 	FORCEINLINE AActor* GetMoveGoal() const { return DestinationActor.Get(); }
 	FORCEINLINE bool HasPartialPath() const { return Path.IsValid() && Path->IsPartial(); }
 	FORCEINLINE bool DidMoveReachGoal() const { return bLastMoveReachedGoal && (Status == EPathFollowingStatus::Idle); }
@@ -345,8 +345,8 @@ class AIMODULE_API UPathFollowingComponent : public UActorComponent, public IAIR
 	FORCEINLINE FBasedPosition GetCurrentTargetLocationBased() const { return CurrentDestination; }
 	FORCEINLINE FVector GetMoveGoalLocationOffset() const { return MoveOffset; }
 	bool HasStartedNavLinkMove() const { return bWalkingNavLinkStart; }
-	bool IsCurrentSegmentNavigationLink() const;
-	FVector GetCurrentDirection() const;
+	AIMODULE_API bool IsCurrentSegmentNavigationLink() const;
+	AIMODULE_API FVector GetCurrentDirection() const;
 	/** note that CurrentMoveInput is only valid if MovementComp->UseAccelerationForPathFollowing() == true */
 	FVector GetCurrentMoveInput() const { return CurrentMoveInput; }
 
@@ -355,66 +355,66 @@ class AIMODULE_API UPathFollowingComponent : public UActorComponent, public IAIR
 
 	FORCEINLINE const FNavPathSharedPtr GetPath() const { return Path; }
 	FORCEINLINE bool HasValidPath() const { return Path.IsValid() && Path->IsValid(); }
-	bool HasDirectPath() const;
+	AIMODULE_API bool HasDirectPath() const;
 
 	/** readable name of current status */
-	FString GetStatusDesc() const;
+	AIMODULE_API FString GetStatusDesc() const;
 	/** readable name of result enum */
-	FString GetResultDesc(EPathFollowingResult::Type Result) const;
+	AIMODULE_API FString GetResultDesc(EPathFollowingResult::Type Result) const;
 
-	void SetDestinationActor(const AActor* InDestinationActor);
+	AIMODULE_API void SetDestinationActor(const AActor* InDestinationActor);
 
 	/** returns index of the currently followed element of path. Depending on the actual 
 	 *	path it may represent different things, like a path point or navigation corridor index */
 	virtual int32 GetCurrentPathElement() const { return MoveSegmentEndIndex; }
 
-	virtual void GetDebugStringTokens(TArray<FString>& Tokens, TArray<EPathFollowingDebugTokens::Type>& Flags) const;
-	virtual FString GetDebugString() const;
+	AIMODULE_API virtual void GetDebugStringTokens(TArray<FString>& Tokens, TArray<EPathFollowingDebugTokens::Type>& Flags) const;
+	AIMODULE_API virtual FString GetDebugString() const;
 
-	virtual void DisplayDebug(UCanvas* Canvas, const FDebugDisplayInfo& DebugDisplay, float& YL, float& YPos) const;
+	AIMODULE_API virtual void DisplayDebug(UCanvas* Canvas, const FDebugDisplayInfo& DebugDisplay, float& YL, float& YPos) const;
 #if ENABLE_VISUAL_LOG
-	virtual void DescribeSelfToVisLog(struct FVisualLogEntry* Snapshot) const;
+	AIMODULE_API virtual void DescribeSelfToVisLog(struct FVisualLogEntry* Snapshot) const;
 #endif // ENABLE_VISUAL_LOG
 
 	/** called when moving agent collides with another actor */
 	UFUNCTION()
-	virtual void OnActorBump(AActor* SelfActor, AActor* OtherActor, FVector NormalImpulse, const FHitResult& Hit);
+	AIMODULE_API virtual void OnActorBump(AActor* SelfActor, AActor* OtherActor, FVector NormalImpulse, const FHitResult& Hit);
 
 	// IPathFollowingAgentInterface begin
-	virtual void OnUnableToMove(const UObject& Instigator) override;		
+	AIMODULE_API virtual void OnUnableToMove(const UObject& Instigator) override;		
 	//virtual void OnMoveBlockedBy(const FHitResult& BlockingImpact) {}
-	virtual void OnStartedFalling() override;
+	AIMODULE_API virtual void OnStartedFalling() override;
 	virtual void OnLanded() override {}
 	// IPathFollowingAgentInterface end
 
 	/** Check if path following can be activated */
-	virtual bool IsPathFollowingAllowed() const;
+	AIMODULE_API virtual bool IsPathFollowingAllowed() const;
 
 	/** call when moving agent finishes using custom nav link, returns control back to path following */
-	virtual void FinishUsingCustomLink(INavLinkCustomInterface* CustomNavLink);
+	AIMODULE_API virtual void FinishUsingCustomLink(INavLinkCustomInterface* CustomNavLink);
 
 	/** called when owner is preparing new pathfinding request */
 	virtual void OnPathfindingQuery(FPathFindingQuery& Query) {}
 
 	// IAIResourceInterface begin
-	virtual void LockResource(EAIRequestPriority::Type LockSource) override;
-	virtual void ClearResourceLock(EAIRequestPriority::Type LockSource) override;
-	virtual void ForceUnlockResource() override;
-	virtual bool IsResourceLocked() const override;
+	AIMODULE_API virtual void LockResource(EAIRequestPriority::Type LockSource) override;
+	AIMODULE_API virtual void ClearResourceLock(EAIRequestPriority::Type LockSource) override;
+	AIMODULE_API virtual void ForceUnlockResource() override;
+	AIMODULE_API virtual bool IsResourceLocked() const override;
 	// IAIResourceInterface end
 
 	/** path observer */
-	void OnPathEvent(FNavigationPath* InPath, ENavPathEvent::Type Event);
+	AIMODULE_API void OnPathEvent(FNavigationPath* InPath, ENavPathEvent::Type Event);
 
 	/** helper function for sending a path for visual log */
-	static void LogPathHelper(const AActor* LogOwner, FNavPathSharedPtr InLogPath, const AActor* LogGoalActor);
-	static void LogPathHelper(const AActor* LogOwner, FNavigationPath* InLogPath, const AActor* LogGoalActor);
+	static AIMODULE_API void LogPathHelper(const AActor* LogOwner, FNavPathSharedPtr InLogPath, const AActor* LogGoalActor);
+	static AIMODULE_API void LogPathHelper(const AActor* LogOwner, FNavigationPath* InLogPath, const AActor* LogGoalActor);
 
 	UFUNCTION(BlueprintCallable, Category="AI|Components|PathFollowing", meta = (DeprecatedFunction, DeprecationMessage = "This function is now deprecated, please use AIController.GetMoveStatus instead"))
-	EPathFollowingAction::Type GetPathActionType() const;
+	AIMODULE_API EPathFollowingAction::Type GetPathActionType() const;
 
 	UFUNCTION(BlueprintCallable, Category="AI|Components|PathFollowing", meta = (DeprecatedFunction, DeprecationMessage = "This function is now deprecated, please use AIController.GetImmediateMoveDestination instead"))
-	FVector GetPathDestination() const;
+	AIMODULE_API FVector GetPathDestination() const;
 
 #if WITH_EDITORONLY_DATA
 	// This delegate is now deprecated, please use OnRequestFinished instead
@@ -570,47 +570,47 @@ protected:
 	int32 DecelerationSegmentIndex;
 
 	/** reset path following data */
-	virtual void Reset();
+	AIMODULE_API virtual void Reset();
 
 	/** Called if owning Controller possesses new pawn or ends up pawn-less. 
 	 *	Doesn't get called if owner is not an AContoller */
-	virtual void OnNewPawn(APawn* NewPawn);
+	AIMODULE_API virtual void OnNewPawn(APawn* NewPawn);
 
 	/** should verify if agent if still on path ater movement has been resumed? */
-	virtual bool ShouldCheckPathOnResume() const;
+	AIMODULE_API virtual bool ShouldCheckPathOnResume() const;
 
 	/** sets variables related to current move segment */
-	virtual void SetMoveSegment(int32 SegmentStartIndex);
+	AIMODULE_API virtual void SetMoveSegment(int32 SegmentStartIndex);
 	
 	/** follow current path segment */
-	virtual void FollowPathSegment(float DeltaTime);
+	AIMODULE_API virtual void FollowPathSegment(float DeltaTime);
 
 	/** check state of path following, update move segment if needed */
-	virtual void UpdatePathSegment();
+	AIMODULE_API virtual void UpdatePathSegment();
 
 	/** next path segment if custom nav link, try passing control to it */
-	virtual void StartUsingCustomLink(INavLinkCustomInterface* CustomNavLink, const FVector& DestPoint);
+	AIMODULE_API virtual void StartUsingCustomLink(INavLinkCustomInterface* CustomNavLink, const FVector& DestPoint);
 
 	/** update blocked movement detection, @returns true if new sample was added */
-	virtual bool UpdateBlockDetection();
+	AIMODULE_API virtual bool UpdateBlockDetection();
 
 	/** updates braking distance and deceleration segment */
-	virtual void UpdateDecelerationData();
+	AIMODULE_API virtual void UpdateDecelerationData();
 
 	/** check if move is completed */
-	virtual bool HasReachedDestination(const FVector& CurrentLocation) const;
+	AIMODULE_API virtual bool HasReachedDestination(const FVector& CurrentLocation) const;
 
 	/** check if segment is completed */
-	virtual bool HasReachedCurrentTarget(const FVector& CurrentLocation) const;
+	AIMODULE_API virtual bool HasReachedCurrentTarget(const FVector& CurrentLocation) const;
 
 	/** check if moving agent has reached goal defined by cylinder */
-	bool HasReachedInternal(const FVector& GoalLocation, float GoalRadius, float GoalHalfHeight, const FVector& AgentLocation, float RadiusThreshold, float AgentRadiusMultiplier) const;
+	AIMODULE_API bool HasReachedInternal(const FVector& GoalLocation, float GoalRadius, float GoalHalfHeight, const FVector& AgentLocation, float RadiusThreshold, float AgentRadiusMultiplier) const;
 
 	/** check if agent is on path */
-	virtual bool IsOnPath() const;
+	AIMODULE_API virtual bool IsOnPath() const;
 
 	/** check if movement is blocked */
-	bool IsBlocked() const;
+	AIMODULE_API bool IsBlocked() const;
 
 	/** switch to next segment on path */
 	FORCEINLINE void SetNextMoveSegment() { SetMoveSegment(GetNextPathIndex()); }
@@ -624,56 +624,56 @@ protected:
 	*	if so determines index of next path point
 	*	@return what PathFollowingComponent thinks should be next path point. INDEX_NONE if given path is invalid
 	*	@note this function does not set MoveSegmentEndIndex */
-	virtual int32 DetermineStartingPathPoint(const FNavigationPath* ConsideredPath) const;
+	AIMODULE_API virtual int32 DetermineStartingPathPoint(const FNavigationPath* ConsideredPath) const;
 
 	/** @return index of path point, that should be target of current move segment */
-	virtual int32 DetermineCurrentTargetPathPoint(int32 StartIndex);
+	AIMODULE_API virtual int32 DetermineCurrentTargetPathPoint(int32 StartIndex);
 
 	/** check if movement component is valid or tries to grab one from owner 
 	 *	@param bForce results in looking for owner's movement component even if pointer to one is already cached */
-	virtual bool UpdateMovementComponent(bool bForce = false);
+	AIMODULE_API virtual bool UpdateMovementComponent(bool bForce = false);
 
 	/** called after receiving update event from current path
 	 *  @return false if path was not accepted and move request needs to be aborted */
-	virtual bool HandlePathUpdateEvent();
+	AIMODULE_API virtual bool HandlePathUpdateEvent();
 
 	/** called from timer if component spends too much time in Waiting state */
-	virtual void OnWaitingPathTimeout();
+	AIMODULE_API virtual void OnWaitingPathTimeout();
 
 	/** clears Block Detection stored data effectively resetting the mechanism */
-	void ResetBlockDetectionData();
+	AIMODULE_API void ResetBlockDetectionData();
 
 	/** force creating new location sample for block detection */
-	void ForceBlockDetectionUpdate();
+	AIMODULE_API void ForceBlockDetectionUpdate();
 
 	/** set move focus in AI owner */
-	virtual void UpdateMoveFocus();
+	AIMODULE_API virtual void UpdateMoveFocus();
 
 	/** defines if the agent should reset his velocity when the path is finished*/
-	virtual bool ShouldStopMovementOnPathFinished() const; 
+	AIMODULE_API virtual bool ShouldStopMovementOnPathFinished() const; 
 
 	/** For given path finds a path node at which
 	 *	PathfollowingComponent should start doing 
 	 *	precise is-goal-in-acceptance-radius  tests */
-	int32 FindPreciseAcceptanceRadiusTestsStartNodeIndex(const FNavigationPath& PathInstance, const FVector& GoalLocation) const;
+	AIMODULE_API int32 FindPreciseAcceptanceRadiusTestsStartNodeIndex(const FNavigationPath& PathInstance, const FVector& GoalLocation) const;
 
 	/** Based on Path's properties, original move goal location and requested AcceptanceRadius
 	 *	this function calculates actual acceptance radius to apply when testing if the agent
 	 *	has successfully reached requested goal's vicinity */
-	float GetFinalAcceptanceRadius(const FNavigationPath& PathInstance, const FVector OriginalGoalLocation, const FVector* PathEndOverride = nullptr) const;
+	AIMODULE_API float GetFinalAcceptanceRadius(const FNavigationPath& PathInstance, const FVector OriginalGoalLocation, const FVector* PathEndOverride = nullptr) const;
 
 	/** debug point reach test values */
-	void DebugReachTest(float& CurrentDot, float& CurrentDistance, float& CurrentHeight, uint8& bDotFailed, uint8& bDistanceFailed, uint8& bHeightFailed) const;
+	AIMODULE_API void DebugReachTest(float& CurrentDot, float& CurrentDistance, float& CurrentHeight, uint8& bDotFailed, uint8& bDistanceFailed, uint8& bHeightFailed) const;
 
 	/** called when NavigationSystem finishes initial navigation data registration.
 	 *	This is usually required by AI agents hand-placed on levels to find MyNavData */
-	virtual void OnNavigationInitDone();
+	AIMODULE_API virtual void OnNavigationInitDone();
 
 	/** called when NavigationSystem registers new navigation data type while this component
 	 *	instance has empty MyNavData. This is usually the case for AI agents hand-placed
 	 *	on levels. */
 	UFUNCTION()
-	void OnNavDataRegistered(ANavigationData* NavData);
+	AIMODULE_API void OnNavDataRegistered(ANavigationData* NavData);
 
 
 	/** used to keep track of which subsystem requested this AI resource be locked */
@@ -686,7 +686,7 @@ private:
 
 	/** used for debugging purposes to be able to identify which logged information
 	 *	results from which request, if there was multiple ones during one frame */
-	static uint32 NextRequestId;
+	static AIMODULE_API uint32 NextRequestId;
 	FAIRequestID CurrentRequestId;
 
 	/** Current location on navigation data.  Lazy-updated, so read this via GetCurrentNavLocation(). 
@@ -698,7 +698,7 @@ public:
 	/** special float constant to symbolize "use default value". This does not contain 
 	 *	value to be used, it's used to detect the fact that it's requested, and 
 	 *	appropriate value from querier/doer will be pulled */
-	static const float DefaultAcceptanceRadius;
+	static AIMODULE_API const float DefaultAcceptanceRadius;
 
 #if !UE_BUILD_SHIPPING
 	uint8 DEBUG_bMovingDirectlyToGoal : 1;
diff --git a/Engine/Source/Runtime/AIModule/Classes/Navigation/PathFollowingManager.h b/Engine/Source/Runtime/AIModule/Classes/Navigation/PathFollowingManager.h
--- a/Engine/Source/Runtime/AIModule/Classes/Navigation/PathFollowingManager.h
+++ b/Engine/Source/Runtime/AIModule/Classes/Navigation/PathFollowingManager.h
@@ -9,13 +9,13 @@
 class AController;
 
 
-UCLASS()
-class AIMODULE_API UPathFollowingManager : public UObject, public IPathFollowingManagerInterface
+UCLASS(MinimalAPI)
+class UPathFollowingManager : public UObject, public IPathFollowingManagerInterface
 {
 	GENERATED_BODY()
 public:
-	UPathFollowingManager(const FObjectInitializer& ObjectInitializer = FObjectInitializer::Get());
+	AIMODULE_API UPathFollowingManager(const FObjectInitializer& ObjectInitializer = FObjectInitializer::Get());
 
-	static void StopMovement(const AController& Controller);
-	static bool IsFollowingAPath(const AController& Controller);
+	static AIMODULE_API void StopMovement(const AController& Controller);
+	static AIMODULE_API bool IsFollowingAPath(const AController& Controller);
 };
diff --git a/Engine/Source/Runtime/AIModule/Classes/Navigation/RecastGraphAStar.h b/Engine/Source/Runtime/AIModule/Classes/Navigation/RecastGraphAStar.h
--- a/Engine/Source/Runtime/AIModule/Classes/Navigation/RecastGraphAStar.h
+++ b/Engine/Source/Runtime/AIModule/Classes/Navigation/RecastGraphAStar.h
@@ -32,7 +32,7 @@ struct dtMeshTile;
 struct dtLink;
 class dtNavMesh;
 
-struct AIMODULE_API FRecastNeighbour
+struct FRecastNeighbour
 {
 	friend FRecastGraphWrapper;
 	friend FRecastGraphAStarFilter;
@@ -52,7 +52,7 @@ public:
 	unsigned char Side;
 };
 
-struct AIMODULE_API FRecastAStarResult : public dtQueryResult
+struct FRecastAStarResult : public dtQueryResult
 {
 	void Reset(const int32 PathLength)
 	{
@@ -63,11 +63,11 @@ struct AIMODULE_API FRecastAStarResult : public dtQueryResult
 		data.resize(PathLength);
 	}
 
-	dtPolyRef SetPathInfo(const int32 Index, const FRecastAStarSearchNode& SearchNode);
+	AIMODULE_API dtPolyRef SetPathInfo(const int32 Index, const FRecastAStarSearchNode& SearchNode);
 };
 
 USTRUCT()
-struct AIMODULE_API FRecastGraphWrapper
+struct FRecastGraphWrapper
 {
 	GENERATED_BODY()
 
@@ -75,10 +75,10 @@ public:
 	FRecastGraphWrapper() {}
 
 	/** Initialization of the wrapper from the RecastNavMesh pointer */
-	void Initialize(const ARecastNavMesh* InRecastNavMeshActor);
+	AIMODULE_API void Initialize(const ARecastNavMesh* InRecastNavMeshActor);
 
 	/** Implementation that converts EGraphAStarResult into a dtStatus */
-	dtStatus ConvertToRecastStatus(const FRecastAStar& Algo, const FRecastGraphAStarFilter& Filter, const EGraphAStarResult AStarResult) const;
+	AIMODULE_API dtStatus ConvertToRecastStatus(const FRecastAStar& Algo, const FRecastGraphAStarFilter& Filter, const EGraphAStarResult AStarResult) const;
 
 	//////////////////////////////////////////////////////////////////////////
 	// FGraphAStar: TGraph
@@ -88,7 +88,7 @@ public:
 	{
 		return NodeRef != INVALID_NAVNODEREF;
 	}
-	FRecastNeighbour GetNeighbour(const FRecastAStarSearchNode& Node, const int32 NeighbourIndex) const;
+	AIMODULE_API FRecastNeighbour GetNeighbour(const FRecastAStarSearchNode& Node, const int32 NeighbourIndex) const;
 	//////////////////////////////////////////////////////////////////////////
 
 	FORCEINLINE const dtNavMeshQuery& GetRecastQuery() const { return RecastQuery; }
@@ -100,7 +100,7 @@ protected:
 
 	FORCEINLINE const ARecastNavMesh* GetRecastNavMeshActor() const { checkSlow(RecastNavMeshActor);  return RecastNavMeshActor; }
 	FORCEINLINE const dtNavMesh* GetDetourNavMesh() const { checkSlow(DetourNavMesh); return DetourNavMesh; }
-	void BindFilter(FRecastGraphAStarFilter& AStarFilter);
+	AIMODULE_API void BindFilter(FRecastGraphAStarFilter& AStarFilter);
 
 private:
 	UPROPERTY(Transient)
@@ -113,7 +113,7 @@ private:
 	mutable unsigned int CachedNextLink = DT_NULL_LINK;
 };
 
-struct AIMODULE_API FRecastAStarSearchNode : public FGraphAStarDefaultNode<FRecastGraphWrapper>
+struct FRecastAStarSearchNode : public FGraphAStarDefaultNode<FRecastGraphWrapper>
 {
 	typedef FGraphAStarDefaultNode<FRecastGraphWrapper> Super;
 
@@ -157,7 +157,7 @@ struct AIMODULE_API FRecastAStarSearchNode : public FGraphAStarDefaultNode<FReca
 	}
 };
 
-struct AIMODULE_API FRecastAStar : public FGraphAStar<FRecastGraphWrapper, FRecastGraphPolicy, FRecastAStarSearchNode>
+struct FRecastAStar : public FGraphAStar<FRecastGraphWrapper, FRecastGraphPolicy, FRecastAStarSearchNode>
 {
 	typedef FGraphAStar<FRecastGraphWrapper, FRecastGraphPolicy, FRecastAStarSearchNode> Super;
 	FRecastAStar(const FRecastGraphWrapper& Graph)
@@ -165,9 +165,9 @@ struct AIMODULE_API FRecastAStar : public FGraphAStar<FRecastGraphWrapper, FReca
 	{}
 };
 
-struct AIMODULE_API FRecastGraphAStarFilter
+struct FRecastGraphAStarFilter
 {
-	FRecastGraphAStarFilter(FRecastGraphWrapper& InRecastGraphWrapper, const FRecastQueryFilter& InFilter, uint32 InMaxSearchNodes, const FVector::FReal InCostLimit, const UObject* Owner);
+	AIMODULE_API FRecastGraphAStarFilter(FRecastGraphWrapper& InRecastGraphWrapper, const FRecastQueryFilter& InFilter, uint32 InMaxSearchNodes, const FVector::FReal InCostLimit, const UObject* Owner);
 
 	FORCEINLINE bool WantsPartialSolution() const
 	{ 
@@ -244,4 +244,4 @@ private:
 	const FRecastGraphWrapper& RecastGraphWrapper;
 	uint32 MaxSearchNodes;
 	FVector::FReal CostLimit;
-};
\ No newline at end of file
+};
diff --git a/Engine/Source/Runtime/AIModule/Classes/Perception/AIPerceptionComponent.h b/Engine/Source/Runtime/AIModule/Classes/Perception/AIPerceptionComponent.h
--- a/Engine/Source/Runtime/AIModule/Classes/Perception/AIPerceptionComponent.h
+++ b/Engine/Source/Runtime/AIModule/Classes/Perception/AIPerceptionComponent.h
@@ -46,7 +46,7 @@ struct FActorPerceptionUpdateInfo
 
 DECLARE_DYNAMIC_MULTICAST_DELEGATE_OneParam(FActorPerceptionInfoUpdatedDelegate, const FActorPerceptionUpdateInfo&, UpdateInfo);
 
-struct AIMODULE_API FActorPerceptionInfo
+struct FActorPerceptionInfo
 {
 	TWeakObjectPtr<AActor> Target;
 
@@ -172,7 +172,7 @@ struct AIMODULE_API FActorPerceptionInfo
 	}
 	
 	/** takes all "newer" info from Other and absorbs it */
-	void Merge(const FActorPerceptionInfo& Other);
+	AIMODULE_API void Merge(const FActorPerceptionInfo& Other);
 };
 
 USTRUCT(BlueprintType, meta = (DisplayName = "Sensed Actor's Data"))
@@ -198,12 +198,12 @@ struct FActorPerceptionBlueprintInfo
  *	AIPerceptionComponent is used to register as stimuli listener in AIPerceptionSystem
  *	and gathers registered stimuli. UpdatePerception is called when component gets new stimuli (batched)
  */
-UCLASS(ClassGroup=AI, HideCategories=(Activation, Collision), meta=(BlueprintSpawnableComponent), config=Game)
-class AIMODULE_API UAIPerceptionComponent : public UActorComponent
+UCLASS(ClassGroup=AI, HideCategories=(Activation, Collision), meta=(BlueprintSpawnableComponent), config=Game, MinimalAPI)
+class UAIPerceptionComponent : public UActorComponent
 {
 	GENERATED_UCLASS_BODY()
 	
-	static const int32 InitialStimuliToProcessArraySize;
+	static AIMODULE_API const int32 InitialStimuliToProcessArraySize;
 
 	typedef TMap<TObjectKey<AActor>, FActorPerceptionInfo> TActorPerceptionContainer;
 	typedef TActorPerceptionContainer FActorPerceptionContainer;
@@ -257,56 +257,56 @@ private:
 
 public:
 
-	virtual void PostInitProperties() override;
-	virtual void BeginDestroy() override;
-	virtual void OnRegister() override;
-	virtual void OnUnregister() override;
+	AIMODULE_API virtual void PostInitProperties() override;
+	AIMODULE_API virtual void BeginDestroy() override;
+	AIMODULE_API virtual void OnRegister() override;
+	AIMODULE_API virtual void OnUnregister() override;
 
 	UFUNCTION()
-	void OnOwnerEndPlay(AActor* Actor, EEndPlayReason::Type EndPlayReason);
+	AIMODULE_API void OnOwnerEndPlay(AActor* Actor, EEndPlayReason::Type EndPlayReason);
 	
-	void GetLocationAndDirection(FVector& Location, FVector& Direction) const;
-	const AActor* GetBodyActor() const;
-	AActor* GetMutableBodyActor();
+	AIMODULE_API void GetLocationAndDirection(FVector& Location, FVector& Direction) const;
+	AIMODULE_API const AActor* GetBodyActor() const;
+	AIMODULE_API AActor* GetMutableBodyActor();
 
 	FORCEINLINE const FPerceptionChannelAllowList GetPerceptionFilter() const { return PerceptionFilter; }
 
-	FGenericTeamId GetTeamIdentifier() const;
+	AIMODULE_API FGenericTeamId GetTeamIdentifier() const;
 	FORCEINLINE FPerceptionListenerID GetListenerId() const { return PerceptionListenerId; }
 
-	FVector GetActorLocation(const AActor& Actor) const;
+	AIMODULE_API FVector GetActorLocation(const AActor& Actor) const;
 	FORCEINLINE const FActorPerceptionInfo* GetActorInfo(const AActor& Actor) const { return PerceptualData.Find(&Actor); }
 	FORCEINLINE FActorPerceptionContainer::TIterator GetPerceptualDataIterator() { return FActorPerceptionContainer::TIterator(PerceptualData); }
 	FORCEINLINE FActorPerceptionContainer::TConstIterator GetPerceptualDataConstIterator() const { return FActorPerceptionContainer::TConstIterator(PerceptualData); }
 
-	virtual void GetHostileActors(TArray<AActor*>& OutActors) const;
+	AIMODULE_API virtual void GetHostileActors(TArray<AActor*>& OutActors) const;
 	
-	void GetHostileActorsBySense(TSubclassOf<UAISense> SenseToFilterBy, TArray<AActor*>& OutActors) const;
+	AIMODULE_API void GetHostileActorsBySense(TSubclassOf<UAISense> SenseToFilterBy, TArray<AActor*>& OutActors) const;
 
 	/**	Retrieves all actors in PerceptualData matching the predicate.
 	 *	@return whether dead data (invalid actors) have been found while iterating over PerceptualData
 	 */
-	bool GetFilteredActors(TFunctionRef<bool(const FActorPerceptionInfo&)> Predicate, TArray<AActor*>& OutActors) const;
+	AIMODULE_API bool GetFilteredActors(TFunctionRef<bool(const FActorPerceptionInfo&)> Predicate, TArray<AActor*>& OutActors) const;
 
 	// @note Will stop on first age 0 stimulus
-	const FActorPerceptionInfo* GetFreshestTrace(const FAISenseID Sense) const;
+	AIMODULE_API const FActorPerceptionInfo* GetFreshestTrace(const FAISenseID Sense) const;
 	
-	void SetDominantSense(TSubclassOf<UAISense> InDominantSense);
+	AIMODULE_API void SetDominantSense(TSubclassOf<UAISense> InDominantSense);
 	FORCEINLINE FAISenseID GetDominantSenseID() const { return DominantSenseID; }
 	FORCEINLINE TSubclassOf<UAISense> GetDominantSense() const { return DominantSense; }
-	UAISenseConfig* GetSenseConfig(const FAISenseID& SenseID);
-	const UAISenseConfig* GetSenseConfig(const FAISenseID& SenseID) const;
-	void ConfigureSense(UAISenseConfig& SenseConfig);
+	AIMODULE_API UAISenseConfig* GetSenseConfig(const FAISenseID& SenseID);
+	AIMODULE_API const UAISenseConfig* GetSenseConfig(const FAISenseID& SenseID) const;
+	AIMODULE_API void ConfigureSense(UAISenseConfig& SenseConfig);
 
 	typedef TArray<UAISenseConfig*>::TConstIterator TAISenseConfigConstIterator;
-	TAISenseConfigConstIterator GetSensesConfigIterator() const;
+	AIMODULE_API TAISenseConfigConstIterator GetSensesConfigIterator() const;
 
 	/** Notifies AIPerceptionSystem to update properties for this "stimuli listener" */
 	UFUNCTION(BlueprintCallable, Category="AI|Perception")
-	void RequestStimuliListenerUpdate();
+	AIMODULE_API void RequestStimuliListenerUpdate();
 
 	/** Allows toggling senses on and off */
-	void UpdatePerceptionAllowList(const FAISenseID Channel, const bool bNewValue);
+	AIMODULE_API void UpdatePerceptionAllowList(const FAISenseID Channel, const bool bNewValue);
 
 	UE_DEPRECATED(5.0, "Use UpdatePerceptionAllowList instead")
 	void UpdatePerceptionWhitelist(const FAISenseID Channel, const bool bNewValue)
@@ -315,54 +315,54 @@ public:
 	}
 	
 
-	void RegisterStimulus(AActor* Source, const FAIStimulus& Stimulus);
-	void ProcessStimuli();
+	AIMODULE_API void RegisterStimulus(AActor* Source, const FAIStimulus& Stimulus);
+	AIMODULE_API void ProcessStimuli();
 	/** Returns true if, as result of stimuli aging, this listener needs an update (like if some stimuli expired) */
-	bool AgeStimuli(const float ConstPerceptionAgingRate);
-	void ForgetActor(AActor* ActorToForget);
+	AIMODULE_API bool AgeStimuli(const float ConstPerceptionAgingRate);
+	AIMODULE_API void ForgetActor(AActor* ActorToForget);
 
 	/** basically cleans up PerceptualData, resulting in loss of all previous perception */
 	UFUNCTION(BlueprintCallable, Category = "AI|Perception")
-	void ForgetAll();
+	AIMODULE_API void ForgetAll();
 
-	float GetYoungestStimulusAge(const AActor& Source) const;
-	bool HasAnyActiveStimulus(const AActor& Source) const;
-	bool HasAnyCurrentStimulus(const AActor& Source) const;
-	bool HasActiveStimulus(const AActor& Source, FAISenseID Sense) const;
+	AIMODULE_API float GetYoungestStimulusAge(const AActor& Source) const;
+	AIMODULE_API bool HasAnyActiveStimulus(const AActor& Source) const;
+	AIMODULE_API bool HasAnyCurrentStimulus(const AActor& Source) const;
+	AIMODULE_API bool HasActiveStimulus(const AActor& Source, FAISenseID Sense) const;
 
 #if WITH_GAMEPLAY_DEBUGGER_MENU
-	virtual void DescribeSelfToGameplayDebugger(FGameplayDebuggerCategory* DebuggerCategory) const;
+	AIMODULE_API virtual void DescribeSelfToGameplayDebugger(FGameplayDebuggerCategory* DebuggerCategory) const;
 #endif // WITH_GAMEPLAY_DEBUGGER_MENU
 
 #if ENABLE_VISUAL_LOG
-	virtual void DescribeSelfToVisLog(FVisualLogEntry* Snapshot) const;
+	AIMODULE_API virtual void DescribeSelfToVisLog(FVisualLogEntry* Snapshot) const;
 #endif // ENABLE_VISUAL_LOG
 
 	//----------------------------------------------------------------------//
 	// blueprint interface
 	//----------------------------------------------------------------------//
 	UFUNCTION(BlueprintCallable, Category = "AI|Perception")
-	void GetPerceivedHostileActors(TArray<AActor*>& OutActors) const;
+	AIMODULE_API void GetPerceivedHostileActors(TArray<AActor*>& OutActors) const;
 
 	UFUNCTION(BlueprintCallable, Category = "AI|Perception")
-	void GetPerceivedHostileActorsBySense(const TSubclassOf<UAISense> SenseToUse, TArray<AActor*>& OutActors) const;
+	AIMODULE_API void GetPerceivedHostileActorsBySense(const TSubclassOf<UAISense> SenseToUse, TArray<AActor*>& OutActors) const;
 
 	/** If SenseToUse is none all actors currently perceived in any way will get fetched */
 	UFUNCTION(BlueprintCallable, Category = "AI|Perception")
-	void GetCurrentlyPerceivedActors(TSubclassOf<UAISense> SenseToUse, TArray<AActor*>& OutActors) const;
+	AIMODULE_API void GetCurrentlyPerceivedActors(TSubclassOf<UAISense> SenseToUse, TArray<AActor*>& OutActors) const;
 
 	/** If SenseToUse is none all actors ever perceived in any way (and not forgotten yet) will get fetched */
 	UFUNCTION(BlueprintCallable, Category = "AI|Perception")
-	void GetKnownPerceivedActors(TSubclassOf<UAISense> SenseToUse, TArray<AActor*>& OutActors) const;
+	AIMODULE_API void GetKnownPerceivedActors(TSubclassOf<UAISense> SenseToUse, TArray<AActor*>& OutActors) const;
 
 	/** Retrieves whatever has been sensed about given actor */
 	UFUNCTION(BlueprintCallable, Category = "AI|Perception")
-	bool GetActorsPerception(AActor* Actor, FActorPerceptionBlueprintInfo& Info);
+	AIMODULE_API bool GetActorsPerception(AActor* Actor, FActorPerceptionBlueprintInfo& Info);
 
 	/** Note that this works only if given sense has been already configured for
 	 *	this component instance */
 	UFUNCTION(BlueprintCallable, Category = "AI|Perception")
-	void SetSenseEnabled(TSubclassOf<UAISense> SenseClass, const bool bEnable);
+	AIMODULE_API void SetSenseEnabled(TSubclassOf<UAISense> SenseClass, const bool bEnable);
 
 	//////////////////////////////////////////////////////////////////////////
 	// Might want to move these to special "BP_AIPerceptionComponent"
@@ -414,21 +414,21 @@ protected:
 	const FActorPerceptionContainer& GetPerceptualData() const { return PerceptualData; }
 
 	/** called to clean up on owner's end play or destruction */
-	virtual void CleanUp();
+	AIMODULE_API virtual void CleanUp();
 
-	void RemoveDeadData();
+	AIMODULE_API void RemoveDeadData();
 
 	/** Updates the stimulus entry in StimulusStore, if NewStimulus is more recent or stronger */
-	virtual void RefreshStimulus(FAIStimulus& StimulusStore, const FAIStimulus& NewStimulus);
+	AIMODULE_API virtual void RefreshStimulus(FAIStimulus& StimulusStore, const FAIStimulus& NewStimulus);
 
 	/** @note no need to call super implementation, it's there just for some validity checking */
-	virtual void HandleExpiredStimulus(FAIStimulus& StimulusStore);
+	AIMODULE_API virtual void HandleExpiredStimulus(FAIStimulus& StimulusStore);
 	
 private:
 	friend UAIPerceptionSystem;
 
-	void RegisterSenseConfig(UAISenseConfig& SenseConfig, UAIPerceptionSystem& AIPerceptionSys);
+	AIMODULE_API void RegisterSenseConfig(UAISenseConfig& SenseConfig, UAIPerceptionSystem& AIPerceptionSys);
 	void StoreListenerId(FPerceptionListenerID InListenerId) { PerceptionListenerId = InListenerId; }
-	void SetMaxStimulusAge(FAISenseID SenseId, float MaxAge);
+	AIMODULE_API void SetMaxStimulusAge(FAISenseID SenseId, float MaxAge);
 };
 
diff --git a/Engine/Source/Runtime/AIModule/Classes/Perception/AIPerceptionListenerInterface.h b/Engine/Source/Runtime/AIModule/Classes/Perception/AIPerceptionListenerInterface.h
--- a/Engine/Source/Runtime/AIModule/Classes/Perception/AIPerceptionListenerInterface.h
+++ b/Engine/Source/Runtime/AIModule/Classes/Perception/AIPerceptionListenerInterface.h
@@ -9,13 +9,13 @@
 
 class UAIPerceptionComponent;
 
-UINTERFACE()
-class AIMODULE_API UAIPerceptionListenerInterface : public UInterface
+UINTERFACE(MinimalAPI)
+class UAIPerceptionListenerInterface : public UInterface
 {
 	GENERATED_UINTERFACE_BODY()
 };
 
-class AIMODULE_API IAIPerceptionListenerInterface
+class IAIPerceptionListenerInterface
 {
 	GENERATED_IINTERFACE_BODY()
 
diff --git a/Engine/Source/Runtime/AIModule/Classes/Perception/AIPerceptionStimuliSourceComponent.h b/Engine/Source/Runtime/AIModule/Classes/Perception/AIPerceptionStimuliSourceComponent.h
--- a/Engine/Source/Runtime/AIModule/Classes/Perception/AIPerceptionStimuliSourceComponent.h
+++ b/Engine/Source/Runtime/AIModule/Classes/Perception/AIPerceptionStimuliSourceComponent.h
@@ -11,8 +11,8 @@
 #include "AIPerceptionStimuliSourceComponent.generated.h"
 
 /** Gives owning actor a way to auto-register as perception system's sense stimuli source */
-UCLASS(ClassGroup = AI, HideCategories = (Activation, Collision), meta = (BlueprintSpawnableComponent), config = Game)
-class AIMODULE_API UAIPerceptionStimuliSourceComponent : public UActorComponent
+UCLASS(ClassGroup = AI, HideCategories = (Activation, Collision), meta = (BlueprintSpawnableComponent), config = Game, MinimalAPI)
+class UAIPerceptionStimuliSourceComponent : public UActorComponent
 {
 	GENERATED_BODY()
 
@@ -25,29 +25,29 @@ protected:
 	UPROPERTY(EditAnywhere, Category = "AI Perception", BlueprintReadOnly)
 	TArray<TSubclassOf<UAISense> > RegisterAsSourceForSenses;
 
-	virtual void OnRegister() override;
+	AIMODULE_API virtual void OnRegister() override;
 public:
 
-	UAIPerceptionStimuliSourceComponent(const FObjectInitializer& ObjectInitializer);
+	AIMODULE_API UAIPerceptionStimuliSourceComponent(const FObjectInitializer& ObjectInitializer);
 
 #if WITH_EDITOR
-	virtual void PostEditChangeProperty(FPropertyChangedEvent& PropertyChangedEvent) override;
+	AIMODULE_API virtual void PostEditChangeProperty(FPropertyChangedEvent& PropertyChangedEvent) override;
 #endif // WITH_EDITOR
 
 	/** Registers owning actor as source of stimuli for senses specified in RegisterAsSourceForSenses. 
 	 *	Note that you don't have to do it if bAutoRegisterAsSource == true */
 	UFUNCTION(BlueprintCallable, Category = "AI|Perception")
-	void RegisterWithPerceptionSystem();
+	AIMODULE_API void RegisterWithPerceptionSystem();
 
 	/** Registers owning actor as source for specified sense class */
 	UFUNCTION(BlueprintCallable, Category = "AI|Perception")
-	void RegisterForSense(TSubclassOf<UAISense> SenseClass);
+	AIMODULE_API void RegisterForSense(TSubclassOf<UAISense> SenseClass);
 
 	/** Unregister owning actor from being a source of sense stimuli */
 	UFUNCTION(BlueprintCallable, Category = "AI|Perception")
-	void UnregisterFromPerceptionSystem();
+	AIMODULE_API void UnregisterFromPerceptionSystem();
 
 	/** Unregisters owning actor from sources list of a specified sense class */
 	UFUNCTION(BlueprintCallable, Category = "AI|Perception")
-	void UnregisterFromSense(TSubclassOf<UAISense> SenseClass);
+	AIMODULE_API void UnregisterFromSense(TSubclassOf<UAISense> SenseClass);
 };
diff --git a/Engine/Source/Runtime/AIModule/Classes/Perception/AIPerceptionSystem.h b/Engine/Source/Runtime/AIModule/Classes/Perception/AIPerceptionSystem.h
--- a/Engine/Source/Runtime/AIModule/Classes/Perception/AIPerceptionSystem.h
+++ b/Engine/Source/Runtime/AIModule/Classes/Perception/AIPerceptionSystem.h
@@ -26,18 +26,18 @@ class APawn;
 /**
  *	By design checks perception between hostile teams
  */
-UCLASS(ClassGroup=AI, config=Game, defaultconfig)
-class AIMODULE_API UAIPerceptionSystem : public UAISubsystem
+UCLASS(ClassGroup=AI, config=Game, defaultconfig, MinimalAPI)
+class UAIPerceptionSystem : public UAISubsystem
 {
 	GENERATED_BODY()
 		
 public:
 
-	UAIPerceptionSystem(const FObjectInitializer& ObjectInitializer = FObjectInitializer::Get());
+	AIMODULE_API UAIPerceptionSystem(const FObjectInitializer& ObjectInitializer = FObjectInitializer::Get());
 	
 	// FTickableGameObject begin
-	virtual void Tick(float DeltaTime) override;
-	virtual TStatId GetStatId() const override;
+	AIMODULE_API virtual void Tick(float DeltaTime) override;
+	AIMODULE_API virtual TStatId GetStatId() const override;
 	// FTickableGameObject end
 
 protected:	
@@ -87,8 +87,8 @@ public:
 	FORCEINLINE bool IsSenseInstantiated(const FAISenseID& SenseID) const { return SenseID.IsValid() && Senses.IsValidIndex(SenseID) && Senses[SenseID] != nullptr; }
 
 	/** Registers listener if not registered */
-	void UpdateListener(UAIPerceptionComponent& Listener);
-	void UnregisterListener(UAIPerceptionComponent& Listener);
+	AIMODULE_API void UpdateListener(UAIPerceptionComponent& Listener);
+	AIMODULE_API void UnregisterListener(UAIPerceptionComponent& Listener);
 
 	template<typename FEventClass, typename FSenseClass = typename FEventClass::FSenseClass>
 	void OnEvent(const FEventClass& Event)
@@ -126,60 +126,60 @@ public:
 	}
 
 	UFUNCTION(BlueprintCallable, Category = "AI|Perception")
-	void ReportEvent(UAISenseEvent* PerceptionEvent);
+	AIMODULE_API void ReportEvent(UAISenseEvent* PerceptionEvent);
 
 	UFUNCTION(BlueprintCallable, Category = "AI|Perception", meta = (WorldContext="WorldContextObject"))
-	static void ReportPerceptionEvent(UObject* WorldContextObject, UAISenseEvent* PerceptionEvent);
+	static AIMODULE_API void ReportPerceptionEvent(UObject* WorldContextObject, UAISenseEvent* PerceptionEvent);
 
 	/** Registers a source of given sense's stimuli */
 	template<typename FSenseClass>
 	void RegisterSource(AActor& SourceActor);
 
 	/** Registers given actor as a source for all registered senses */
-	void RegisterSource(AActor& SourceActor);
+	AIMODULE_API void RegisterSource(AActor& SourceActor);
 
-	void RegisterSourceForSenseClass(TSubclassOf<UAISense> Sense, AActor& Target);
+	AIMODULE_API void RegisterSourceForSenseClass(TSubclassOf<UAISense> Sense, AActor& Target);
 
 	/** 
 	 *	unregisters given actor from the list of active stimuli sources
 	 *	@param Sense if null will result in removing SourceActor from all the senses
 	 */
-	void UnregisterSource(AActor& SourceActor, const TSubclassOf<UAISense> Sense = nullptr);
+	AIMODULE_API void UnregisterSource(AActor& SourceActor, const TSubclassOf<UAISense> Sense = nullptr);
 
-	void OnListenerForgetsActor(const UAIPerceptionComponent& Listener, AActor& ActorToForget);
-	void OnListenerForgetsAll(const UAIPerceptionComponent& Listener);
-	void OnListenerConfigUpdated(FAISenseID SenseID, const UAIPerceptionComponent& Listener);
+	AIMODULE_API void OnListenerForgetsActor(const UAIPerceptionComponent& Listener, AActor& ActorToForget);
+	AIMODULE_API void OnListenerForgetsAll(const UAIPerceptionComponent& Listener);
+	AIMODULE_API void OnListenerConfigUpdated(FAISenseID SenseID, const UAIPerceptionComponent& Listener);
 
-	void RegisterDelayedStimulus(FPerceptionListenerID ListenerId, float Delay, AActor* Instigator, const FAIStimulus& Stimulus);
+	AIMODULE_API void RegisterDelayedStimulus(FPerceptionListenerID ListenerId, float Delay, AActor* Instigator, const FAIStimulus& Stimulus);
 
-	static UAIPerceptionSystem* GetCurrent(UObject* WorldContextObject);
-	static UAIPerceptionSystem* GetCurrent(UWorld& World);
+	static AIMODULE_API UAIPerceptionSystem* GetCurrent(UObject* WorldContextObject);
+	static AIMODULE_API UAIPerceptionSystem* GetCurrent(UWorld& World);
 
-	static void MakeNoiseImpl(AActor* NoiseMaker, float Loudness, APawn* NoiseInstigator, const FVector& NoiseLocation, float MaxRange, FName Tag);
+	static AIMODULE_API void MakeNoiseImpl(AActor* NoiseMaker, float Loudness, APawn* NoiseInstigator, const FVector& NoiseLocation, float MaxRange, FName Tag);
 
 	UFUNCTION(BlueprintCallable, Category = "AI|Perception", meta = (WorldContext="WorldContextObject"))
-	static bool RegisterPerceptionStimuliSource(UObject* WorldContextObject, TSubclassOf<UAISense> Sense, AActor* Target);
+	static AIMODULE_API bool RegisterPerceptionStimuliSource(UObject* WorldContextObject, TSubclassOf<UAISense> Sense, AActor* Target);
 
-	FAISenseID RegisterSenseClass(TSubclassOf<UAISense> SenseClass);
+	AIMODULE_API FAISenseID RegisterSenseClass(TSubclassOf<UAISense> SenseClass);
 
 	UFUNCTION(BlueprintCallable, Category = "AI|Perception", meta = (WorldContext="WorldContextObject"))
-	static TSubclassOf<UAISense> GetSenseClassForStimulus(UObject* WorldContextObject, const FAIStimulus& Stimulus);
+	static AIMODULE_API TSubclassOf<UAISense> GetSenseClassForStimulus(UObject* WorldContextObject, const FAIStimulus& Stimulus);
 
 #if WITH_GAMEPLAY_DEBUGGER_MENU
-	virtual void DescribeSelfToGameplayDebugger(FGameplayDebuggerCategory& DebuggerCategory) const;
+	AIMODULE_API virtual void DescribeSelfToGameplayDebugger(FGameplayDebuggerCategory& DebuggerCategory) const;
 #endif // WITH_GAMEPLAY_DEBUGGER_MENU
 
 protected:
 	
 	UFUNCTION()
-	void OnPerceptionStimuliSourceEndPlay(AActor* Actor, EEndPlayReason::Type EndPlayReason);
+	AIMODULE_API void OnPerceptionStimuliSourceEndPlay(AActor* Actor, EEndPlayReason::Type EndPlayReason);
 	
 	/** requests registration of a given actor as a perception data source for specified sense */
-	void RegisterSource(FAISenseID SenseID, AActor& SourceActor);
+	AIMODULE_API void RegisterSource(FAISenseID SenseID, AActor& SourceActor);
 
 	/** iterates over all pawns and registeres them as a source for sense indicated by SenseID. Note that this will 
 	 *  be performed only for senses that request that (see UAISense.bAutoRegisterAllPawnsAsSources).*/
-	virtual void RegisterAllPawnsAsSourcesForSense(FAISenseID SenseID);
+	AIMODULE_API virtual void RegisterAllPawnsAsSourcesForSense(FAISenseID SenseID);
 
 	enum EDelayedStimulusSorting 
 	{
@@ -188,22 +188,22 @@ protected:
 	};
 	/** sorts DelayedStimuli and delivers all the ones that are no longer "in the future"
 	 *	@return true if any stimuli has become "current" stimuli (meaning being no longer in future) */
-	bool DeliverDelayedStimuli(EDelayedStimulusSorting Sorting);
-	void OnNewListener(const FPerceptionListener& NewListener);
-	void OnListenerUpdate(const FPerceptionListener& UpdatedListener);
-	void OnListenerRemoved(const FPerceptionListener& UpdatedListener);
-	void PerformSourceRegistration();
+	AIMODULE_API bool DeliverDelayedStimuli(EDelayedStimulusSorting Sorting);
+	AIMODULE_API void OnNewListener(const FPerceptionListener& NewListener);
+	AIMODULE_API void OnListenerUpdate(const FPerceptionListener& UpdatedListener);
+	AIMODULE_API void OnListenerRemoved(const FPerceptionListener& UpdatedListener);
+	AIMODULE_API void PerformSourceRegistration();
 
 	/** Returns true if aging resulted in tagging any of the listeners to process 
 	 *	its stimuli (@see MarkForStimulusProcessing)*/
-	bool AgeStimuli(const float Amount);
+	AIMODULE_API bool AgeStimuli(const float Amount);
 
 	friend class UAISense;
 	FORCEINLINE AIPerception::FListenerMap& GetListenersMap() { return ListenerContainer; }
 
 	friend class UAISystem;
-	virtual void OnNewPawn(APawn& Pawn);
-	virtual void StartPlay();
+	AIMODULE_API virtual void OnNewPawn(APawn& Pawn);
+	AIMODULE_API virtual void StartPlay();
 
 	/** Timestamp of the next stimuli aging */
 	double NextStimuliAgingTick;
diff --git a/Engine/Source/Runtime/AIModule/Classes/Perception/AIPerceptionTypes.h b/Engine/Source/Runtime/AIModule/Classes/Perception/AIPerceptionTypes.h
--- a/Engine/Source/Runtime/AIModule/Classes/Perception/AIPerceptionTypes.h
+++ b/Engine/Source/Runtime/AIModule/Classes/Perception/AIPerceptionTypes.h
@@ -12,12 +12,12 @@ class UAIPerceptionComponent;
 class UAISense;
 
 //////////////////////////////////////////////////////////////////////////
-struct AIMODULE_API FAISenseCounter : FAIBasicCounter<uint8>
+struct FAISenseCounter : FAIBasicCounter<uint8>
 {};
 typedef FAINamedID<FAISenseCounter> FAISenseID;
 
 //////////////////////////////////////////////////////////////////////////
-struct AIMODULE_API FPerceptionListenerCounter : FAIBasicCounter<uint32>
+struct FPerceptionListenerCounter : FAIBasicCounter<uint32>
 {};
 typedef FAIGenericID<FPerceptionListenerCounter> FPerceptionListenerID;
 
@@ -126,11 +126,11 @@ struct FPerceptionChannelAllowList
 };
 
 USTRUCT(BlueprintType)
-struct AIMODULE_API FAIStimulus
+struct FAIStimulus
 {
 	GENERATED_USTRUCT_BODY()
 
-	static const float NeverHappenedAge;
+	static AIMODULE_API const float NeverHappenedAge;
 
 	enum FResult
 	{
@@ -169,7 +169,7 @@ protected:
 public:
 	
 	/** this is the recommended constructor. Use others if you know what you're doing. */
-	FAIStimulus(const UAISense& Sense, float StimulusStrength, const FVector& InStimulusLocation, const FVector& InReceiverLocation, FResult Result = SensingSucceeded, FName InStimulusTag = NAME_None);
+	AIMODULE_API FAIStimulus(const UAISense& Sense, float StimulusStrength, const FVector& InStimulusLocation, const FVector& InReceiverLocation, FResult Result = SensingSucceeded, FName InStimulusTag = NAME_None);
 
 	// default constructor
 	FAIStimulus()
@@ -198,12 +198,12 @@ public:
 	FORCEINLINE bool IsValid() const { return Type != FAISenseID::InvalidID() && GetAge() < NeverHappenedAge; }
 
 #if !(UE_BUILD_SHIPPING || UE_BUILD_TEST)
-	FString GetDebugDescription() const;
+	AIMODULE_API FString GetDebugDescription() const;
 #endif // !(UE_BUILD_SHIPPING || UE_BUILD_TEST)
 };
 
 USTRUCT(BlueprintType)
-struct AIMODULE_API FAISenseAffiliationFilter
+struct FAISenseAffiliationFilter
 {
 	GENERATED_USTRUCT_BODY()
 
@@ -241,7 +241,7 @@ struct AIMODULE_API FAISenseAffiliationFilter
 };
 
 /** Should contain only cached information common to all senses. Sense-specific data needs to be stored by senses themselves */
-struct AIMODULE_API FPerceptionListener
+struct FPerceptionListener
 {
 	TWeakObjectPtr<UAIPerceptionComponent> Listener;
 
@@ -257,36 +257,36 @@ private:
 
 	FPerceptionListenerID ListenerID;
 
-	FPerceptionListener();
+	AIMODULE_API FPerceptionListener();
 public:
-	FPerceptionListener(UAIPerceptionComponent& InListener);
+	AIMODULE_API FPerceptionListener(UAIPerceptionComponent& InListener);
 
-	void UpdateListenerProperties(UAIPerceptionComponent& Listener);
+	AIMODULE_API void UpdateListenerProperties(UAIPerceptionComponent& Listener);
 
 	bool operator==(const UAIPerceptionComponent* Other) const { return Listener.Get() == Other; }
 	bool operator==(const FPerceptionListener& Other) const { return Listener == Other.Listener; }
 
-	void CacheLocation();
+	AIMODULE_API void CacheLocation();
 
-	void RegisterStimulus(AActor* Source, const FAIStimulus& Stimulus);
+	AIMODULE_API void RegisterStimulus(AActor* Source, const FAIStimulus& Stimulus);
 
 	FORCEINLINE bool HasAnyNewStimuli() const { return bHasStimulusToProcess; }
-	void ProcessStimuli();
+	AIMODULE_API void ProcessStimuli();
 
 	FORCEINLINE bool HasSense(FAISenseID SenseID) const { return Filter.ShouldRespondToChannel(SenseID); }
 
 	// used to remove "dead" listeners
-	static const FPerceptionListener NullListener;
+	static AIMODULE_API const FPerceptionListener NullListener;
 
 	FORCEINLINE FPerceptionListenerID GetListenerID() const { return ListenerID; }
 
-	FName GetBodyActorName() const;
-	uint32 GetBodyActorUniqueID() const;
+	AIMODULE_API FName GetBodyActorName() const;
+	AIMODULE_API uint32 GetBodyActorUniqueID() const;
 
 	/** Returns pointer to the actor representing this listener's physical body */
-	const AActor* GetBodyActor() const;
+	AIMODULE_API const AActor* GetBodyActor() const;
 
-	const IGenericTeamAgentInterface* GetTeamAgent() const;
+	AIMODULE_API const IGenericTeamAgentInterface* GetTeamAgent() const;
 
 private:
 	friend class UAIPerceptionSystem;
@@ -294,7 +294,7 @@ private:
 	FORCEINLINE void MarkForStimulusProcessing() { bHasStimulusToProcess = true; }
 };
 
-struct AIMODULE_API FPerceptionStimuliSource
+struct FPerceptionStimuliSource
 {
 	TWeakObjectPtr<AActor> SourceActor;
 	FPerceptionChannelAllowList RelevantSenses;
diff --git a/Engine/Source/Runtime/AIModule/Classes/Perception/AISense.h b/Engine/Source/Runtime/AIModule/Classes/Perception/AISense.h
--- a/Engine/Source/Runtime/AIModule/Classes/Perception/AISense.h
+++ b/Engine/Source/Runtime/AIModule/Classes/Perception/AISense.h
@@ -16,12 +16,12 @@ class UAISenseEvent;
 
 DECLARE_DELEGATE_OneParam(FOnPerceptionListenerUpdateDelegate, const FPerceptionListener&);
 
-UCLASS(ClassGroup = AI, abstract, config = Engine)
-class AIMODULE_API UAISense : public UObject
+UCLASS(ClassGroup = AI, abstract, config = Engine, MinimalAPI)
+class UAISense : public UObject
 {
 	GENERATED_UCLASS_BODY()
 
-	static const float SuspendNextUpdate;
+	static AIMODULE_API const float SuspendNextUpdate;
 
 protected:
 	UPROPERTY(EditDefaultsOnly, BlueprintReadOnly, Category = "AI Perception", config)
@@ -61,10 +61,10 @@ protected:
 				
 public:
 
-	virtual UWorld* GetWorld() const override;
+	AIMODULE_API virtual UWorld* GetWorld() const override;
 
 	/** use with caution! Needs to be called before any senses get instantiated or listeners registered. DOES NOT update any perceptions system instances */
-	static void HardcodeSenseID(TSubclassOf<UAISense> SenseClass, FAISenseID HardcodedID);
+	static AIMODULE_API void HardcodeSenseID(TSubclassOf<UAISense> SenseClass, FAISenseID HardcodedID);
 
 	static FAISenseID GetSenseID(const TSubclassOf<UAISense> SenseClass) { return SenseClass ? ((const UAISense*)SenseClass->GetDefaultObject())->SenseID : FAISenseID::InvalidID(); }
 	template<typename TSense>
@@ -76,7 +76,7 @@ public:
 
 	FORCEINLINE bool WantsUpdateOnlyOnPerceptionValueChange() const { return (NotifyType == EAISenseNotifyType::OnPerceptionChange); }
 
-	virtual void PostInitProperties() override;
+	AIMODULE_API virtual void PostInitProperties() override;
 
 	/** 
 	 *	@return should this sense be ticked now
@@ -99,8 +99,8 @@ public:
 	virtual void RegisterSource(AActor& SourceActors){}
 	virtual void UnregisterSource(AActor& SourceActors){}
 
-	virtual void RegisterWrappedEvent(UAISenseEvent& PerceptionEvent);
-	virtual FAISenseID UpdateSenseID();
+	AIMODULE_API virtual void RegisterWrappedEvent(UAISenseEvent& PerceptionEvent);
+	AIMODULE_API virtual FAISenseID UpdateSenseID();
 
 	bool NeedsNotificationOnForgetting() const { return bNeedsForgettingNotification; }
 	virtual void OnListenerForgetsActor(const FPerceptionListener& Listener, AActor& ActorToForget) {}
@@ -115,14 +115,14 @@ public:
 	bool ShouldAutoRegisterAllPawnsAsSources() const { return bAutoRegisterAllPawnsAsSources; }
 
 #if WITH_GAMEPLAY_DEBUGGER_MENU
-	virtual void DescribeSelfToGameplayDebugger(const UAIPerceptionSystem& PerceptionSystem, FGameplayDebuggerCategory& DebuggerCategory) const;
+	AIMODULE_API virtual void DescribeSelfToGameplayDebugger(const UAIPerceptionSystem& PerceptionSystem, FGameplayDebuggerCategory& DebuggerCategory) const;
 #endif // WITH_GAMEPLAY_DEBUGGER_MENU
 
 protected:
 	friend UAIPerceptionSystem;
 	/** gets called when perception system gets notified about new spawned pawn. 
 	 *	@Note: do not call super implementation. It's used to detect when subclasses don't override it */
-	virtual void OnNewPawn(APawn& NewPawn);
+	AIMODULE_API virtual void OnNewPawn(APawn& NewPawn);
 
 	/** @return time until next update */
 	virtual float Update() { return FLT_MAX; }
@@ -135,12 +135,12 @@ protected:
 
 	FORCEINLINE UAIPerceptionSystem* GetPerceptionSystem() { return PerceptionSystemInstance; }
 
-	void SetSenseID(FAISenseID Index);
+	AIMODULE_API void SetSenseID(FAISenseID Index);
 
 	/** returning pointer rather then a reference to prevent users from
 	 *	accidentally creating copies by creating non-reference local vars */
-	AIPerception::FListenerMap* GetListeners();
+	AIMODULE_API AIPerception::FListenerMap* GetListeners();
 
 	/** To be called only for BP-generated classes */
-	void ForceSenseID(FAISenseID SenseID);
+	AIMODULE_API void ForceSenseID(FAISenseID SenseID);
 };
diff --git a/Engine/Source/Runtime/AIModule/Classes/Perception/AISenseBlueprintListener.h b/Engine/Source/Runtime/AIModule/Classes/Perception/AISenseBlueprintListener.h
--- a/Engine/Source/Runtime/AIModule/Classes/Perception/AISenseBlueprintListener.h
+++ b/Engine/Source/Runtime/AIModule/Classes/Perception/AISenseBlueprintListener.h
@@ -5,8 +5,8 @@
 #include "Engine/UserDefinedStruct.h"
 #include "AISenseBlueprintListener.generated.h"
 
-UCLASS(ClassGroup=AI)
-class AIMODULE_API UAISenseBlueprintListener : public UUserDefinedStruct
+UCLASS(ClassGroup=AI, MinimalAPI)
+class UAISenseBlueprintListener : public UUserDefinedStruct
 {
 	GENERATED_BODY()
 };
diff --git a/Engine/Source/Runtime/AIModule/Classes/Perception/AISenseConfig.h b/Engine/Source/Runtime/AIModule/Classes/Perception/AISenseConfig.h
--- a/Engine/Source/Runtime/AIModule/Classes/Perception/AISenseConfig.h
+++ b/Engine/Source/Runtime/AIModule/Classes/Perception/AISenseConfig.h
@@ -14,8 +14,8 @@
 class FGameplayDebuggerCategory;
 class UAIPerceptionComponent;
 
-UCLASS(ClassGroup = AI, abstract, EditInlineNew, config=Game)
-class AIMODULE_API UAISenseConfig : public UObject
+UCLASS(ClassGroup = AI, abstract, EditInlineNew, config=Game, MinimalAPI)
+class UAISenseConfig : public UObject
 {
 	GENERATED_BODY()
 
@@ -35,18 +35,18 @@ protected:
 	mutable FString CachedSenseName;
 
 public:
-	UAISenseConfig(const FObjectInitializer& ObjectInitializer = FObjectInitializer::Get());
+	AIMODULE_API UAISenseConfig(const FObjectInitializer& ObjectInitializer = FObjectInitializer::Get());
 
-	virtual TSubclassOf<UAISense> GetSenseImplementation() const;
+	AIMODULE_API virtual TSubclassOf<UAISense> GetSenseImplementation() const;
 
-	FAISenseID GetSenseID() const;
-	FString GetSenseName() const;
+	AIMODULE_API FAISenseID GetSenseID() const;
+	AIMODULE_API FString GetSenseName() const;
 	FColor GetDebugColor() const { return DebugColor; }
 	void SetMaxAge(const float NewMaxAge) { MaxAge = NewMaxAge; }
 	float GetMaxAge() const { return MaxAge == 0.f ? FAIStimulus::NeverHappenedAge : MaxAge; }	
 	bool IsEnabled() const { return bStartsEnabled; }
 
 #if WITH_GAMEPLAY_DEBUGGER_MENU
-	virtual void DescribeSelfToGameplayDebugger(const UAIPerceptionComponent* PerceptionComponent, FGameplayDebuggerCategory* DebuggerCategory) const;
+	AIMODULE_API virtual void DescribeSelfToGameplayDebugger(const UAIPerceptionComponent* PerceptionComponent, FGameplayDebuggerCategory* DebuggerCategory) const;
 #endif // WITH_GAMEPLAY_DEBUGGER_MENU
 };
diff --git a/Engine/Source/Runtime/AIModule/Classes/Perception/AISenseConfig_Blueprint.h b/Engine/Source/Runtime/AIModule/Classes/Perception/AISenseConfig_Blueprint.h
--- a/Engine/Source/Runtime/AIModule/Classes/Perception/AISenseConfig_Blueprint.h
+++ b/Engine/Source/Runtime/AIModule/Classes/Perception/AISenseConfig_Blueprint.h
@@ -11,8 +11,8 @@
 
 class UAISense_Blueprint;
 
-UCLASS(Blueprintable, Abstract, hidedropdown)
-class AIMODULE_API UAISenseConfig_Blueprint : public UAISenseConfig
+UCLASS(Blueprintable, Abstract, hidedropdown, MinimalAPI)
+class UAISenseConfig_Blueprint : public UAISenseConfig
 {
 	GENERATED_UCLASS_BODY()
 
@@ -20,5 +20,5 @@ public:
 	UPROPERTY(EditDefaultsOnly, BlueprintReadOnly, Category = "Sense", NoClear, config)
 	TSubclassOf<UAISense_Blueprint> Implementation;
 
-	virtual TSubclassOf<UAISense> GetSenseImplementation() const override;
+	AIMODULE_API virtual TSubclassOf<UAISense> GetSenseImplementation() const override;
 };
diff --git a/Engine/Source/Runtime/AIModule/Classes/Perception/AISenseConfig_Damage.h b/Engine/Source/Runtime/AIModule/Classes/Perception/AISenseConfig_Damage.h
--- a/Engine/Source/Runtime/AIModule/Classes/Perception/AISenseConfig_Damage.h
+++ b/Engine/Source/Runtime/AIModule/Classes/Perception/AISenseConfig_Damage.h
@@ -10,8 +10,8 @@
 #include "Perception/AISense_Damage.h"
 #include "AISenseConfig_Damage.generated.h"
 
-UCLASS(meta = (DisplayName = "AI Damage sense config"))
-class AIMODULE_API UAISenseConfig_Damage : public UAISenseConfig
+UCLASS(meta = (DisplayName = "AI Damage sense config"), MinimalAPI)
+class UAISenseConfig_Damage : public UAISenseConfig
 {
 	GENERATED_UCLASS_BODY()
 
@@ -19,5 +19,5 @@ public:
 	UPROPERTY(EditDefaultsOnly, BlueprintReadOnly, Category = "Sense", NoClear, config)
 	TSubclassOf<UAISense_Damage> Implementation;
 
-	virtual TSubclassOf<UAISense> GetSenseImplementation() const override;
+	AIMODULE_API virtual TSubclassOf<UAISense> GetSenseImplementation() const override;
 };
diff --git a/Engine/Source/Runtime/AIModule/Classes/Perception/AISenseConfig_Hearing.h b/Engine/Source/Runtime/AIModule/Classes/Perception/AISenseConfig_Hearing.h
--- a/Engine/Source/Runtime/AIModule/Classes/Perception/AISenseConfig_Hearing.h
+++ b/Engine/Source/Runtime/AIModule/Classes/Perception/AISenseConfig_Hearing.h
@@ -14,8 +14,8 @@
 class FGameplayDebuggerCategory;
 class UAIPerceptionComponent;
 
-UCLASS(meta = (DisplayName = "AI Hearing config"))
-class AIMODULE_API UAISenseConfig_Hearing : public UAISenseConfig
+UCLASS(meta = (DisplayName = "AI Hearing config"), MinimalAPI)
+class UAISenseConfig_Hearing : public UAISenseConfig
 {
 	GENERATED_UCLASS_BODY()
 
@@ -37,9 +37,9 @@ public:
 	UPROPERTY(EditDefaultsOnly, BlueprintReadOnly, Category = "Sense", config)
 	FAISenseAffiliationFilter DetectionByAffiliation;
 
-	virtual TSubclassOf<UAISense> GetSenseImplementation() const override;
+	AIMODULE_API virtual TSubclassOf<UAISense> GetSenseImplementation() const override;
 
 #if WITH_GAMEPLAY_DEBUGGER_MENU
-	virtual void DescribeSelfToGameplayDebugger(const UAIPerceptionComponent* PerceptionComponent, FGameplayDebuggerCategory* DebuggerCategory) const;
+	AIMODULE_API virtual void DescribeSelfToGameplayDebugger(const UAIPerceptionComponent* PerceptionComponent, FGameplayDebuggerCategory* DebuggerCategory) const;
 #endif // WITH_GAMEPLAY_DEBUGGER_MENU
 };
diff --git a/Engine/Source/Runtime/AIModule/Classes/Perception/AISenseConfig_Prediction.h b/Engine/Source/Runtime/AIModule/Classes/Perception/AISenseConfig_Prediction.h
--- a/Engine/Source/Runtime/AIModule/Classes/Perception/AISenseConfig_Prediction.h
+++ b/Engine/Source/Runtime/AIModule/Classes/Perception/AISenseConfig_Prediction.h
@@ -9,10 +9,10 @@
 #include "Perception/AISenseConfig.h"
 #include "AISenseConfig_Prediction.generated.h"
 
-UCLASS(meta = (DisplayName = "AI Prediction sense config"))
-class AIMODULE_API UAISenseConfig_Prediction : public UAISenseConfig
+UCLASS(meta = (DisplayName = "AI Prediction sense config"), MinimalAPI)
+class UAISenseConfig_Prediction : public UAISenseConfig
 {
 	GENERATED_UCLASS_BODY()
 public:	
-	virtual TSubclassOf<UAISense> GetSenseImplementation() const override;
+	AIMODULE_API virtual TSubclassOf<UAISense> GetSenseImplementation() const override;
 };
diff --git a/Engine/Source/Runtime/AIModule/Classes/Perception/AISenseConfig_Sight.h b/Engine/Source/Runtime/AIModule/Classes/Perception/AISenseConfig_Sight.h
--- a/Engine/Source/Runtime/AIModule/Classes/Perception/AISenseConfig_Sight.h
+++ b/Engine/Source/Runtime/AIModule/Classes/Perception/AISenseConfig_Sight.h
@@ -14,8 +14,8 @@
 class FGameplayDebuggerCategory;
 class UAIPerceptionComponent;
 
-UCLASS(meta = (DisplayName = "AI Sight config"))
-class AIMODULE_API UAISenseConfig_Sight : public UAISenseConfig
+UCLASS(meta = (DisplayName = "AI Sight config"), MinimalAPI)
+class UAISenseConfig_Sight : public UAISenseConfig
 {
 	GENERATED_UCLASS_BODY()
 
@@ -52,13 +52,13 @@ public:
 	UPROPERTY(EditDefaultsOnly, BlueprintReadOnly, Category = "Sense", config, meta = (UIMin = 0.0, ClampMin = 0.0))
 	float NearClippingRadius;
 
-	virtual TSubclassOf<UAISense> GetSenseImplementation() const override;
+	AIMODULE_API virtual TSubclassOf<UAISense> GetSenseImplementation() const override;
 
 #if WITH_EDITOR
-	virtual void PostEditChangeChainProperty(FPropertyChangedChainEvent& PropertyChangedEvent) override;
+	AIMODULE_API virtual void PostEditChangeChainProperty(FPropertyChangedChainEvent& PropertyChangedEvent) override;
 #endif // WITH_EDITOR
 
 #if WITH_GAMEPLAY_DEBUGGER_MENU
-	virtual void DescribeSelfToGameplayDebugger(const UAIPerceptionComponent* PerceptionComponent, FGameplayDebuggerCategory* DebuggerCategory) const;
+	AIMODULE_API virtual void DescribeSelfToGameplayDebugger(const UAIPerceptionComponent* PerceptionComponent, FGameplayDebuggerCategory* DebuggerCategory) const;
 #endif // WITH_GAMEPLAY_DEBUGGER_MENU
 };
diff --git a/Engine/Source/Runtime/AIModule/Classes/Perception/AISenseConfig_Team.h b/Engine/Source/Runtime/AIModule/Classes/Perception/AISenseConfig_Team.h
--- a/Engine/Source/Runtime/AIModule/Classes/Perception/AISenseConfig_Team.h
+++ b/Engine/Source/Runtime/AIModule/Classes/Perception/AISenseConfig_Team.h
@@ -10,8 +10,8 @@
 #include "Perception/AISense_Team.h"
 #include "AISenseConfig_Team.generated.h"
 
-UCLASS(meta = (DisplayName = "AI Team sense config"))
-class AIMODULE_API UAISenseConfig_Team : public UAISenseConfig
+UCLASS(meta = (DisplayName = "AI Team sense config"), MinimalAPI)
+class UAISenseConfig_Team : public UAISenseConfig
 {
 	GENERATED_UCLASS_BODY()
 public:	
diff --git a/Engine/Source/Runtime/AIModule/Classes/Perception/AISenseConfig_Touch.h b/Engine/Source/Runtime/AIModule/Classes/Perception/AISenseConfig_Touch.h
--- a/Engine/Source/Runtime/AIModule/Classes/Perception/AISenseConfig_Touch.h
+++ b/Engine/Source/Runtime/AIModule/Classes/Perception/AISenseConfig_Touch.h
@@ -9,10 +9,10 @@
 #include "Perception/AISenseConfig.h"
 #include "AISenseConfig_Touch.generated.h"
 
-UCLASS(meta = (DisplayName = "AI Touch config"))
-class AIMODULE_API UAISenseConfig_Touch : public UAISenseConfig
+UCLASS(meta = (DisplayName = "AI Touch config"), MinimalAPI)
+class UAISenseConfig_Touch : public UAISenseConfig
 {
 	GENERATED_UCLASS_BODY()
 public:	
-	virtual TSubclassOf<UAISense> GetSenseImplementation() const override;
+	AIMODULE_API virtual TSubclassOf<UAISense> GetSenseImplementation() const override;
 };
diff --git a/Engine/Source/Runtime/AIModule/Classes/Perception/AISenseEvent.h b/Engine/Source/Runtime/AIModule/Classes/Perception/AISenseEvent.h
--- a/Engine/Source/Runtime/AIModule/Classes/Perception/AISenseEvent.h
+++ b/Engine/Source/Runtime/AIModule/Classes/Perception/AISenseEvent.h
@@ -9,13 +9,13 @@
 #include "Perception/AIPerceptionTypes.h"
 #include "AISenseEvent.generated.h"
 
-UCLASS(ClassGroup = AI, abstract, EditInlineNew, config = Game)
-class AIMODULE_API UAISenseEvent : public UObject
+UCLASS(ClassGroup = AI, abstract, EditInlineNew, config = Game, MinimalAPI)
+class UAISenseEvent : public UObject
 {
 	GENERATED_BODY()
 
 public:
-	virtual FAISenseID GetSenseID() const PURE_VIRTUAL(UAISenseEvent::GetSenseID, return FAISenseID::InvalidID(););
+	AIMODULE_API virtual FAISenseID GetSenseID() const PURE_VIRTUAL(UAISenseEvent::GetSenseID, return FAISenseID::InvalidID(););
 
 #if ENABLE_VISUAL_LOG
 	virtual void DrawToVLog(UObject& LogOwner) const {}
diff --git a/Engine/Source/Runtime/AIModule/Classes/Perception/AISenseEvent_Damage.h b/Engine/Source/Runtime/AIModule/Classes/Perception/AISenseEvent_Damage.h
--- a/Engine/Source/Runtime/AIModule/Classes/Perception/AISenseEvent_Damage.h
+++ b/Engine/Source/Runtime/AIModule/Classes/Perception/AISenseEvent_Damage.h
@@ -8,8 +8,8 @@
 #include "Perception/AISense_Damage.h"
 #include "AISenseEvent_Damage.generated.h"
 
-UCLASS()
-class AIMODULE_API UAISenseEvent_Damage : public UAISenseEvent
+UCLASS(MinimalAPI)
+class UAISenseEvent_Damage : public UAISenseEvent
 {
 	GENERATED_BODY()
 
@@ -17,7 +17,7 @@ public:
 	UPROPERTY(EditAnywhere, BlueprintReadWrite, Category = "Sense")
 	FAIDamageEvent Event;
 	
-	virtual FAISenseID GetSenseID() const override;
+	AIMODULE_API virtual FAISenseID GetSenseID() const override;
 
 	FORCEINLINE FAIDamageEvent GetDamageEvent()
 	{
diff --git a/Engine/Source/Runtime/AIModule/Classes/Perception/AISenseEvent_Hearing.h b/Engine/Source/Runtime/AIModule/Classes/Perception/AISenseEvent_Hearing.h
--- a/Engine/Source/Runtime/AIModule/Classes/Perception/AISenseEvent_Hearing.h
+++ b/Engine/Source/Runtime/AIModule/Classes/Perception/AISenseEvent_Hearing.h
@@ -9,8 +9,8 @@
 #include "Perception/AISenseEvent.h"
 #include "AISenseEvent_Hearing.generated.h"
 
-UCLASS()
-class AIMODULE_API UAISenseEvent_Hearing : public UAISenseEvent
+UCLASS(MinimalAPI)
+class UAISenseEvent_Hearing : public UAISenseEvent
 {
 	GENERATED_BODY()
 
@@ -19,8 +19,8 @@ protected:
 	FAINoiseEvent Event;
 
 public:
-	UAISenseEvent_Hearing(const FObjectInitializer& ObjectInitializer = FObjectInitializer::Get());
-	virtual FAISenseID GetSenseID() const override;
+	AIMODULE_API UAISenseEvent_Hearing(const FObjectInitializer& ObjectInitializer = FObjectInitializer::Get());
+	AIMODULE_API virtual FAISenseID GetSenseID() const override;
 	
 	FORCEINLINE FAINoiseEvent GetNoiseEvent()
 	{
diff --git a/Engine/Source/Runtime/AIModule/Classes/Perception/AISense_Blueprint.h b/Engine/Source/Runtime/AIModule/Classes/Perception/AISense_Blueprint.h
--- a/Engine/Source/Runtime/AIModule/Classes/Perception/AISense_Blueprint.h
+++ b/Engine/Source/Runtime/AIModule/Classes/Perception/AISense_Blueprint.h
@@ -14,8 +14,8 @@ class UAIPerceptionComponent;
 class UAISenseEvent;
 class UUserDefinedStruct;
 
-UCLASS(ClassGroup = AI, Abstract, Blueprintable, hidedropdown)
-class AIMODULE_API UAISense_Blueprint : public UAISense
+UCLASS(ClassGroup = AI, Abstract, Blueprintable, hidedropdown, MinimalAPI)
+class UAISense_Blueprint : public UAISense
 {
 	GENERATED_BODY()
 
@@ -30,53 +30,53 @@ protected:
 	TArray<TObjectPtr<UAISenseEvent>> UnprocessedEvents;
 
 public:
-	UAISense_Blueprint(const FObjectInitializer& ObjectInitializer = FObjectInitializer::Get());
+	AIMODULE_API UAISense_Blueprint(const FObjectInitializer& ObjectInitializer = FObjectInitializer::Get());
 
 	//void RegisterEvent(const FAITeamStimulusEvent& Event);	
 
 	/** returns requested amount of time to pass until next frame. 
 	 *	Return 0 to get update every frame (WARNING: hits performance) */
 	UFUNCTION(BlueprintImplementableEvent)
-	float OnUpdate(const TArray<UAISenseEvent*>& EventsToProcess);
+	AIMODULE_API float OnUpdate(const TArray<UAISenseEvent*>& EventsToProcess);
 
 	/**
 	 *	@param PerceptionComponent is ActorListener's AIPerceptionComponent instance
 	 */
 	UFUNCTION(BlueprintImplementableEvent)
-	void OnListenerRegistered(AActor* ActorListener, UAIPerceptionComponent* PerceptionComponent);
+	AIMODULE_API void OnListenerRegistered(AActor* ActorListener, UAIPerceptionComponent* PerceptionComponent);
 
 	/**
 	 *	@param PerceptionComponent is ActorListener's AIPerceptionComponent instance
 	 */
 	UFUNCTION(BlueprintImplementableEvent)
-	void OnListenerUpdated(AActor* ActorListener, UAIPerceptionComponent* PerceptionComponent);
+	AIMODULE_API void OnListenerUpdated(AActor* ActorListener, UAIPerceptionComponent* PerceptionComponent);
 
 	/** called when a listener unregistered from this sense. Most often this is called due to actor's death
 	 *	@param PerceptionComponent is ActorListener's AIPerceptionComponent instance
 	 */
 	UFUNCTION(BlueprintImplementableEvent)
-	void OnListenerUnregistered(AActor* ActorListener, UAIPerceptionComponent* PerceptionComponent);
+	AIMODULE_API void OnListenerUnregistered(AActor* ActorListener, UAIPerceptionComponent* PerceptionComponent);
 	
 	UFUNCTION(BlueprintCallable, Category = "AI|Perception")
-	void GetAllListenerActors(TArray<AActor*>& ListenerActors) const;
+	AIMODULE_API void GetAllListenerActors(TArray<AActor*>& ListenerActors) const;
 
 	UFUNCTION(BlueprintCallable, Category = "AI|Perception")
-	void GetAllListenerComponents(TArray<UAIPerceptionComponent*>& ListenerComponents) const;
+	AIMODULE_API void GetAllListenerComponents(TArray<UAIPerceptionComponent*>& ListenerComponents) const;
 
 	/** called when sense's instance gets notified about new pawn that has just been spawned */
 	UFUNCTION(BlueprintImplementableEvent, DisplayName="OnNewPawn", meta=(ScriptName="OnNewPawn"))
-	void K2_OnNewPawn(APawn* NewPawn);
+	AIMODULE_API void K2_OnNewPawn(APawn* NewPawn);
 
-	virtual FAISenseID UpdateSenseID() override;
-	virtual void RegisterWrappedEvent(UAISenseEvent& PerceptionEvent) override;
+	AIMODULE_API virtual FAISenseID UpdateSenseID() override;
+	AIMODULE_API virtual void RegisterWrappedEvent(UAISenseEvent& PerceptionEvent) override;
 
 protected:
-	virtual void OnNewPawn(APawn& NewPawn) override;
-	virtual float Update() override;
+	AIMODULE_API virtual void OnNewPawn(APawn& NewPawn) override;
+	AIMODULE_API virtual float Update() override;
 	
-	void OnNewListenerImpl(const FPerceptionListener& NewListener);
-	void OnListenerUpdateImpl(const FPerceptionListener& UpdatedListener);
-	void OnListenerRemovedImpl(const FPerceptionListener& UpdatedListener);
+	AIMODULE_API void OnNewListenerImpl(const FPerceptionListener& NewListener);
+	AIMODULE_API void OnListenerUpdateImpl(const FPerceptionListener& UpdatedListener);
+	AIMODULE_API void OnListenerRemovedImpl(const FPerceptionListener& UpdatedListener);
 
 private:
 	static TMap<FNameEntryId, FAISenseID> BPSenseToSenseID;
diff --git a/Engine/Source/Runtime/AIModule/Classes/Perception/AISense_Damage.h b/Engine/Source/Runtime/AIModule/Classes/Perception/AISense_Damage.h
--- a/Engine/Source/Runtime/AIModule/Classes/Perception/AISense_Damage.h
+++ b/Engine/Source/Runtime/AIModule/Classes/Perception/AISense_Damage.h
@@ -11,7 +11,7 @@ class IAIPerceptionListenerInterface;
 class UAISenseEvent;
 
 USTRUCT(BlueprintType)
-struct AIMODULE_API FAIDamageEvent
+struct FAIDamageEvent
 {	
 	GENERATED_USTRUCT_BODY()
 
@@ -44,20 +44,20 @@ struct AIMODULE_API FAIDamageEvent
 	UPROPERTY(EditAnywhere, BlueprintReadWrite, Category = "Sense")
 	FName Tag;
 	
-	FAIDamageEvent();
-	FAIDamageEvent(AActor* InDamagedActor, AActor* InInstigator, float DamageAmount, const FVector& EventLocation, const FVector& InHitLocation = FAISystem::InvalidLocation, FName InTag = NAME_None);
-	void Compile();
+	AIMODULE_API FAIDamageEvent();
+	AIMODULE_API FAIDamageEvent(AActor* InDamagedActor, AActor* InInstigator, float DamageAmount, const FVector& EventLocation, const FVector& InHitLocation = FAISystem::InvalidLocation, FName InTag = NAME_None);
+	AIMODULE_API void Compile();
 
 	bool IsValid() const
 	{
 		return DamagedActor != nullptr;
 	}
 
-	IAIPerceptionListenerInterface* GetDamagedActorAsPerceptionListener() const;
+	AIMODULE_API IAIPerceptionListenerInterface* GetDamagedActorAsPerceptionListener() const;
 };
 
-UCLASS(ClassGroup=AI)
-class AIMODULE_API UAISense_Damage : public UAISense
+UCLASS(ClassGroup=AI, MinimalAPI)
+class UAISense_Damage : public UAISense
 {
 	GENERATED_UCLASS_BODY()
 
@@ -65,13 +65,13 @@ class AIMODULE_API UAISense_Damage : public UAISense
 	TArray<FAIDamageEvent> RegisteredEvents;
 
 public:		
-	void RegisterEvent(const FAIDamageEvent& Event);	
-	virtual void RegisterWrappedEvent(UAISenseEvent& PerceptionEvent) override;
+	AIMODULE_API void RegisterEvent(const FAIDamageEvent& Event);	
+	AIMODULE_API virtual void RegisterWrappedEvent(UAISenseEvent& PerceptionEvent) override;
 
 	/** EventLocation will be reported as Instigator's location at the moment of event happening*/
 	UFUNCTION(BlueprintCallable, Category = "AI|Perception", meta = (WorldContext="WorldContextObject", AdvancedDisplay="HitLocation"))
-	static void ReportDamageEvent(UObject* WorldContextObject, AActor* DamagedActor, AActor* Instigator, float DamageAmount, FVector EventLocation, FVector HitLocation, FName Tag = NAME_None);
+	static AIMODULE_API void ReportDamageEvent(UObject* WorldContextObject, AActor* DamagedActor, AActor* Instigator, float DamageAmount, FVector EventLocation, FVector HitLocation, FName Tag = NAME_None);
 
 protected:
-	virtual float Update() override;
+	AIMODULE_API virtual float Update() override;
 };
diff --git a/Engine/Source/Runtime/AIModule/Classes/Perception/AISense_Hearing.h b/Engine/Source/Runtime/AIModule/Classes/Perception/AISense_Hearing.h
--- a/Engine/Source/Runtime/AIModule/Classes/Perception/AISense_Hearing.h
+++ b/Engine/Source/Runtime/AIModule/Classes/Perception/AISense_Hearing.h
@@ -12,7 +12,7 @@ class UAISenseConfig_Hearing;
 class UAISenseEvent;
 
 USTRUCT(BlueprintType)
-struct AIMODULE_API FAINoiseEvent
+struct FAINoiseEvent
 {	
 	GENERATED_USTRUCT_BODY()
 
@@ -53,15 +53,15 @@ struct AIMODULE_API FAINoiseEvent
 
 	FGenericTeamId TeamIdentifier;
 		
-	FAINoiseEvent();
-	FAINoiseEvent(AActor* InInstigator, const FVector& InNoiseLocation, float InLoudness = 1.f, float InMaxRange = 0.f, FName Tag = NAME_None);
+	AIMODULE_API FAINoiseEvent();
+	AIMODULE_API FAINoiseEvent(AActor* InInstigator, const FVector& InNoiseLocation, float InLoudness = 1.f, float InMaxRange = 0.f, FName Tag = NAME_None);
 
 	/** Verifies and calculates derived data */
-	void Compile();
+	AIMODULE_API void Compile();
 };
 
-UCLASS(ClassGroup=AI, Config=Game)
-class AIMODULE_API UAISense_Hearing : public UAISense
+UCLASS(ClassGroup=AI, Config=Game, MinimalAPI)
+class UAISense_Hearing : public UAISense
 {
 	GENERATED_UCLASS_BODY()
 		
@@ -85,13 +85,13 @@ protected:
 	TMap<FPerceptionListenerID, FDigestedHearingProperties> DigestedProperties;
 
 public:	
-	void RegisterEvent(const FAINoiseEvent& Event);	
-	void RegisterEventsBatch(const TArray<FAINoiseEvent>& Events);
+	AIMODULE_API void RegisterEvent(const FAINoiseEvent& Event);	
+	AIMODULE_API void RegisterEventsBatch(const TArray<FAINoiseEvent>& Events);
 
-	virtual void PostInitProperties() override;
+	AIMODULE_API virtual void PostInitProperties() override;
 
 	// part of BP interface. Translates PerceptionEvent to FAINoiseEvent and call RegisterEvent(const FAINoiseEvent& Event)
-	virtual void RegisterWrappedEvent(UAISenseEvent& PerceptionEvent) override;
+	AIMODULE_API virtual void RegisterWrappedEvent(UAISenseEvent& PerceptionEvent) override;
 
 	/**
 	 * Report a noise event.
@@ -103,13 +103,13 @@ public:
 	 * @param Tag Identifier for the event.
 	 */
 	UFUNCTION(BlueprintCallable, Category = "AI|Perception", meta = (WorldContext="WorldContextObject"))
-	static void ReportNoiseEvent(UObject* WorldContextObject, FVector NoiseLocation, float Loudness = 1.f, AActor* Instigator = nullptr, float MaxRange = 0.f, FName Tag = NAME_None);
+	static AIMODULE_API void ReportNoiseEvent(UObject* WorldContextObject, FVector NoiseLocation, float Loudness = 1.f, AActor* Instigator = nullptr, float MaxRange = 0.f, FName Tag = NAME_None);
 
 protected:
-	virtual float Update() override;
-	virtual void RegisterMakeNoiseDelegate();
+	AIMODULE_API virtual float Update() override;
+	AIMODULE_API virtual void RegisterMakeNoiseDelegate();
 
-	void OnNewListenerImpl(const FPerceptionListener& NewListener);
-	void OnListenerUpdateImpl(const FPerceptionListener& UpdatedListener);
-	void OnListenerRemovedImpl(const FPerceptionListener& UpdatedListener);
+	AIMODULE_API void OnNewListenerImpl(const FPerceptionListener& NewListener);
+	AIMODULE_API void OnListenerUpdateImpl(const FPerceptionListener& UpdatedListener);
+	AIMODULE_API void OnListenerRemovedImpl(const FPerceptionListener& UpdatedListener);
 };
diff --git a/Engine/Source/Runtime/AIModule/Classes/Perception/AISense_Prediction.h b/Engine/Source/Runtime/AIModule/Classes/Perception/AISense_Prediction.h
--- a/Engine/Source/Runtime/AIModule/Classes/Perception/AISense_Prediction.h
+++ b/Engine/Source/Runtime/AIModule/Classes/Perception/AISense_Prediction.h
@@ -12,7 +12,7 @@ class APawn;
 class UAISense_Prediction;
 
 USTRUCT()
-struct AIMODULE_API FAIPredictionEvent
+struct FAIPredictionEvent
 {	
 	GENERATED_USTRUCT_BODY()
 
@@ -34,8 +34,8 @@ struct AIMODULE_API FAIPredictionEvent
 	}
 };
 
-UCLASS(ClassGroup=AI)
-class AIMODULE_API UAISense_Prediction : public UAISense
+UCLASS(ClassGroup=AI, MinimalAPI)
+class UAISense_Prediction : public UAISense
 {
 	GENERATED_UCLASS_BODY()
 
@@ -43,18 +43,18 @@ class AIMODULE_API UAISense_Prediction : public UAISense
 	TArray<FAIPredictionEvent> RegisteredEvents;
 
 public:		
-	void RegisterEvent(const FAIPredictionEvent& Event);	
+	AIMODULE_API void RegisterEvent(const FAIPredictionEvent& Event);	
 
 	/** Asks perception system to supply Requestor with PredictedActor's predicted location in PredictionTime seconds
 	 *	Location is being predicted based on PredicterActor's current location and velocity */
 	UFUNCTION(BlueprintCallable, Category = "AI|Perception")
-	static void RequestControllerPredictionEvent(AAIController* Requestor, AActor* PredictedActor, float PredictionTime);
+	static AIMODULE_API void RequestControllerPredictionEvent(AAIController* Requestor, AActor* PredictedActor, float PredictionTime);
 
 	/** Asks perception system to supply Requestor with PredictedActor's predicted location in PredictionTime seconds
 	 *	Location is being predicted based on PredicterActor's current location and velocity */
 	UFUNCTION(BlueprintCallable, Category = "AI|Perception")
-	static void RequestPawnPredictionEvent(APawn* Requestor, AActor* PredictedActor, float PredictionTime);
+	static AIMODULE_API void RequestPawnPredictionEvent(APawn* Requestor, AActor* PredictedActor, float PredictionTime);
 
 protected:
-	virtual float Update() override;
+	AIMODULE_API virtual float Update() override;
 };
diff --git a/Engine/Source/Runtime/AIModule/Classes/Perception/AISense_Sight.h b/Engine/Source/Runtime/AIModule/Classes/Perception/AISense_Sight.h
--- a/Engine/Source/Runtime/AIModule/Classes/Perception/AISense_Sight.h
+++ b/Engine/Source/Runtime/AIModule/Classes/Perception/AISense_Sight.h
@@ -25,7 +25,7 @@ namespace ESightPerceptionEventName
 }
 
 USTRUCT()
-struct AIMODULE_API FAISightEvent
+struct FAISightEvent
 {
 	GENERATED_USTRUCT_BODY()
 
@@ -48,17 +48,17 @@ struct AIMODULE_API FAISightEvent
 	}
 };
 
-struct AIMODULE_API FAISightTarget
+struct FAISightTarget
 {
 	typedef uint32 FTargetId;
-	static const FTargetId InvalidTargetId;
+	static AIMODULE_API const FTargetId InvalidTargetId;
 
 	TWeakObjectPtr<AActor> Target;
 	IAISightTargetInterface* SightTargetInterface;
 	FGenericTeamId TeamId;
 	FTargetId TargetId;
 
-	FAISightTarget(AActor* InTarget = NULL, FGenericTeamId InTeamId = FGenericTeamId::NoTeam);
+	AIMODULE_API FAISightTarget(AActor* InTarget = NULL, FGenericTeamId InTeamId = FGenericTeamId::NoTeam);
 
 	FORCEINLINE FVector GetLocationSimple() const
 	{
@@ -179,7 +179,7 @@ struct FAISightQuery
 	};
 };
 
-struct AIMODULE_API FAISightQueryID
+struct FAISightQueryID
 {
 	FPerceptionListenerID ObserverId;
 	FAISightTarget::FTargetId TargetId;
@@ -197,8 +197,8 @@ struct AIMODULE_API FAISightQueryID
 
 DECLARE_DELEGATE_FiveParams(FOnPendingVisibilityQueryProcessedDelegate, const FAISightQueryID&, const bool, const float, const FVector&, const TOptional<int32>&);
 
-UCLASS(ClassGroup=AI, config=Game)
-class AIMODULE_API UAISense_Sight : public UAISense
+UCLASS(ClassGroup=AI, config=Game, MinimalAPI)
+class UAISense_Sight : public UAISense
 {
 	GENERATED_UCLASS_BODY()
 
@@ -282,49 +282,49 @@ protected:
 
 public:
 
-	virtual void PostInitProperties() override;
+	AIMODULE_API virtual void PostInitProperties() override;
 	
-	void RegisterEvent(const FAISightEvent& Event);	
+	AIMODULE_API void RegisterEvent(const FAISightEvent& Event);	
 
-	virtual void RegisterSource(AActor& SourceActors) override;
-	virtual void UnregisterSource(AActor& SourceActor) override;
+	AIMODULE_API virtual void RegisterSource(AActor& SourceActors) override;
+	AIMODULE_API virtual void UnregisterSource(AActor& SourceActor) override;
 	
-	virtual void OnListenerForgetsActor(const FPerceptionListener& Listener, AActor& ActorToForget) override;
-	virtual void OnListenerForgetsAll(const FPerceptionListener& Listener) override;
+	AIMODULE_API virtual void OnListenerForgetsActor(const FPerceptionListener& Listener, AActor& ActorToForget) override;
+	AIMODULE_API virtual void OnListenerForgetsAll(const FPerceptionListener& Listener) override;
 
 #if WITH_GAMEPLAY_DEBUGGER_MENU
-	virtual void DescribeSelfToGameplayDebugger(const UAIPerceptionSystem& PerceptionSystem, FGameplayDebuggerCategory& DebuggerCategory) const override;
+	AIMODULE_API virtual void DescribeSelfToGameplayDebugger(const UAIPerceptionSystem& PerceptionSystem, FGameplayDebuggerCategory& DebuggerCategory) const override;
 #endif // WITH_GAMEPLAY_DEBUGGER_MENU
 
 protected:
-	virtual float Update() override;
+	AIMODULE_API virtual float Update() override;
 
-	EVisibilityResult ComputeVisibility(UWorld* World, FAISightQuery& SightQuery, FPerceptionListener& Listener, const AActor* ListenerActor, FAISightTarget& Target, AActor* TargetActor, const FDigestedSightProperties& PropDigest, float& OutStimulusStrength, FVector& OutSeenLocation, int32& OutNumberOfLoSChecksPerformed, int32& OutNumberOfAsyncLosCheckRequested) const;
-	virtual bool ShouldAutomaticallySeeTarget(const FDigestedSightProperties& PropDigest, FAISightQuery* SightQuery, FPerceptionListener& Listener, AActor* TargetActor, float& OutStimulusStrength) const;
+	AIMODULE_API EVisibilityResult ComputeVisibility(UWorld* World, FAISightQuery& SightQuery, FPerceptionListener& Listener, const AActor* ListenerActor, FAISightTarget& Target, AActor* TargetActor, const FDigestedSightProperties& PropDigest, float& OutStimulusStrength, FVector& OutSeenLocation, int32& OutNumberOfLoSChecksPerformed, int32& OutNumberOfAsyncLosCheckRequested) const;
+	AIMODULE_API virtual bool ShouldAutomaticallySeeTarget(const FDigestedSightProperties& PropDigest, FAISightQuery* SightQuery, FPerceptionListener& Listener, AActor* TargetActor, float& OutStimulusStrength) const;
 	
 	UE_DEPRECATED(5.3, "Please use the UpdateQueryVisibilityStatus version which takes an Actor& instead.")
-	void UpdateQueryVisibilityStatus(FAISightQuery& SightQuery, FPerceptionListener& Listener, const bool bIsVisible, const FVector& SeenLocation, const float StimulusStrength, AActor* TargetActor, const FVector& TargetLocation) const;
-	void UpdateQueryVisibilityStatus(FAISightQuery& SightQuery, FPerceptionListener& Listener, const bool bIsVisible, const FVector& SeenLocation, const float StimulusStrength, AActor& TargetActor, const FVector& TargetLocation) const;
+	AIMODULE_API void UpdateQueryVisibilityStatus(FAISightQuery& SightQuery, FPerceptionListener& Listener, const bool bIsVisible, const FVector& SeenLocation, const float StimulusStrength, AActor* TargetActor, const FVector& TargetLocation) const;
+	AIMODULE_API void UpdateQueryVisibilityStatus(FAISightQuery& SightQuery, FPerceptionListener& Listener, const bool bIsVisible, const FVector& SeenLocation, const float StimulusStrength, AActor& TargetActor, const FVector& TargetLocation) const;
 
-	void OnPendingCanBeSeenQueryProcessed(const FAISightQueryID& QueryID, const bool bIsVisible, const float StimulusStrength, const FVector& SeenLocation, const TOptional<int32>& UserData);
-	void OnPendingTraceQueryProcessed(const FTraceHandle& TraceHandle, FTraceDatum& TraceDatum);
-	void OnPendingQueryProcessed(const int32 SightQueryIndex, const bool bIsVisible, const float StimulusStrength, const FVector& SeenLocation, const TOptional<int32>& UserData, const TOptional<AActor*> InTargetActor = NullOpt);
+	AIMODULE_API void OnPendingCanBeSeenQueryProcessed(const FAISightQueryID& QueryID, const bool bIsVisible, const float StimulusStrength, const FVector& SeenLocation, const TOptional<int32>& UserData);
+	AIMODULE_API void OnPendingTraceQueryProcessed(const FTraceHandle& TraceHandle, FTraceDatum& TraceDatum);
+	AIMODULE_API void OnPendingQueryProcessed(const int32 SightQueryIndex, const bool bIsVisible, const float StimulusStrength, const FVector& SeenLocation, const TOptional<int32>& UserData, const TOptional<AActor*> InTargetActor = NullOpt);
 
-	void OnNewListenerImpl(const FPerceptionListener& NewListener);
-	void OnListenerUpdateImpl(const FPerceptionListener& UpdatedListener);
-	void OnListenerRemovedImpl(const FPerceptionListener& RemovedListener);
-	virtual void OnListenerConfigUpdated(const FPerceptionListener& UpdatedListener) override;
+	AIMODULE_API void OnNewListenerImpl(const FPerceptionListener& NewListener);
+	AIMODULE_API void OnListenerUpdateImpl(const FPerceptionListener& UpdatedListener);
+	AIMODULE_API void OnListenerRemovedImpl(const FPerceptionListener& RemovedListener);
+	AIMODULE_API virtual void OnListenerConfigUpdated(const FPerceptionListener& UpdatedListener) override;
 	
-	void GenerateQueriesForListener(const FPerceptionListener& Listener, const FDigestedSightProperties& PropertyDigest, const TFunction<void(FAISightQuery&)>& OnAddedFunc = nullptr);
+	AIMODULE_API void GenerateQueriesForListener(const FPerceptionListener& Listener, const FDigestedSightProperties& PropertyDigest, const TFunction<void(FAISightQuery&)>& OnAddedFunc = nullptr);
 
-	void RemoveAllQueriesByListener(const FPerceptionListener& Listener, const TFunction<void(const FAISightQuery&)>& OnRemoveFunc = nullptr);
-	void RemoveAllQueriesToTarget(const FAISightTarget::FTargetId& TargetId, const TFunction<void(const FAISightQuery&)>& OnRemoveFunc = nullptr);
+	AIMODULE_API void RemoveAllQueriesByListener(const FPerceptionListener& Listener, const TFunction<void(const FAISightQuery&)>& OnRemoveFunc = nullptr);
+	AIMODULE_API void RemoveAllQueriesToTarget(const FAISightTarget::FTargetId& TargetId, const TFunction<void(const FAISightQuery&)>& OnRemoveFunc = nullptr);
 
 	/** returns information whether new LoS queries have been added */
-	bool RegisterTarget(AActor& TargetActor, const TFunction<void(FAISightQuery&)>& OnAddedFunc = nullptr);
+	AIMODULE_API bool RegisterTarget(AActor& TargetActor, const TFunction<void(FAISightQuery&)>& OnAddedFunc = nullptr);
 
-	float CalcQueryImportance(const FPerceptionListener& Listener, const FVector& TargetLocation, const float SightRadiusSq) const;
-	bool RegisterNewQuery(const FPerceptionListener& Listener, const IGenericTeamAgentInterface* ListenersTeamAgent, const AActor& TargetActor, const FAISightTarget::FTargetId& TargetId, const FVector& TargetLocation, const FDigestedSightProperties& PropDigest, const TFunction<void(FAISightQuery&)>& OnAddedFunc);
+	AIMODULE_API float CalcQueryImportance(const FPerceptionListener& Listener, const FVector& TargetLocation, const float SightRadiusSq) const;
+	AIMODULE_API bool RegisterNewQuery(const FPerceptionListener& Listener, const IGenericTeamAgentInterface* ListenersTeamAgent, const AActor& TargetActor, const FAISightTarget::FTargetId& TargetId, const FVector& TargetLocation, const FDigestedSightProperties& PropDigest, const TFunction<void(FAISightQuery&)>& OnAddedFunc);
 
 protected:
 	enum FQueriesOperationPostProcess
diff --git a/Engine/Source/Runtime/AIModule/Classes/Perception/AISense_Team.h b/Engine/Source/Runtime/AIModule/Classes/Perception/AISense_Team.h
--- a/Engine/Source/Runtime/AIModule/Classes/Perception/AISense_Team.h
+++ b/Engine/Source/Runtime/AIModule/Classes/Perception/AISense_Team.h
@@ -11,7 +11,7 @@
 class UAISense_Team;
 
 USTRUCT()
-struct AIMODULE_API FAITeamStimulusEvent
+struct FAITeamStimulusEvent
 {	
 	GENERATED_USTRUCT_BODY()
 
@@ -33,7 +33,7 @@ public:
 	TObjectPtr<AActor> Enemy;
 		
 	FAITeamStimulusEvent() : Broadcaster(nullptr), Enemy(nullptr) {}
-	FAITeamStimulusEvent(AActor* InBroadcaster, AActor* InEnemy, const FVector& InLastKnowLocation, float EventRange, float PassedInfoAge = 0.f, float InStrength = 1.f);
+	AIMODULE_API FAITeamStimulusEvent(AActor* InBroadcaster, AActor* InEnemy, const FVector& InLastKnowLocation, float EventRange, float PassedInfoAge = 0.f, float InStrength = 1.f);
 
 	FORCEINLINE void CacheBroadcastLocation()
 	{
@@ -46,8 +46,8 @@ public:
 	}
 };
 
-UCLASS(ClassGroup=AI)
-class AIMODULE_API UAISense_Team : public UAISense
+UCLASS(ClassGroup=AI, MinimalAPI)
+class UAISense_Team : public UAISense
 {
 	GENERATED_UCLASS_BODY()
 
@@ -55,8 +55,8 @@ class AIMODULE_API UAISense_Team : public UAISense
 	TArray<FAITeamStimulusEvent> RegisteredEvents;
 
 public:		
-	void RegisterEvent(const FAITeamStimulusEvent& Event);	
+	AIMODULE_API void RegisterEvent(const FAITeamStimulusEvent& Event);	
 
 protected:
-	virtual float Update() override;
+	AIMODULE_API virtual float Update() override;
 };
diff --git a/Engine/Source/Runtime/AIModule/Classes/Perception/AISense_Touch.h b/Engine/Source/Runtime/AIModule/Classes/Perception/AISense_Touch.h
--- a/Engine/Source/Runtime/AIModule/Classes/Perception/AISense_Touch.h
+++ b/Engine/Source/Runtime/AIModule/Classes/Perception/AISense_Touch.h
@@ -11,7 +11,7 @@ class IAIPerceptionListenerInterface;
 class UAISense_Touch;
 
 USTRUCT()
-struct AIMODULE_API FAITouchEvent
+struct FAITouchEvent
 {	
 	GENERATED_USTRUCT_BODY()
 
@@ -32,11 +32,11 @@ struct AIMODULE_API FAITouchEvent
 	{
 	}
 
-	IAIPerceptionListenerInterface* GetTouchedActorAsPerceptionListener() const;
+	AIMODULE_API IAIPerceptionListenerInterface* GetTouchedActorAsPerceptionListener() const;
 };
 
-UCLASS(ClassGroup=AI)
-class AIMODULE_API UAISense_Touch : public UAISense
+UCLASS(ClassGroup=AI, MinimalAPI)
+class UAISense_Touch : public UAISense
 {
 	GENERATED_UCLASS_BODY()
 
@@ -44,11 +44,11 @@ class AIMODULE_API UAISense_Touch : public UAISense
 	TArray<FAITouchEvent> RegisteredEvents;
 
 public:		
-	void RegisterEvent(const FAITouchEvent& Event);	
+	AIMODULE_API void RegisterEvent(const FAITouchEvent& Event);	
 
 	UFUNCTION(BlueprintCallable, Category = "AI|Perception", meta = (WorldContext = "WorldContextObject"))
-	static void ReportTouchEvent(UObject* WorldContextObject, AActor* TouchReceiver, AActor* OtherActor, FVector Location);
+	static AIMODULE_API void ReportTouchEvent(UObject* WorldContextObject, AActor* TouchReceiver, AActor* OtherActor, FVector Location);
 
 protected:
-	virtual float Update() override;
+	AIMODULE_API virtual float Update() override;
 };
