AI: Logging: Downgrade log message to verbose to make cookcommandlet logs easier to read.
#rnx

diff --git a/Engine/Source/Runtime/AIModule/Private/AIModule.cpp b/Engine/Source/Runtime/AIModule/Private/AIModule.cpp
--- a/Engine/Source/Runtime/AIModule/Private/AIModule.cpp
+++ b/Engine/Source/Runtime/AIModule/Private/AIModule.cpp
@@ -112,7 +112,7 @@ void FAIModule::ShutdownModule()
 
 UAISystemBase* FAIModule::CreateAISystemInstance(UWorld* World)
 {
-	UE_LOG(LogAIModule, Log, TEXT("Creating AISystem for world %s"), *GetNameSafe(World));
+	UE_LOG(LogAIModule, Verbose, TEXT("Creating AISystem for world %s"), *GetNameSafe(World));
 	
 	FSoftClassPath AISystemClassName = UAISystemBase::GetAISystemClassName();
 
