Fix invalid memory access (use-after-free) in FGraphAStar::FindPath.

diff --git a/Engine/Source/Runtime/AIModule/Public/GraphAStar.h b/Engine/Source/Runtime/AIModule/Public/GraphAStar.h
--- a/Engine/Source/Runtime/AIModule/Public/GraphAStar.h
+++ b/Engine/Source/Runtime/AIModule/Public/GraphAStar.h
@@ -526,6 +526,10 @@ struct FGraphAStar
 		int32 BestNodeIndex = StartPoolNode.SearchNodeIndex;
 		FVector::FReal BestNodeCost = StartPoolNode.TotalCost;
 
+		const int32 StartPoolSearchNodeIndex = StartPoolNode.SearchNodeIndex;
+		const FGraphNodeRef StartPoolNodeRef = StartPoolNode.NodeRef;
+		// Don't use the StartPoolNode reference beyond this point since it might not be valid after ProcessSingleNode().
+
 		EGraphAStarResult Result = EGraphAStarResult::SearchSuccess;
 		const bool bIsBound = true;
 		
@@ -546,12 +550,12 @@ struct FGraphAStar
 		{
 			// store the path. Note that it will be reversed!
 			int32 SearchNodeIndex = BestNodeIndex;
-			int32 PathLength = ShouldIncludeStartNodeInPath(Filter) && BestNodeIndex != StartPoolNode.SearchNodeIndex ? 1 : 0;
+			int32 PathLength = ShouldIncludeStartNodeInPath(Filter) && BestNodeIndex != StartPoolSearchNodeIndex ? 1 : 0;
 			do 
 			{
 				PathLength++;
 				SearchNodeIndex = NodePool[SearchNodeIndex].ParentNodeIndex;
-			} while (NodePool.IsValidIndex(SearchNodeIndex) && NodePool[SearchNodeIndex].NodeRef != StartPoolNode.NodeRef && ensure(PathLength < Policy::FatalPathLength));
+			} while (NodePool.IsValidIndex(SearchNodeIndex) && NodePool[SearchNodeIndex].NodeRef != StartPoolNodeRef && ensure(PathLength < Policy::FatalPathLength));
 			
 			if (PathLength >= Policy::FatalPathLength)
 			{
