GraphAStar code massage to remove the possibility of shooting oneself's foot by trying to use StartPoolNode after it potentially got deallocated (with NodePool's reallocation).

diff --git a/Engine/Source/Runtime/AIModule/Public/GraphAStar.h b/Engine/Source/Runtime/AIModule/Public/GraphAStar.h
--- a/Engine/Source/Runtime/AIModule/Public/GraphAStar.h
+++ b/Engine/Source/Runtime/AIModule/Public/GraphAStar.h
@@ -517,18 +517,22 @@ struct FGraphAStar
 		OpenList.Reset();
 
 		// kick off the search with the first node
-		FSearchNode& StartPoolNode = NodePool.Add(StartNode);
-		StartPoolNode.TraversalCost = 0;
-		StartPoolNode.TotalCost = GetHeuristicCost(Filter, StartNode, EndNode) * Filter.GetHeuristicScale();
+		int32 BestNodeIndex = INDEX_NONE;
+		FVector::FReal BestNodeCost = TNumericLimits<FVector::FReal>::Max();
+		{
+			// scoping StartPoolNode to make it clear it's not safe to use after ProcessSingleNode due to potential NodePool reallocation
+			FSearchNode& StartPoolNode = NodePool.Add(StartNode);
+			StartPoolNode.TraversalCost = 0;
+			StartPoolNode.TotalCost = GetHeuristicCost(Filter, StartNode, EndNode) * Filter.GetHeuristicScale();
 
-		OpenList.Push(StartPoolNode);
+			OpenList.Push(StartPoolNode);
 
-		int32 BestNodeIndex = StartPoolNode.SearchNodeIndex;
-		FVector::FReal BestNodeCost = StartPoolNode.TotalCost;
+			BestNodeIndex = StartPoolNode.SearchNodeIndex;
+			BestNodeCost = StartPoolNode.TotalCost;
+		}
 
-		const int32 StartPoolSearchNodeIndex = StartPoolNode.SearchNodeIndex;
-		const FGraphNodeRef StartPoolNodeRef = StartPoolNode.NodeRef;
-		// Don't use the StartPoolNode reference beyond this point since it might not be valid after ProcessSingleNode().
+		const int32 StartPoolSearchNodeIndex = NodePool[BestNodeIndex].SearchNodeIndex;
+		const FGraphNodeRef StartPoolNodeRef = NodePool[BestNodeIndex].NodeRef;
 
 		EGraphAStarResult Result = EGraphAStarResult::SearchSuccess;
 		const bool bIsBound = true;
@@ -593,14 +597,19 @@ struct FGraphAStar
 		OpenList.Reset();
 
 		// kick off the search with the first node
-		FSearchNode& StartPoolNode = NodePool.Add(StartNode);
-		StartPoolNode.TraversalCost = 0;
-		StartPoolNode.TotalCost = 0;
+		int32 BestNodeIndex = INDEX_NONE;
+		FVector::FReal BestNodeCost = TNumericLimits<FVector::FReal>::Max();
+		{
+			// scoping StartPoolNode to make it clear it's not safe to use after ProcessSingleNode due to potential NodePool reallocation
+			FSearchNode& StartPoolNode = NodePool.Add(StartNode);
+			StartPoolNode.TraversalCost = 0;
+			StartPoolNode.TotalCost = 0;
 
-		OpenList.Push(StartPoolNode);
+			OpenList.Push(StartPoolNode);
 
-		int32 BestNodeIndex = StartPoolNode.SearchNodeIndex;
-		FVector::FReal BestNodeCost = StartPoolNode.TotalCost;
+			BestNodeIndex = StartPoolNode.SearchNodeIndex;
+			BestNodeCost = StartPoolNode.TotalCost;
+		}
 		
 		const FSearchNode FakeEndNode = StartNode;
 		const bool bIsBound = false;
