Fix AI Sight Sense MT detector re-entrance : add internal call that requires to get access to the MT detector prior to the method call

#test forced code to repro the bug and confirm the fix in a gym with bots

diff --git a/Engine/Source/Runtime/AIModule/Classes/Perception/AISense_Sight.h b/Engine/Source/Runtime/AIModule/Classes/Perception/AISense_Sight.h
--- a/Engine/Source/Runtime/AIModule/Classes/Perception/AISense_Sight.h
+++ b/Engine/Source/Runtime/AIModule/Classes/Perception/AISense_Sight.h
@@ -319,6 +319,8 @@ protected:
 
 	AIMODULE_API void RemoveAllQueriesByListener(const FPerceptionListener& Listener, const TFunction<void(const FAISightQuery&)>& OnRemoveFunc = nullptr);
 	AIMODULE_API void RemoveAllQueriesToTarget(const FAISightTarget::FTargetId& TargetId, const TFunction<void(const FAISightQuery&)>& OnRemoveFunc = nullptr);
+	/** RemoveAllQueriesToTarget version that need to already have a write access on QueriesListAccessDetector*/
+	void RemoveAllQueriesToTarget_Internal(const FAISightTarget::FTargetId& TargetId, const TFunction<void(const FAISightQuery&)>& OnRemoveFunc = nullptr);
 
 	/** returns information whether new LoS queries have been added */
 	AIMODULE_API bool RegisterTarget(AActor& TargetActor, const TFunction<void(FAISightQuery&)>& OnAddedFunc = nullptr);
diff --git a/Engine/Source/Runtime/AIModule/Private/Perception/AISense_Sight.cpp b/Engine/Source/Runtime/AIModule/Private/Perception/AISense_Sight.cpp
--- a/Engine/Source/Runtime/AIModule/Private/Perception/AISense_Sight.cpp
+++ b/Engine/Source/Runtime/AIModule/Private/Perception/AISense_Sight.cpp
@@ -495,7 +495,7 @@ float UAISense_Sight::Update()
 			for (const auto& TargetId : InvalidTargets)
 			{
 				// remove affected queries
-				RemoveAllQueriesToTarget(TargetId);
+				RemoveAllQueriesToTarget_Internal(TargetId);
 				// remove target itself
 				ObservedTargets.Remove(TargetId);
 			}
@@ -1036,8 +1036,13 @@ void UAISense_Sight::RemoveAllQueriesByListener(const FPerceptionListener& Liste
 
 void UAISense_Sight::RemoveAllQueriesToTarget(const FAISightTarget::FTargetId& TargetId, const TFunction<void(const FAISightQuery&)>& OnRemoveFunc/*= nullptr */)
 {
-	SCOPE_CYCLE_COUNTER(STAT_AI_Sense_Sight_RemoveToTarget);
 	UE_MT_SCOPED_WRITE_ACCESS(QueriesListAccessDetector);
+	RemoveAllQueriesToTarget_Internal(TargetId, OnRemoveFunc);
+}
+
+void UAISense_Sight::RemoveAllQueriesToTarget_Internal(const FAISightTarget::FTargetId& TargetId, const TFunction<void(const FAISightQuery&)>& OnRemoveFunc/*= nullptr */)
+{
+	SCOPE_CYCLE_COUNTER(STAT_AI_Sense_Sight_RemoveToTarget);
 
 	auto RemoveQuery = [&TargetId, &OnRemoveFunc](TArray<FAISightQuery>& SightQueries, const int32 QueryIndex)->EReverseForEachResult
 	{
