Trivial change to IndexedHandle.h to laverage some parameters being const at compile time

diff --git a/Engine/Source/Runtime/AIModule/Public/IndexedHandle.h b/Engine/Source/Runtime/AIModule/Public/IndexedHandle.h
--- a/Engine/Source/Runtime/AIModule/Public/IndexedHandle.h
+++ b/Engine/Source/Runtime/AIModule/Public/IndexedHandle.h
@@ -230,7 +230,7 @@ public:
 
 		if (ensureMsgf(bIsValid, TEXT("Trying to Remove Invalid Handle with Index %d, Serial Number %d"), IndexedHandle.GetIndex(), IndexedHandle.GetSerialNumber()))
 		{
-			if(bOptimizeHandleReuse)
+			if constexpr (bOptimizeHandleReuse)
 			{
 			    const TIndexType Idx = IndexedHandle.GetIndex();
 			    const int32 InsertPosition = Algo::LowerBound(FreeHandleIndices, Idx, [](const TIndexType& Lhs, const TIndexType& Rhs) {return Lhs > Rhs;});
