Fix BT cooldown decorator initial last use time to be largest negative number.
Changed cooldown calculation to prevent potential undefined behavior arithmetic and avoid additional branches.

diff --git a/Engine/Source/Runtime/AIModule/Private/BehaviorTree/Decorators/BTDecorator_Cooldown.cpp b/Engine/Source/Runtime/AIModule/Private/BehaviorTree/Decorators/BTDecorator_Cooldown.cpp
--- a/Engine/Source/Runtime/AIModule/Private/BehaviorTree/Decorators/BTDecorator_Cooldown.cpp
+++ b/Engine/Source/Runtime/AIModule/Private/BehaviorTree/Decorators/BTDecorator_Cooldown.cpp
@@ -24,8 +24,8 @@ void UBTDecorator_Cooldown::PostLoad()
 bool UBTDecorator_Cooldown::CalculateRawConditionValue(UBehaviorTreeComponent& OwnerComp, uint8* NodeMemory) const
 {
 	FBTCooldownDecoratorMemory* DecoratorMemory = CastInstanceNodeMemory<FBTCooldownDecoratorMemory>(NodeMemory);
-	const double TimePassed = (OwnerComp.GetWorld()->GetTimeSeconds() - DecoratorMemory->LastUseTimestamp);
-	return TimePassed >= CoolDownTime;
+	const double RecalcTime = (OwnerComp.GetWorld()->GetTimeSeconds() - CoolDownTime);
+	return RecalcTime >= DecoratorMemory->LastUseTimestamp;
 }
 
 void UBTDecorator_Cooldown::InitializeMemory(UBehaviorTreeComponent& OwnerComp, uint8* NodeMemory, EBTMemoryInit::Type InitType) const
@@ -33,7 +33,7 @@ void UBTDecorator_Cooldown::InitializeMemory(UBehaviorTreeComponent& OwnerComp,
 	FBTCooldownDecoratorMemory* DecoratorMemory = CastInstanceNodeMemory<FBTCooldownDecoratorMemory>(NodeMemory);
 	if (InitType == EBTMemoryInit::Initialize)
 	{
-		DecoratorMemory->LastUseTimestamp = DBL_MIN;
+		DecoratorMemory->LastUseTimestamp = TNumericLimits<double>::Lowest();
 	}
 
 	DecoratorMemory->bRequestedRestart = false;
@@ -51,8 +51,8 @@ void UBTDecorator_Cooldown::TickNode(UBehaviorTreeComponent& OwnerComp, uint8* N
 	FBTCooldownDecoratorMemory* DecoratorMemory = CastInstanceNodeMemory<FBTCooldownDecoratorMemory>(NodeMemory);
 	if (!DecoratorMemory->bRequestedRestart)
 	{
-		const double TimePassed = (OwnerComp.GetWorld()->GetTimeSeconds() - DecoratorMemory->LastUseTimestamp);
-		if (TimePassed >= CoolDownTime)
+		const double RecalcTime = (OwnerComp.GetWorld()->GetTimeSeconds() - CoolDownTime);
+		if (RecalcTime >= DecoratorMemory->LastUseTimestamp)
 		{
 			DecoratorMemory->bRequestedRestart = true;
 			OwnerComp.RequestExecution(this);
