AIPerceptionComponent :
- Add templated version of GetSenseConfig()
- Fix missing const ref
- Fix typos in comments
#rnx

diff --git a/Engine/Source/Runtime/AIModule/Classes/Perception/AIPerceptionComponent.h b/Engine/Source/Runtime/AIModule/Classes/Perception/AIPerceptionComponent.h
--- a/Engine/Source/Runtime/AIModule/Classes/Perception/AIPerceptionComponent.h
+++ b/Engine/Source/Runtime/AIModule/Classes/Perception/AIPerceptionComponent.h
@@ -52,22 +52,22 @@ struct FActorPerceptionInfo
 
 	TArray<FAIStimulus> LastSensedStimuli;
 
-	/** if != MAX indicates the sense that takes precedense over other senses when it comes
+	/** if != MAX indicates the sense that takes precedence over other senses when it comes
 		to determining last stimulus location */
 	FAISenseID DominantSense;
 
 	/** indicates whether this Actor is hostile to perception holder */
 	uint32 bIsHostile : 1;
-	
-	FActorPerceptionInfo(AActor* InTarget = NULL)
-		: Target(InTarget), DominantSense(FAISenseID::InvalidID())
+
+	explicit FActorPerceptionInfo(AActor* InTarget = nullptr)
+		: Target(InTarget), DominantSense(FAISenseID::InvalidID()), bIsHostile(false)
 	{
 		LastSensedStimuli.AddDefaulted(FAISenseID::GetSize());
 	}
 
 	/** Retrieves last known location. Active (last reported as "successful")
 	 *	stimuli are preferred. */
-	FVector GetLastStimulusLocation(float* OptionalAge = NULL) const 
+	FVector GetLastStimulusLocation(float* OptionalAge = nullptr) const 
 	{
 		FVector Location(FAISystem::InvalidLocation);
 		float BestAge = FLT_MAX;
@@ -136,7 +136,7 @@ struct FActorPerceptionInfo
 	*
 	* @return Location of the last sensed stimuli or FAISystem::InvalidLocation if given sense has never registered related Target actor or if last stimuli has expired.
 	*/
-	FORCEINLINE FVector GetStimulusLocation(FAISenseID Sense) const
+	FORCEINLINE FVector GetStimulusLocation(const FAISenseID Sense) const
 	{
 		return LastSensedStimuli.IsValidIndex(Sense) && (LastSensedStimuli[Sense].IsValid() && (LastSensedStimuli[Sense].IsExpired() == false)) ? LastSensedStimuli[Sense].StimulusLocation : FAISystem::InvalidLocation;
 	}
@@ -146,7 +146,7 @@ struct FActorPerceptionInfo
 	*
 	* @return Location of the receiver for the last sensed stimuli or FAISystem::InvalidLocation if given sense has never registered related Target actor or last stimuli has expired.
 	*/
-	FORCEINLINE FVector GetReceiverLocation(FAISenseID Sense) const
+	FORCEINLINE FVector GetReceiverLocation(const FAISenseID Sense) const
 	{
 		return LastSensedStimuli.IsValidIndex(Sense) && (LastSensedStimuli[Sense].IsValid() && (LastSensedStimuli[Sense].IsExpired() == false)) ? LastSensedStimuli[Sense].ReceiverLocation : FAISystem::InvalidLocation;
 	}
@@ -156,7 +156,7 @@ struct FActorPerceptionInfo
 	*
 	* @return True if a target has been registered (even if not currently sensed) for the given sense and the stimuli is not expired.
 	*/
-	FORCEINLINE bool HasKnownStimulusOfSense(FAISenseID Sense) const
+	FORCEINLINE bool HasKnownStimulusOfSense(const FAISenseID Sense) const
 	{
 		return LastSensedStimuli.IsValidIndex(Sense) && (LastSensedStimuli[Sense].IsValid() && (LastSensedStimuli[Sense].IsExpired() == false));
 	}
@@ -166,7 +166,7 @@ struct FActorPerceptionInfo
 	*
 	* @return True if a target is still sensed for the given sense and the stimuli is not expired.
 	*/
-	FORCEINLINE bool IsSenseActive(FAISenseID Sense) const
+	FORCEINLINE bool IsSenseActive(const FAISenseID Sense) const
 	{
 		return LastSensedStimuli.IsValidIndex(Sense) && LastSensedStimuli[Sense].IsActive();
 	}
@@ -189,7 +189,7 @@ struct FActorPerceptionBlueprintInfo
 	UPROPERTY(BlueprintReadWrite, Category = "AI|Perception")
 	uint32 bIsHostile : 1;
 
-	FActorPerceptionBlueprintInfo() : Target(NULL), bIsHostile(false)
+	FActorPerceptionBlueprintInfo() : Target(nullptr), bIsHostile(false)
 	{}
 	FActorPerceptionBlueprintInfo(const FActorPerceptionInfo& Info);
 };
@@ -269,7 +269,7 @@ public:
 	AIMODULE_API const AActor* GetBodyActor() const;
 	AIMODULE_API AActor* GetMutableBodyActor();
 
-	FORCEINLINE const FPerceptionChannelAllowList GetPerceptionFilter() const { return PerceptionFilter; }
+	FORCEINLINE FPerceptionChannelAllowList GetPerceptionFilter() const { return PerceptionFilter; }
 
 	AIMODULE_API FGenericTeamId GetTeamIdentifier() const;
 	FORCEINLINE FPerceptionListenerID GetListenerId() const { return PerceptionListenerId; }
@@ -286,7 +286,7 @@ public:
 	/**	Retrieves all actors in PerceptualData matching the predicate.
 	 *	@return whether dead data (invalid actors) have been found while iterating over PerceptualData
 	 */
-	AIMODULE_API bool GetFilteredActors(TFunctionRef<bool(const FActorPerceptionInfo&)> Predicate, TArray<AActor*>& OutActors) const;
+	AIMODULE_API bool GetFilteredActors(const TFunctionRef<bool(const FActorPerceptionInfo&)>& Predicate, TArray<AActor*>& OutActors) const;
 
 	// @note Will stop on first age 0 stimulus
 	AIMODULE_API const FActorPerceptionInfo* GetFreshestTrace(const FAISenseID Sense) const;
@@ -296,6 +296,20 @@ public:
 	FORCEINLINE TSubclassOf<UAISense> GetDominantSense() const { return DominantSense; }
 	AIMODULE_API UAISenseConfig* GetSenseConfig(const FAISenseID& SenseID);
 	AIMODULE_API const UAISenseConfig* GetSenseConfig(const FAISenseID& SenseID) const;
+
+	template<typename T, typename = std::enable_if_t<std::is_base_of_v<UAISenseConfig, T>>>
+	T* GetSenseConfig() const
+	{
+		for (UAISenseConfig* SenseConfig : SensesConfig)
+		{
+			if (T* SenseConfigType = Cast<T>(SenseConfig))
+			{
+				return SenseConfigType;
+			}
+		}
+		return nullptr;
+	}
+	
 	AIMODULE_API void ConfigureSense(UAISenseConfig& SenseConfig);
 
 	typedef TArray<UAISenseConfig*>::TConstIterator TAISenseConfigConstIterator;
@@ -328,7 +342,7 @@ public:
 	AIMODULE_API float GetYoungestStimulusAge(const AActor& Source) const;
 	AIMODULE_API bool HasAnyActiveStimulus(const AActor& Source) const;
 	AIMODULE_API bool HasAnyCurrentStimulus(const AActor& Source) const;
-	AIMODULE_API bool HasActiveStimulus(const AActor& Source, FAISenseID Sense) const;
+	AIMODULE_API bool HasActiveStimulus(const AActor& Source, const FAISenseID Sense) const;
 
 #if WITH_GAMEPLAY_DEBUGGER_MENU
 	AIMODULE_API virtual void DescribeSelfToGameplayDebugger(FGameplayDebuggerCategory* DebuggerCategory) const;
@@ -383,7 +397,7 @@ public:
 
 	/**
 	 * Notifies all bound objects that perception info has been updated for a given target.
-	 * The notification is broadcasted for any received stimulus or on change of state
+	 * The notification is broadcast for any received stimulus or on change of state
 	 * according to the stimulus configuration.
 	 * 
 	 * Note - This delegate will not be called if source actor is no longer valid 
@@ -398,7 +412,7 @@ public:
 
 	/**
 	 * Notifies all bound objects that perception info has been updated for a given target.
-	 * The notification is broadcasted for any received stimulus or on change of state
+	 * The notification is broadcast for any received stimulus or on change of state
 	 * according to the stimulus configuration.
 	 *
 	 * Note - This delegate will be called even if source actor is no longer valid 
@@ -427,8 +441,8 @@ protected:
 private:
 	friend UAIPerceptionSystem;
 
-	AIMODULE_API void RegisterSenseConfig(UAISenseConfig& SenseConfig, UAIPerceptionSystem& AIPerceptionSys);
-	void StoreListenerId(FPerceptionListenerID InListenerId) { PerceptionListenerId = InListenerId; }
-	AIMODULE_API void SetMaxStimulusAge(FAISenseID SenseId, float MaxAge);
+	AIMODULE_API void RegisterSenseConfig(const UAISenseConfig& SenseConfig, UAIPerceptionSystem& AIPerceptionSys);
+	void StoreListenerId(const FPerceptionListenerID InListenerId) { PerceptionListenerId = InListenerId; }
+	AIMODULE_API void SetMaxStimulusAge(const FAISenseID SenseId, float MaxAge);
 };
 
diff --git a/Engine/Source/Runtime/AIModule/Private/Perception/AIPerceptionComponent.cpp b/Engine/Source/Runtime/AIModule/Private/Perception/AIPerceptionComponent.cpp
--- a/Engine/Source/Runtime/AIModule/Private/Perception/AIPerceptionComponent.cpp
+++ b/Engine/Source/Runtime/AIModule/Private/Perception/AIPerceptionComponent.cpp
@@ -73,7 +73,7 @@ UAIPerceptionComponent::UAIPerceptionComponent(const FObjectInitializer& ObjectI
 void UAIPerceptionComponent::RequestStimuliListenerUpdate()
 {
 	UAIPerceptionSystem* AIPerceptionSys = UAIPerceptionSystem::GetCurrent(GetWorld());
-	if (AIPerceptionSys != NULL)
+	if (AIPerceptionSys != nullptr)
 	{
 		AIPerceptionSys->UpdateListener(*this);
 	}
@@ -161,7 +161,7 @@ UAIPerceptionComponent::TAISenseConfigConstIterator UAIPerceptionComponent::GetS
 	return ToRawPtrTArrayUnsafe(SensesConfig).CreateConstIterator();
 }
 
-void UAIPerceptionComponent::SetMaxStimulusAge(FAISenseID SenseID, float MaxAge)
+void UAIPerceptionComponent::SetMaxStimulusAge(const FAISenseID SenseID, float MaxAge)
 {
 	if (!ensureMsgf(SenseID.IsValid(), TEXT("Sense must exist to update max age")))
 	{
@@ -214,7 +214,7 @@ void UAIPerceptionComponent::OnRegister()
 		}
 	}
 
-	// this should not be needed but aparently AAIController::PostRegisterAllComponents
+	// this should not be needed but apparently AAIController::PostRegisterAllComponents
 	// gets called component's OnRegister
 	AIOwner = Cast<AAIController>(GetOwner());
 	ensure(AIOwner == nullptr || AIOwner->GetAIPerceptionComponent() == nullptr || AIOwner->GetAIPerceptionComponent() == this
@@ -225,7 +225,7 @@ void UAIPerceptionComponent::OnRegister()
 	}
 }
 
-void UAIPerceptionComponent::RegisterSenseConfig(UAISenseConfig& SenseConfig, UAIPerceptionSystem& AIPerceptionSys)
+void UAIPerceptionComponent::RegisterSenseConfig(const UAISenseConfig& SenseConfig, UAIPerceptionSystem& AIPerceptionSys)
 {
 	const TSubclassOf<UAISense> SenseImplementation = SenseConfig.GetSenseImplementation();
 	if (SenseImplementation)
@@ -314,7 +314,7 @@ void UAIPerceptionComponent::UpdatePerceptionAllowList(const FAISenseID Channel,
 	}
 }
 
-bool UAIPerceptionComponent::GetFilteredActors(TFunctionRef<bool(const FActorPerceptionInfo&)> Predicate, TArray<AActor*>& OutActors) const
+bool UAIPerceptionComponent::GetFilteredActors(const TFunctionRef<bool(const FActorPerceptionInfo&)>& Predicate, TArray<AActor*>& OutActors) const
 {
 	bool bDeadDataFound = false;
 
@@ -347,7 +347,7 @@ void UAIPerceptionComponent::GetHostileActors(TArray<AActor*>& OutActors) const
 	{
 		FSimpleDelegateGraphTask::CreateAndDispatchWhenReady(
 			FSimpleDelegateGraphTask::FDelegate::CreateUObject(const_cast<UAIPerceptionComponent*>(this), &UAIPerceptionComponent::RemoveDeadData),
-			GET_STATID(STAT_FSimpleDelegateGraphTask_RequestingRemovalOfDeadPerceptionData), NULL, ENamedThreads::GameThread);
+			GET_STATID(STAT_FSimpleDelegateGraphTask_RequestingRemovalOfDeadPerceptionData), nullptr, ENamedThreads::GameThread);
 	}
 }
 
@@ -369,7 +369,7 @@ void UAIPerceptionComponent::GetHostileActorsBySense(TSubclassOf<UAISense> Sense
 	{
 		FSimpleDelegateGraphTask::CreateAndDispatchWhenReady(
 			FSimpleDelegateGraphTask::FDelegate::CreateUObject(const_cast<UAIPerceptionComponent*>(this), &UAIPerceptionComponent::RemoveDeadData),
-			GET_STATID(STAT_FSimpleDelegateGraphTask_RequestingRemovalOfDeadPerceptionData), NULL, ENamedThreads::GameThread);
+			GET_STATID(STAT_FSimpleDelegateGraphTask_RequestingRemovalOfDeadPerceptionData), nullptr, ENamedThreads::GameThread);
 	}
 }
 
@@ -377,7 +377,7 @@ const FActorPerceptionInfo* UAIPerceptionComponent::GetFreshestTrace(const FAISe
 {
 	// @note will stop on first age 0 stimulus
 	float BestAge = FAIStimulus::NeverHappenedAge;
-	const FActorPerceptionInfo* Result = NULL;
+	const FActorPerceptionInfo* Result = nullptr;
 
 	bool bDeadDataFound = false;
 	
@@ -408,7 +408,7 @@ const FActorPerceptionInfo* UAIPerceptionComponent::GetFreshestTrace(const FAISe
 	{
 		FSimpleDelegateGraphTask::CreateAndDispatchWhenReady(
 			FSimpleDelegateGraphTask::FDelegate::CreateUObject(const_cast<UAIPerceptionComponent*>(this), &UAIPerceptionComponent::RemoveDeadData),
-			GET_STATID(STAT_FSimpleDelegateGraphTask_RequestingRemovalOfDeadPerceptionData), NULL, ENamedThreads::GameThread);
+			GET_STATID(STAT_FSimpleDelegateGraphTask_RequestingRemovalOfDeadPerceptionData), nullptr, ENamedThreads::GameThread);
 	}
 
 	return Result;
@@ -453,8 +453,8 @@ void UAIPerceptionComponent::GetLocationAndDirection(FVector& Location, FVector&
 
 const AActor* UAIPerceptionComponent::GetBodyActor() const
 {
-	AController* OwnerController = Cast<AController>(GetOuter());
-	if (OwnerController != NULL)
+	const AController* OwnerController = Cast<AController>(GetOuter());
+	if (OwnerController != nullptr)
 	{
 		return OwnerController->GetPawn();
 	}
@@ -500,7 +500,7 @@ void UAIPerceptionComponent::ProcessStimuli()
 		FActorPerceptionInfo* PerceptualInfo = PerceptualData.Find(SourceKey);
 		AActor* SourceActor = nullptr;
 
-		if (PerceptualInfo == NULL)
+		if (PerceptualInfo == nullptr)
 		{
 			if (SourcedStimulus.Stimulus.WasSuccessfullySensed() == false)
 			{
@@ -627,7 +627,7 @@ void UAIPerceptionComponent::RefreshStimulus(FAIStimulus& StimulusStore, const F
 {
 	// if new stimulus is younger or stronger
 	// note that stimulus Age depends on PerceptionSystem::PerceptionAgingRate. It's possible that 
-	// both already stored and the new stimulus have Age of 0, but stored stimulus' acctual age is in [0, PerceptionSystem::PerceptionAgingRate)
+	// both already stored and the new stimulus have Age of 0, but stored stimulus' actual age is in [0, PerceptionSystem::PerceptionAgingRate)
 	if (NewStimulus.GetAge() <= StimulusStore.GetAge() || StimulusStore.Strength < NewStimulus.Strength)
 	{
 		StimulusStore = NewStimulus;
@@ -700,7 +700,7 @@ void UAIPerceptionComponent::ForgetAll()
 float UAIPerceptionComponent::GetYoungestStimulusAge(const AActor& Source) const
 {
 	const FActorPerceptionInfo* Info = GetActorInfo(Source);
-	if (Info == NULL)
+	if (Info == nullptr)
 	{
 		return FAIStimulus::NeverHappenedAge;
 	}
@@ -724,7 +724,7 @@ float UAIPerceptionComponent::GetYoungestStimulusAge(const AActor& Source) const
 bool UAIPerceptionComponent::HasAnyActiveStimulus(const AActor& Source) const
 {
 	const FActorPerceptionInfo* Info = GetActorInfo(Source);
-	if (Info == NULL)
+	if (Info == nullptr)
 	{
 		return false;
 	}
@@ -735,7 +735,7 @@ bool UAIPerceptionComponent::HasAnyActiveStimulus(const AActor& Source) const
 bool UAIPerceptionComponent::HasAnyCurrentStimulus(const AActor& Source) const
 {
 	const FActorPerceptionInfo* Info = GetActorInfo(Source);
-	if (Info == NULL)
+	if (Info == nullptr)
 	{
 		return false;
 	}
@@ -743,7 +743,7 @@ bool UAIPerceptionComponent::HasAnyCurrentStimulus(const AActor& Source) const
 	return Info->HasAnyCurrentStimulus();
 }
 
-bool UAIPerceptionComponent::HasActiveStimulus(const AActor& Source, FAISenseID Sense) const
+bool UAIPerceptionComponent::HasActiveStimulus(const AActor& Source, const FAISenseID Sense) const
 {
 	const FActorPerceptionInfo* Info = GetActorInfo(Source);
 	return (Info 
@@ -872,7 +872,7 @@ void UAIPerceptionComponent::DescribeSelfToGameplayDebugger(FGameplayDebuggerCat
 		}
 	}
 
-	for (UAISenseConfig* SenseConfig : SensesConfig)
+	for (const UAISenseConfig* SenseConfig : SensesConfig)
 	{
 		if (SenseConfig)
 		{
