Fix for BT not storing Aux nodes from tree searches. This can cause aux nodes not to be monitored at all when they should be and can make AI unrepsonsive.

diff --git a/Engine/Source/Developer/AITestSuite/Classes/BehaviorTree/TestBTTask_SetValuesWithLogs.h b/Engine/Source/Developer/AITestSuite/Classes/BehaviorTree/TestBTTask_SetValuesWithLogs.h
new file mode 100644
--- /dev/null
+++ b/Engine/Source/Developer/AITestSuite/Classes/BehaviorTree/TestBTTask_SetValuesWithLogs.h
@@ -0,0 +1,76 @@
+// Copyright Epic Games, Inc. All Rights Reserved.
+#pragma once
+
+#include "UObject/ObjectMacros.h"
+#include "BehaviorTree/BTTaskNode.h"
+#include "TestBTTask_SetValuesWithLogs.generated.h"
+
+struct FBTSetValueTaskMemory
+{
+	uint64 EndFrameIdx;
+	uint64 EndFrameIdx2;
+};
+
+UCLASS(meta=(HiddenNode))
+class UTestBTTask_SetValuesWithLogs : public UBTTaskNode
+{
+	GENERATED_BODY()
+
+public:
+	UTestBTTask_SetValuesWithLogs(const FObjectInitializer& ObjectInitializer);
+
+	/** Index logged when we execute the task */
+	UPROPERTY()
+	int32 LogIndex;
+
+	/** Index logged when we finish the task */
+	UPROPERTY()
+	int32 LogFinished;
+
+	/** Number of execution ticks until we set Blackboard KeyName2 Value2 */
+	UPROPERTY()
+	int32 ExecutionTicks1;
+
+	/** Number of execution ticks after ExecutionTicks1 that we end the task */
+	UPROPERTY()
+	int32 ExecutionTicks2;
+
+	/** Index logged when we tick the task */
+	UPROPERTY()
+	int32 LogTickIndex;
+
+	/** Blackboard keyname used when we start the task */
+	UPROPERTY()
+	FName KeyName1;
+
+	/** Blackboard value set when we start the task */
+	UPROPERTY()
+	int32 Value1;
+
+	/** Blackboard keyname used after ExecutionTicks1 execution ticks */
+	UPROPERTY()
+	FName KeyName2;
+
+	/** Blackboard value set after ExecutionTicks1 execution ticks */
+	UPROPERTY()
+	int32 Value2;
+
+	/** Blackboard keyname used when we abort the task */
+	UPROPERTY()
+	FName OnAbortKeyName;
+
+	/** Blackboard value set when we abort the task */
+	UPROPERTY()
+	int32 OnAbortValue;
+
+	/** Result when we finish the task */
+	UPROPERTY()
+	TEnumAsByte<EBTNodeResult::Type> TaskResult;
+
+protected:
+	virtual void TickTask(UBehaviorTreeComponent& OwnerComp, uint8* NodeMemory, float DeltaSeconds) override;
+	void LogExecution(UBehaviorTreeComponent& OwnerComp, int32 LogNumber);
+	virtual EBTNodeResult::Type ExecuteTask(UBehaviorTreeComponent& OwnerComp, uint8* NodeMemory) override;
+	virtual EBTNodeResult::Type AbortTask(UBehaviorTreeComponent& OwnerComp, uint8* NodeMemory) override;
+	virtual uint16 GetInstanceMemorySize() const override;
+};
diff --git a/Engine/Source/Developer/AITestSuite/Private/BehaviorTree/TestBTTask_SetValuesWithLogs.cpp b/Engine/Source/Developer/AITestSuite/Private/BehaviorTree/TestBTTask_SetValuesWithLogs.cpp
new file mode 100644
--- /dev/null
+++ b/Engine/Source/Developer/AITestSuite/Private/BehaviorTree/TestBTTask_SetValuesWithLogs.cpp
@@ -0,0 +1,92 @@
+// Copyright Epic Games, Inc. All Rights Reserved.
+
+#include "BehaviorTree/TestBTTask_SetValuesWithLogs.h"
+#include "BehaviorTree/Blackboard/BlackboardKeyType_Int.h"
+#include "BehaviorTree/BlackboardComponent.h"
+#include "AITestsCommon.h"
+#include "MockAI_BT.h"
+
+#include UE_INLINE_GENERATED_CPP_BY_NAME(TestBTTask_SetValuesWithLogs)
+
+UTestBTTask_SetValuesWithLogs::UTestBTTask_SetValuesWithLogs(const FObjectInitializer& ObjectInitializer) : Super(ObjectInitializer)
+{
+	NodeName = "SetValues";
+	ExecutionTicks1 = 0;
+	ExecutionTicks2 = 0;
+	LogIndex = 0;
+	LogFinished = -1;
+	LogTickIndex = -1;
+	TaskResult = EBTNodeResult::Succeeded;
+	KeyName1 = TEXT("Int");
+	Value1 = 1;
+	KeyName2 = TEXT("Int2");
+	Value2 = 1;
+	OnAbortKeyName = FName();
+	OnAbortValue = 1;
+
+	bIgnoreRestartSelf = true;
+
+	INIT_TASK_NODE_NOTIFY_FLAGS();
+}
+
+EBTNodeResult::Type UTestBTTask_SetValuesWithLogs::ExecuteTask(UBehaviorTreeComponent& OwnerComp, uint8* NodeMemory)
+{
+	FBTSetValueTaskMemory* MyMemory = CastInstanceNodeMemory<FBTSetValueTaskMemory>(NodeMemory);
+	MyMemory->EndFrameIdx = ExecutionTicks1 + FAITestHelpers::FramesCounter();
+	MyMemory->EndFrameIdx2 = MyMemory->EndFrameIdx + ExecutionTicks2;
+
+	OwnerComp.GetBlackboardComponent()->SetValue<UBlackboardKeyType_Int>(KeyName1, Value1);
+	LogExecution(OwnerComp, LogIndex);
+
+	if (ExecutionTicks1 == 0)
+	{
+		OwnerComp.GetBlackboardComponent()->SetValue<UBlackboardKeyType_Int>(KeyName2, Value2);
+
+		if (ExecutionTicks2 == 0)
+		{
+			return TaskResult;
+		}
+	}
+
+	return EBTNodeResult::InProgress;
+}
+void UTestBTTask_SetValuesWithLogs::TickTask(UBehaviorTreeComponent& OwnerComp, uint8* NodeMemory, float DeltaSeconds)
+{
+	FBTSetValueTaskMemory* MyMemory = CastInstanceNodeMemory<FBTSetValueTaskMemory>(NodeMemory);
+
+	LogExecution(OwnerComp, LogTickIndex);
+
+	if (FAITestHelpers::FramesCounter() >= MyMemory->EndFrameIdx)
+	{
+		OwnerComp.GetBlackboardComponent()->SetValue<UBlackboardKeyType_Int>(KeyName2, Value2);
+		MyMemory->EndFrameIdx = TNumericLimits<uint64>::Max();
+	}
+
+	if (FAITestHelpers::FramesCounter() >= MyMemory->EndFrameIdx2)
+	{
+		LogExecution(OwnerComp, LogFinished);
+		FinishLatentTask(OwnerComp, TaskResult);
+	}
+}
+
+uint16 UTestBTTask_SetValuesWithLogs::GetInstanceMemorySize() const
+{
+	return sizeof(FBTSetValueTaskMemory);
+}
+
+EBTNodeResult::Type UTestBTTask_SetValuesWithLogs::AbortTask(UBehaviorTreeComponent& OwnerComp, uint8* NodeMemory)
+{
+	if (OnAbortKeyName.IsValid())
+	{
+		OwnerComp.GetBlackboardComponent()->SetValue<UBlackboardKeyType_Int>(OnAbortKeyName, OnAbortValue);
+	}
+	return EBTNodeResult::Aborted;	
+}
+
+void UTestBTTask_SetValuesWithLogs::LogExecution(UBehaviorTreeComponent& OwnerComp, int32 LogNumber)
+{
+	if (LogNumber >= 0)
+	{
+		UMockAI_BT::ExecutionLog.Add(LogNumber);
+	}
+}
\ No newline at end of file
diff --git a/Engine/Source/Developer/AITestSuite/Private/Tests/BTTest.cpp b/Engine/Source/Developer/AITestSuite/Private/Tests/BTTest.cpp
--- a/Engine/Source/Developer/AITestSuite/Private/Tests/BTTest.cpp
+++ b/Engine/Source/Developer/AITestSuite/Private/Tests/BTTest.cpp
@@ -3944,4 +3944,38 @@ struct FAITest_BTCompleteRestartDuringTaskExecute : public FAITest_SimpleBT
 	}
 };
 IMPLEMENT_AI_LATENT_TEST(FAITest_BTCompleteRestartDuringTaskExecute, "System.AI.Behavior Trees.Restart: complete during task execute")
+
+struct FAITest_BTTestOberserverAbortLowerPriorityTwoDeep: public FAITest_SimpleBT
+{
+	FAITest_BTTestOberserverAbortLowerPriorityTwoDeep()
+	{
+		enum
+		{
+			MainTaskExecute = 1,
+			InteruptingTaskExecute,
+		};
+
+		UBTCompositeNode& RootNode = FBTBuilder::AddSelector(*BTAsset);
+		UBTCompositeNode& CompNode = FBTBuilder::AddSelector(RootNode);
+
+		{
+			UBTCompositeNode& Comp1Node = FBTBuilder::AddSelector(CompNode);
+			FBTBuilder::WithDecoratorBlackboard(CompNode, EArithmeticKeyOperation::Equal, 1, EBTFlowAbortMode::LowerPriority, EBTBlackboardRestart::ValueChange, TEXT("Int"));
+
+			UBTCompositeNode& Comp2Node = FBTBuilder::AddSelector(Comp1Node);
+			FBTBuilder::WithDecoratorBlackboard(Comp1Node, EArithmeticKeyOperation::Equal, 1, EBTFlowAbortMode::LowerPriority, EBTBlackboardRestart::ValueChange, TEXT("Int2"));
+
+			FBTBuilder::AddTask(Comp2Node, InteruptingTaskExecute, EBTNodeResult::Succeeded);
+		}
+
+		{
+			FBTBuilder::AddTaskValuesChangedWithLogs(/*ParentNode*/ CompNode, /*LogIndex*/ MainTaskExecute, /*NodeResult*/ EBTNodeResult::Succeeded, /*Value1*/1, /*Value2*/ 1, /*IntKeyName2*/ TEXT("Int"), /* IntKeyName2 */ TEXT("Int2"), /*ExecutionTicks1*/ 5, /*ExecutionTicks2*/ 5);
+		}
+
+		ExpectedResult.Add(MainTaskExecute);
+		ExpectedResult.Add(InteruptingTaskExecute);
+	}
+};
+IMPLEMENT_AI_LATENT_TEST(FAITest_BTTestOberserverAbortLowerPriorityTwoDeep, "System.AI.Behavior Trees.OberserverAbortLowerPriorityTwoDeep: Two nodes in same branch with observer aborts lower priority")
+
 #undef LOCTEXT_NAMESPACE
\ No newline at end of file
diff --git a/Engine/Source/Developer/AITestSuite/Public/BTBuilder.h b/Engine/Source/Developer/AITestSuite/Public/BTBuilder.h
--- a/Engine/Source/Developer/AITestSuite/Public/BTBuilder.h
+++ b/Engine/Source/Developer/AITestSuite/Public/BTBuilder.h
@@ -24,6 +24,7 @@
 #include "BehaviorTree/TestBTTask_Log.h"
 #include "BehaviorTree/TestBTTask_SetFlag.h"
 #include "BehaviorTree/TestBTTask_SetValue.h"
+#include "BehaviorTree/TestBTTask_SetValuesWithLogs.h"
 #include "BehaviorTree/TestBTTask_BTStopAction.h"
 #include "BehaviorTree/TestBTTask_ToggleFlag.h"
 
@@ -54,6 +55,10 @@ struct FBTBuilder
 		KeyData.KeyType = NewObject<UBlackboardKeyType_Int>();
 		BB->Keys.Add(KeyData);
 
+		KeyData.EntryName = TEXT("Int2");
+		KeyData.KeyType = NewObject<UBlackboardKeyType_Int>();
+		BB->Keys.Add(KeyData);
+
 		BB->UpdateParentKeys();
 
 		UBehaviorTree* TreeOb = NewObject<UBehaviorTree>();
@@ -190,6 +195,26 @@ struct FBTBuilder
 		ParentNode.Children[ChildIdx].ChildTask = TaskNode;
 	}
 
+	static void AddTaskValuesChangedWithLogs(UBTCompositeNode& ParentNode, int32 LogIndex, EBTNodeResult::Type NodeResult, int32 Value1, int32 Value2, FName IntKeyName1 = TEXT("Int"), FName IntKeyName2 = TEXT("Int2"), int32 ExecutionTicks1 = 0, int32 ExecutionTicks2 = 0, int32 LogTickIndex = -1, int32 LogFinished = -1, FName IntOnAbortKeyName = FName(), int32 OnAbortValue = 0)
+	{
+		UTestBTTask_SetValuesWithLogs* TaskNode = NewObject<UTestBTTask_SetValuesWithLogs>(ParentNode.GetTreeAsset());
+		TaskNode->LogIndex = LogIndex;
+		TaskNode->LogFinished = LogFinished;
+		TaskNode->ExecutionTicks1 = ExecutionTicks1;
+		TaskNode->ExecutionTicks2 = ExecutionTicks2;
+		TaskNode->LogTickIndex = LogTickIndex;
+		TaskNode->KeyName1 = IntKeyName1;
+		TaskNode->Value1 = Value1;
+		TaskNode->KeyName2 = IntKeyName2;
+		TaskNode->Value2 = Value2;
+		TaskNode->OnAbortKeyName = IntOnAbortKeyName;
+		TaskNode->OnAbortValue = OnAbortValue;
+		TaskNode->TaskResult = NodeResult;
+
+		const int32 ChildIdx = ParentNode.Children.AddZeroed(1);
+		ParentNode.Children[ChildIdx].ChildTask = TaskNode;
+	}
+
 	static void AddTaskSubtree(UBTCompositeNode& ParentNode, UBehaviorTree* TreeAsset)
 	{
 		UBTTask_RunBehavior* TaskNode = NewObject<UBTTask_RunBehavior>(ParentNode.GetTreeAsset());
diff --git a/Engine/Source/Runtime/AIModule/Classes/BehaviorTree/BehaviorTreeComponent.h b/Engine/Source/Runtime/AIModule/Classes/BehaviorTree/BehaviorTreeComponent.h
--- a/Engine/Source/Runtime/AIModule/Classes/BehaviorTree/BehaviorTreeComponent.h
+++ b/Engine/Source/Runtime/AIModule/Classes/BehaviorTree/BehaviorTreeComponent.h
@@ -403,8 +403,16 @@ protected:
 	AIMODULE_API void ApplyDiscardedSearch();
 
 	/** apply updates from specific list */
+	UE_DEPRECATED(5.4, "This function is deprecated. Please use ApplyAllSearchUpdates instead.")
 	AIMODULE_API void ApplySearchUpdates(const TArray<FBehaviorTreeSearchUpdate>& UpdateList, int32 NewNodeExecutionIndex, bool bPostUpdate = false);
 
+	/** 
+	 * Apply updates and post update from specified UpdateList.
+	 * @param bDoPostUpdate if true the post updates will also be processed.
+	 * @param bAllowTaskUpdates If false Task node updates will not be processed from the UpdateList.
+	 */
+	AIMODULE_API void ApplyAllSearchUpdates(const TArray<FBehaviorTreeSearchUpdate>& UpdateList, int32 NewNodeExecutionIndex, bool bDoPostUpdate = true, bool bAllowTaskUpdates = true);
+
 	/** abort currently executed task */
 	AIMODULE_API void AbortCurrentTask();
 
@@ -499,6 +507,8 @@ protected:
 		UBehaviorTreeComponent& BTComp;
 	};
 
+	void TickNewlyAddedAuxNodesHelper();
+
 	UE_DEPRECATED(5.1, "This struct is deprecated. Please use FBTSuspendBranchActionsScoped instead.")
 	typedef FBTSuspendBranchActionsScoped FBTSuspendBranchDeactivationScoped;
 
@@ -510,6 +520,10 @@ protected:
 	friend FBehaviorTreeDebugger;
 	friend FBehaviorTreeInstance;
 
+private:
+	/** Please don't call this function directly instead call ApplyAllSearchUpdates */
+	void ApplySearchUpdatesImpl(const TArray<FBehaviorTreeSearchUpdate>& UpdateList, int32 NewNodeExecutionIndex, bool bPostUpdate, bool bAllowTaskUpdates);
+
 protected:
 	/** data asset defining the tree */
 	UPROPERTY(EditDefaultsOnly, BlueprintReadWrite, Category = AI)
diff --git a/Engine/Source/Runtime/AIModule/Classes/BehaviorTree/BehaviorTreeTypes.h b/Engine/Source/Runtime/AIModule/Classes/BehaviorTree/BehaviorTreeTypes.h
--- a/Engine/Source/Runtime/AIModule/Classes/BehaviorTree/BehaviorTreeTypes.h
+++ b/Engine/Source/Runtime/AIModule/Classes/BehaviorTree/BehaviorTreeTypes.h
@@ -371,12 +371,18 @@ struct FBehaviorTreeInstance
 	/** get list of all active auxiliary nodes */
 	TArrayView<UBTAuxiliaryNode* const> GetActiveAuxNodes() const { return ActiveAuxNodes; }
 
-	/** add specified node to the active nodes list */
+	UE_DEPRECATED(5.4, "This function is deprecated. Please use version that takes UBehaviorTreeComponent.")
 	void AddToActiveAuxNodes(UBTAuxiliaryNode* AuxNode);
 
-	/** remove specified node from the active nodes list */
+	/** add specified node to the active nodes list */
+	void AddToActiveAuxNodes(UBehaviorTreeComponent& OwnerComp, UBTAuxiliaryNode* AuxNode);
+
+	UE_DEPRECATED(5.4, "This function is deprecated. Please use version that takes UBehaviorTreeComponent.")
 	void RemoveFromActiveAuxNodes(UBTAuxiliaryNode* AuxNode);
 
+	/** remove specified node from the active nodes list */
+	void RemoveFromActiveAuxNodes(UBehaviorTreeComponent& OwnerComp, UBTAuxiliaryNode* AuxNode);
+
 	/** remove all auxiliary nodes from active nodes list */
 	void ResetActiveAuxNodes();
 
@@ -412,6 +418,10 @@ protected:
 	/** worker for updating all nodes */
 	void CleanupNodes(UBehaviorTreeComponent& OwnerComp, UBTCompositeNode& Node, EBTMemoryClear::Type CleanupType);
 
+private:
+	void AddToActiveAuxNodesImpl(UBTAuxiliaryNode* AuxNode);
+	void RemoveFromActiveAuxNodesImpl(UBTAuxiliaryNode* AuxNode);
+
 private:
 #if DO_ENSURE
 	/** debug flag to detect modifications to the array of nodes while iterating through it */
diff --git a/Engine/Source/Runtime/AIModule/Private/BehaviorTree/BehaviorTreeComponent.cpp b/Engine/Source/Runtime/AIModule/Private/BehaviorTree/BehaviorTreeComponent.cpp
--- a/Engine/Source/Runtime/AIModule/Private/BehaviorTree/BehaviorTreeComponent.cpp
+++ b/Engine/Source/Runtime/AIModule/Private/BehaviorTree/BehaviorTreeComponent.cpp
@@ -33,6 +33,8 @@ double UBehaviorTreeComponent::FrameSearchTime = 0.;
 int32 UBehaviorTreeComponent::NumSearchTimeCalls = 0;
 #endif
 
+static TAutoConsoleVariable<bool> CVarAddAuxNodesFromFailedSearches(TEXT("BehaviorTree.AddAuxNodesFromFailedSearches"), 1, TEXT("Add Aux Nodes From Failed Searches"));
+
 //----------------------------------------------------------------------//
 // UBehaviorTreeComponent
 //----------------------------------------------------------------------//
@@ -1420,7 +1422,7 @@ void UBehaviorTreeComponent::RequestExecution(const UBTCompositeNode* RequestedO
 	}
 }
 
-void UBehaviorTreeComponent::ApplySearchUpdates(const TArray<FBehaviorTreeSearchUpdate>& UpdateList, int32 NewNodeExecutionIndex, bool bPostUpdate)
+void UBehaviorTreeComponent::ApplySearchUpdatesImpl(const TArray<FBehaviorTreeSearchUpdate>& UpdateList, int32 NewNodeExecutionIndex, bool bPostUpdate, bool bAllowTaskUpdates)
 {
 	for (int32 Index = 0; Index < UpdateList.Num(); Index++)
 	{
@@ -1444,17 +1446,21 @@ void UBehaviorTreeComponent::ApplySearchUpdates(const TArray<FBehaviorTreeSearch
 		{
 			bIsComponentActive = UpdateInstance.GetActiveAuxNodes().Contains(UpdateInfo.AuxNode);
 		}
-		else if (UpdateInfo.TaskNode)
+		else if (bAllowTaskUpdates && UpdateInfo.TaskNode)
 		{
 			ParallelTaskIdx = UpdateInstance.GetParallelTasks().IndexOfByKey(UpdateInfo.TaskNode);
 			bIsComponentActive = (ParallelTaskIdx != INDEX_NONE && UpdateInstance.GetParallelTasks()[ParallelTaskIdx].Status == EBTTaskStatus::Active);
 		}
+		else
+		{
+			continue;
+		}
 
 		const UBTNode* UpdateNode = UpdateInfo.AuxNode ? (const UBTNode*)UpdateInfo.AuxNode : (const UBTNode*)UpdateInfo.TaskNode;
 		checkSlow(UpdateNode);
 
 		if ((UpdateInfo.Mode == EBTNodeUpdateMode::Remove && !bIsComponentActive) ||
-			(UpdateInfo.Mode == EBTNodeUpdateMode::Add && (bIsComponentActive || UpdateNode->GetExecutionIndex() > NewNodeExecutionIndex)) )
+			(UpdateInfo.Mode == EBTNodeUpdateMode::Add && (bIsComponentActive || UpdateNode->GetExecutionIndex() > NewNodeExecutionIndex)))
 		{
 			UpdateInfo.bApplySkipped = true;
 			UE_VLOG(GetOwner(), LogBehaviorTree, Log, TEXT("Skipping: %s for %s: %s (already in the right state)"),
@@ -1488,12 +1494,12 @@ void UBehaviorTreeComponent::ApplySearchUpdates(const TArray<FBehaviorTreeSearch
 			uint8* NodeMemory = (uint8*)UpdateNode->GetNodeMemory<uint8>(UpdateInstance);
 			if (UpdateInfo.Mode == EBTNodeUpdateMode::Remove)
 			{
-				UpdateInstance.RemoveFromActiveAuxNodes(UpdateInfo.AuxNode);
+				UpdateInstance.RemoveFromActiveAuxNodes(*this, UpdateInfo.AuxNode);
 				UpdateInfo.AuxNode->WrappedOnCeaseRelevant(*this, NodeMemory);
 			}
 			else
 			{
-				UpdateInstance.AddToActiveAuxNodes(UpdateInfo.AuxNode);
+				UpdateInstance.AddToActiveAuxNodes(*this, UpdateInfo.AuxNode);
 				UpdateInfo.AuxNode->WrappedOnBecomeRelevant(*this, NodeMemory);
 			}
 		}
@@ -1515,7 +1521,7 @@ void UBehaviorTreeComponent::ApplySearchUpdates(const TArray<FBehaviorTreeSearch
 				const bool bStillValid = InstanceStack.IsValidIndex(UpdateInfo.InstanceIndex) &&
 					InstanceStack[UpdateInfo.InstanceIndex].GetParallelTasks().IsValidIndex(ParallelTaskIdx) &&
 					InstanceStack[UpdateInfo.InstanceIndex].GetParallelTasks()[ParallelTaskIdx] == UpdateInfo.TaskNode;
-				
+
 				if (bStillValid)
 				{
 					// mark as pending abort
@@ -1539,6 +1545,60 @@ void UBehaviorTreeComponent::ApplySearchUpdates(const TArray<FBehaviorTreeSearch
 	}
 }
 
+void UBehaviorTreeComponent::ApplySearchUpdates(const TArray<FBehaviorTreeSearchUpdate>& UpdateList, int32 NewNodeExecutionIndex, bool bPostUpdate)
+{
+	ApplySearchUpdatesImpl(UpdateList, NewNodeExecutionIndex, bPostUpdate, true);
+}
+
+void UBehaviorTreeComponent::ApplyAllSearchUpdates(const TArray<FBehaviorTreeSearchUpdate>& UpdateList, int32 NewNodeExecutionIndex, bool bDoPostUpdate, bool bAllowTaskUpdates)
+{
+	const int32 NumIterations = bDoPostUpdate ? 2 : 1;
+	constexpr int32 PostUpdateIter = 1;
+
+	for (int32 UpdateIter = 0; UpdateIter < NumIterations; ++UpdateIter)
+	{
+		const bool bPostUpdate = UpdateIter == PostUpdateIter;
+
+		ApplySearchUpdatesImpl(UpdateList, NewNodeExecutionIndex, bPostUpdate, bAllowTaskUpdates);
+	}
+
+#if ENABLE_VISUAL_LOG
+	auto LogAuxNodes = [this]()
+	{
+		UE_VLOG(GetOwner(), LogBehaviorTree, VeryVerbose, TEXT("Logging All Aux Nodes"));
+		for (int32 InstanceIndex = 0; InstanceIndex < InstanceStack.Num(); InstanceIndex++)
+		{
+			FBehaviorTreeInstance& InstanceInfo = InstanceStack[InstanceIndex];
+			InstanceInfo.ExecuteOnEachAuxNode([this](const UBTAuxiliaryNode& AuxNode)
+				{
+					UE_VLOG(GetOwner(), LogBehaviorTree, VeryVerbose, TEXT("%s"), *UBehaviorTreeTypes::DescribeNodeHelper(&AuxNode));
+				});
+		}
+	};
+
+	UE_IFVLOG(LogAuxNodes());
+#endif // ENABLE_VISUAL_LOG
+}
+
+void UBehaviorTreeComponent::TickNewlyAddedAuxNodesHelper()
+{
+	UWorld* MyWorld = GetWorld();
+	const float CurrentFrameDeltaSeconds = MyWorld ? MyWorld->GetDeltaSeconds() : 0.0f;
+
+	for (const FBehaviorTreeSearchUpdate& UpdateInfo : SearchData.PendingUpdates)
+	{
+		if (!UpdateInfo.bApplySkipped && UpdateInfo.Mode == EBTNodeUpdateMode::Add && UpdateInfo.AuxNode && InstanceStack.IsValidIndex(UpdateInfo.InstanceIndex))
+		{
+			FBehaviorTreeInstance& InstanceInfo = InstanceStack[UpdateInfo.InstanceIndex];
+			uint8* NodeMemory = UpdateInfo.AuxNode->GetNodeMemory<uint8>(InstanceInfo);
+
+			// We do not care about the next needed DeltaTime, it will be recalculated in the tick later.
+			float NextNeededDeltaTime = 0.0f;
+			UpdateInfo.AuxNode->WrappedTickNode(*this, NodeMemory, CurrentFrameDeltaSeconds, NextNeededDeltaTime);
+		}
+	}
+}
+
 void UBehaviorTreeComponent::ApplySearchData(UBTNode* NewActiveNode)
 {
 	// search is finalized, can't rollback anymore at this point
@@ -1561,28 +1621,12 @@ void UBehaviorTreeComponent::ApplySearchData(UBTNode* NewActiveNode)
 
 	SearchData.bFilterOutRequestFromDeactivatedBranch = true;
 
-	ApplySearchUpdates(SearchData.PendingUpdates, NewNodeExecutionIndex);
-	ApplySearchUpdates(SearchData.PendingUpdates, NewNodeExecutionIndex, true);
+	ApplyAllSearchUpdates(SearchData.PendingUpdates, NewNodeExecutionIndex);
 	
 	SearchData.bFilterOutRequestFromDeactivatedBranch = false;
 
 	// tick newly added aux nodes to compensate for tick-search order changes
-	UWorld* MyWorld = GetWorld();
-	const float CurrentFrameDeltaSeconds = MyWorld ? MyWorld->GetDeltaSeconds() : 0.0f;
-
-	for (int32 Idx = 0; Idx < SearchData.PendingUpdates.Num(); Idx++)
-	{
-		const FBehaviorTreeSearchUpdate& UpdateInfo = SearchData.PendingUpdates[Idx];
-		if (!UpdateInfo.bApplySkipped && UpdateInfo.Mode == EBTNodeUpdateMode::Add && UpdateInfo.AuxNode && InstanceStack.IsValidIndex(UpdateInfo.InstanceIndex))
-		{
-			FBehaviorTreeInstance& InstanceInfo = InstanceStack[UpdateInfo.InstanceIndex];
-			uint8* NodeMemory = UpdateInfo.AuxNode->GetNodeMemory<uint8>(InstanceInfo);
-
-			// We do not care about the next needed DeltaTime, it will be recalculated in the tick later.
-			float NextNeededDeltaTime = 0.0f;
-			UpdateInfo.AuxNode->WrappedTickNode(*this, NodeMemory, CurrentFrameDeltaSeconds, NextNeededDeltaTime);
-		}
-	}
+	TickNewlyAddedAuxNodesHelper();
 
 	// clear update list
 	// nothing should be added during application or tick - all changes are supposed to go to ExecutionRequest accumulator first
@@ -1594,7 +1638,24 @@ void UBehaviorTreeComponent::ApplySearchData(UBTNode* NewActiveNode)
 
 void UBehaviorTreeComponent::ApplyDiscardedSearch()
 {
-	// remove everything else
+	if (CVarAddAuxNodesFromFailedSearches.GetValueOnGameThread())
+	{
+		// Apply aux nodes from last search with the currently active node as the 'NewNode'
+		int32 NewNodeExecutionIdx = 0;
+
+		if (InstanceStack.IsValidIndex(ActiveInstanceIdx))
+		{
+			FBehaviorTreeInstance& BTInstance = InstanceStack[ActiveInstanceIdx];
+
+			NewNodeExecutionIdx = BTInstance.ActiveNode ? BTInstance.ActiveNode->GetExecutionIndex() : 0;
+		}
+
+		ApplyAllSearchUpdates(/*UpdateList*/ SearchData.PendingUpdates, /*NewNodeExecutionIndex*/ NewNodeExecutionIdx, /*bDoPostUpdate*/ true, /*bAllowTaskUpdates*/ false);
+
+		// tick newly added aux nodes to compensate for tick-search order changes
+		TickNewlyAddedAuxNodesHelper();
+	}
+
 	SearchData.PendingUpdates.Reset();
 
 	// don't send deactivation notifies
@@ -1867,6 +1928,8 @@ void UBehaviorTreeComponent::ProcessExecutionRequest()
 		return;
 	}
 
+	UE_VLOG(GetOwner(), LogBehaviorTree, VeryVerbose, TEXT("%hs Active Node %s"), __FUNCTION__, *UBehaviorTreeTypes::DescribeNodeHelper(InstanceStack[ActiveInstanceIdx].ActiveNode));
+
 	if (PendingExecution.IsSet())
 	{
 		ProcessPendingExecution();
@@ -2116,7 +2179,9 @@ void UBehaviorTreeComponent::ProcessExecutionRequest()
 }
 
 void UBehaviorTreeComponent::ProcessPendingExecution()
-{
+{ 
+	UE_VLOG(GetOwner(), LogBehaviorTree, VeryVerbose, TEXT("%hs"), __FUNCTION__);
+
 	// can't continue if current task is still aborting
 	if (bWaitingForLatentAborts || !PendingExecution.IsSet())
 	{
@@ -2331,7 +2396,7 @@ void UBehaviorTreeComponent::UnregisterAuxNodesInBranch(const UBTCompositeNode*
 		if (bApplyImmediately)
 		{
 			FBTSuspendBranchActionsScoped ScopedSuspend(*this, EBTBranchAction::All);
-			ApplySearchUpdates(SearchData.PendingUpdates, 0);
+			ApplyAllSearchUpdates(/*UpdateList*/ SearchData.PendingUpdates, /*NewNodeExecutionIndex*/ 0, /*bDoPostUpdate*/ false);
 			SearchData.PendingUpdates = UpdateListCopy;
 		}
 	}
@@ -2354,7 +2419,7 @@ void UBehaviorTreeComponent::ExecuteTask(UBTTaskNode* TaskNode)
 	{
 		uint8* NodeMemory = (uint8*)ServiceNode->GetNodeMemory<uint8>(ActiveInstance);
 
-		ActiveInstance.AddToActiveAuxNodes(ServiceNode);
+		ActiveInstance.AddToActiveAuxNodes(*this, ServiceNode);
 
 		UE_VLOG(GetOwner(), LogBehaviorTree, Log, TEXT("Activating task service: %s"), *UBehaviorTreeTypes::DescribeNodeHelper(ServiceNode));
 		ServiceNode->WrappedOnBecomeRelevant(*this, NodeMemory);
@@ -2663,7 +2728,7 @@ bool UBehaviorTreeComponent::PushInstance(UBehaviorTree& TreeAsset)
 			// send initial on search start events in case someone is using them for init logic
 			ServiceNode->NotifyParentActivation(SearchData);
 
-			InstanceStack[ActiveInstanceIdx].AddToActiveAuxNodes(ServiceNode);
+			InstanceStack[ActiveInstanceIdx].AddToActiveAuxNodes(*this,ServiceNode);
 			ServiceNode->WrappedOnBecomeRelevant(*this, NodeMemory);
 		}
 
diff --git a/Engine/Source/Runtime/AIModule/Private/BehaviorTree/BehaviorTreeTypes.cpp b/Engine/Source/Runtime/AIModule/Private/BehaviorTree/BehaviorTreeTypes.cpp
--- a/Engine/Source/Runtime/AIModule/Private/BehaviorTree/BehaviorTreeTypes.cpp
+++ b/Engine/Source/Runtime/AIModule/Private/BehaviorTree/BehaviorTreeTypes.cpp
@@ -262,21 +262,51 @@ FBehaviorTreeInstance::~FBehaviorTreeInstance()
 }
 
 void FBehaviorTreeInstance::AddToActiveAuxNodes(UBTAuxiliaryNode* AuxNode)
+{
+	AddToActiveAuxNodesImpl(AuxNode);
+}
+
+void FBehaviorTreeInstance::AddToActiveAuxNodesImpl(UBTAuxiliaryNode* AuxNode)
 {
 #if DO_ENSURE
 	ensureAlwaysMsgf(bIteratingNodes == false, TEXT("Adding aux node while iterating through them is not allowed."));
 #endif // DO_ENSURE
+
 	MEM_STAT_UPDATE_WRAPPER(ActiveAuxNodes.Add(AuxNode));
 }
 
+void FBehaviorTreeInstance::AddToActiveAuxNodes(UBehaviorTreeComponent& OwnerComp, UBTAuxiliaryNode* AuxNode)
+{
+	UE_VLOG(OwnerComp.GetOwner(), LogBehaviorTree, Verbose, TEXT("%hs %s")
+		, __FUNCTION__
+		, *UBehaviorTreeTypes::DescribeNodeHelper(AuxNode));
+
+	AddToActiveAuxNodesImpl(AuxNode);
+}
+
 void FBehaviorTreeInstance::RemoveFromActiveAuxNodes(UBTAuxiliaryNode* AuxNode)
+{
+	RemoveFromActiveAuxNodesImpl(AuxNode);
+}
+
+void FBehaviorTreeInstance::RemoveFromActiveAuxNodesImpl(UBTAuxiliaryNode* AuxNode)
 {
 #if DO_ENSURE
 	ensureAlwaysMsgf(bIteratingNodes == false, TEXT("Removing aux node while iterating through them is not allowed."));
 #endif // DO_ENSURE
+
 	MEM_STAT_UPDATE_WRAPPER(ActiveAuxNodes.RemoveSingleSwap(AuxNode));
 }
 
+void FBehaviorTreeInstance::RemoveFromActiveAuxNodes(UBehaviorTreeComponent& OwnerComp, UBTAuxiliaryNode* AuxNode)
+{
+	UE_VLOG(OwnerComp.GetOwner(), LogBehaviorTree, Verbose, TEXT("%hs %s")
+		, __FUNCTION__
+		, *UBehaviorTreeTypes::DescribeNodeHelper(AuxNode));
+
+	RemoveFromActiveAuxNodesImpl(AuxNode);
+}
+
 void FBehaviorTreeInstance::ResetActiveAuxNodes()
 {
 #if DO_ENSURE
@@ -472,7 +502,7 @@ void FBehaviorTreeSearchData::AddUniqueUpdate(const FBehaviorTreeSearchUpdate& U
 	{
 		const bool bIsActive = OwnerComp.IsAuxNodeActive(UpdateInfo.AuxNode, UpdateInfo.InstanceIndex);
 		bSkipAdding = !bIsActive;
-		UE_CVLOG(bSkipAdding, OwnerComp.GetOwner(), LogBehaviorTree, Verbose, TEXT(">> skipped: inactive aux nodes"));
+		UE_CVLOG(bSkipAdding, OwnerComp.GetOwner(), LogBehaviorTree, Verbose, TEXT(">> skipped: did not push a remove to PendingUpdates due to inactive aux node"));
 	}
 
 	if (!bSkipAdding)
@@ -704,4 +734,3 @@ void UBehaviorTreeTypes::SetBTLoggingContext(const UBTNode* NewBTLoggingContext)
 //----------------------------------------------------------------------//
 // DEPRECATED
 //----------------------------------------------------------------------//
-
