Fix comments in BT decorators DoesPathExist, KeepInCone

diff --git a/Engine/Source/Runtime/AIModule/Classes/BehaviorTree/Decorators/BTDecorator_DoesPathExist.h b/Engine/Source/Runtime/AIModule/Classes/BehaviorTree/Decorators/BTDecorator_DoesPathExist.h
--- a/Engine/Source/Runtime/AIModule/Classes/BehaviorTree/Decorators/BTDecorator_DoesPathExist.h
+++ b/Engine/Source/Runtime/AIModule/Classes/BehaviorTree/Decorators/BTDecorator_DoesPathExist.h
@@ -24,7 +24,7 @@ namespace EPathExistanceQueryType
 }
 
 /**
- * Cooldown decorator node.
+ * Does Path Exist decorator node.
  * A decorator node that bases its condition on whether a path exists between two points from the Blackboard.
  */
 UCLASS(MinimalAPI)
diff --git a/Engine/Source/Runtime/AIModule/Classes/BehaviorTree/Decorators/BTDecorator_KeepInCone.h b/Engine/Source/Runtime/AIModule/Classes/BehaviorTree/Decorators/BTDecorator_KeepInCone.h
--- a/Engine/Source/Runtime/AIModule/Classes/BehaviorTree/Decorators/BTDecorator_KeepInCone.h
+++ b/Engine/Source/Runtime/AIModule/Classes/BehaviorTree/Decorators/BTDecorator_KeepInCone.h
@@ -16,7 +16,7 @@ struct FBTKeepInConeDecoratorMemory
 };
 
 /**
- * Cooldown decorator node.
+ * Keep In Cone decorator node.
  * A decorator node that bases its condition on whether the observed position is still inside a cone. The cone's direction is calculated when the node first becomes relevant.
  */
 UCLASS(HideCategories=(Condition), MinimalAPI)
