PR #9497: fixed mismatch in BT component description format

diff --git a/Engine/Source/Runtime/AIModule/Private/BehaviorTree/BehaviorTreeComponent.cpp b/Engine/Source/Runtime/AIModule/Private/BehaviorTree/BehaviorTreeComponent.cpp
--- a/Engine/Source/Runtime/AIModule/Private/BehaviorTree/BehaviorTreeComponent.cpp
+++ b/Engine/Source/Runtime/AIModule/Private/BehaviorTree/BehaviorTreeComponent.cpp
@@ -3118,7 +3118,7 @@ void UBehaviorTreeComponent::DescribeSelfToVisLog(FVisualLogEntry* Snapshot) con
 			FString ObserversDesc;
 			for (const UBTAuxiliaryNode* AuxNode : InstanceInfo.GetActiveAuxNodes())
 			{
-				ObserversDesc += FString::Printf(TEXT("%d. %s\n"), AuxNode->GetExecutionIndex(), *AuxNode->GetNodeName(), *AuxNode->GetStaticDescription());
+				ObserversDesc += FString::Printf(TEXT("%d. %s: %s\n"), AuxNode->GetExecutionIndex(), *AuxNode->GetNodeName(), *AuxNode->GetStaticDescription());
 			}
 			StatusCategory.Add(TEXT("Observers"), ObserversDesc);
 		}
