Perception:
- Add IsFriendly to FActorPerceptionInfo (to distinguish neutrals from friendlies)
- Add SetStartsEnabled to UAISenseConfig
#rnx

diff --git a/Engine/Source/Runtime/AIModule/Classes/Perception/AIPerceptionComponent.h b/Engine/Source/Runtime/AIModule/Classes/Perception/AIPerceptionComponent.h
--- a/Engine/Source/Runtime/AIModule/Classes/Perception/AIPerceptionComponent.h
+++ b/Engine/Source/Runtime/AIModule/Classes/Perception/AIPerceptionComponent.h
@@ -59,8 +59,11 @@ struct FActorPerceptionInfo
 	/** indicates whether this Actor is hostile to perception holder */
 	uint32 bIsHostile : 1;
 
+	/** indicates whether this Actor is hostile to perception holder */
+	uint32 bIsFriendly : 1;
+
 	explicit FActorPerceptionInfo(AActor* InTarget = nullptr)
-		: Target(InTarget), DominantSense(FAISenseID::InvalidID()), bIsHostile(false)
+		: Target(InTarget), DominantSense(FAISenseID::InvalidID()), bIsHostile(false), bIsFriendly(false)
 	{
 		LastSensedStimuli.AddDefaulted(FAISenseID::GetSize());
 	}
@@ -189,7 +192,10 @@ struct FActorPerceptionBlueprintInfo
 	UPROPERTY(BlueprintReadWrite, Category = "AI|Perception")
 	uint32 bIsHostile : 1;
 
-	FActorPerceptionBlueprintInfo() : Target(nullptr), bIsHostile(false)
+	UPROPERTY(BlueprintReadWrite, Category = "AI|Perception")
+	uint32 bIsFriendly : 1;
+
+	FActorPerceptionBlueprintInfo() : Target(nullptr), bIsHostile(false), bIsFriendly(false)
 	{}
 	FActorPerceptionBlueprintInfo(const FActorPerceptionInfo& Info);
 };
@@ -378,6 +384,9 @@ public:
 	UFUNCTION(BlueprintCallable, Category = "AI|Perception")
 	AIMODULE_API void SetSenseEnabled(TSubclassOf<UAISense> SenseClass, const bool bEnable);
 
+	UFUNCTION(BlueprintCallable, Category = "AI|Perception")
+	AIMODULE_API bool IsSenseEnabled(TSubclassOf<UAISense> SenseClass) const;
+
 	//////////////////////////////////////////////////////////////////////////
 	// Might want to move these to special "BP_AIPerceptionComponent"
 	//////////////////////////////////////////////////////////////////////////
diff --git a/Engine/Source/Runtime/AIModule/Classes/Perception/AISenseConfig.h b/Engine/Source/Runtime/AIModule/Classes/Perception/AISenseConfig.h
--- a/Engine/Source/Runtime/AIModule/Classes/Perception/AISenseConfig.h
+++ b/Engine/Source/Runtime/AIModule/Classes/Perception/AISenseConfig.h
@@ -43,9 +43,16 @@ public:
 	AIMODULE_API FString GetSenseName() const;
 	FColor GetDebugColor() const { return DebugColor; }
 	void SetMaxAge(const float NewMaxAge) { MaxAge = NewMaxAge; }
-	float GetMaxAge() const { return MaxAge == 0.f ? FAIStimulus::NeverHappenedAge : MaxAge; }	
+	float GetMaxAge() const { return MaxAge == 0.f ? FAIStimulus::NeverHappenedAge : MaxAge; }
+	
+	UE_DEPRECATED(5.4, "Use GetStartsEnabled instead")
 	bool IsEnabled() const { return bStartsEnabled; }
 
+	bool GetStartsEnabled() const { return bStartsEnabled; }
+
+	/** Changes whether the given sense starts off enabled. Note that calling the function after given sense config has already been registered won't have any effect */
+	void SetStartsEnabled(bool bEnabled) { bStartsEnabled = bEnabled; }
+
 #if WITH_GAMEPLAY_DEBUGGER_MENU
 	AIMODULE_API virtual void DescribeSelfToGameplayDebugger(const UAIPerceptionComponent* PerceptionComponent, FGameplayDebuggerCategory* DebuggerCategory) const;
 #endif // WITH_GAMEPLAY_DEBUGGER_MENU
diff --git a/Engine/Source/Runtime/AIModule/Private/Perception/AIPerceptionComponent.cpp b/Engine/Source/Runtime/AIModule/Private/Perception/AIPerceptionComponent.cpp
--- a/Engine/Source/Runtime/AIModule/Private/Perception/AIPerceptionComponent.cpp
+++ b/Engine/Source/Runtime/AIModule/Private/Perception/AIPerceptionComponent.cpp
@@ -45,6 +45,7 @@ FActorPerceptionBlueprintInfo::FActorPerceptionBlueprintInfo(const FActorPercept
 	Target = Info.Target.Get();
 	LastSensedStimuli = Info.LastSensedStimuli;
 	bIsHostile = Info.bIsHostile;
+	bIsFriendly = Info.bIsFriendly;
 }
 
 //----------------------------------------------------------------------//
@@ -524,6 +525,7 @@ void UAIPerceptionComponent::ProcessStimuli()
 				PerceptualInfo->DominantSense = DominantSenseID;
 
 				PerceptualInfo->bIsHostile = (FGenericTeamId::GetAttitude(GetOwner(), SourceActor) == ETeamAttitude::Hostile);
+				PerceptualInfo->bIsFriendly = PerceptualInfo->bIsHostile ? false : (FGenericTeamId::GetAttitude(GetOwner(), SourceActor) == ETeamAttitude::Friendly);
 			}
 		}
 
@@ -838,6 +840,17 @@ void UAIPerceptionComponent::SetSenseEnabled(TSubclassOf<UAISense> SenseClass, c
 	}
 }
 
+bool UAIPerceptionComponent::IsSenseEnabled(TSubclassOf<UAISense> SenseClass) const
+{
+	const FAISenseID SenseID = UAISense::GetSenseID(SenseClass);
+	if (!SenseID.IsValid() || GetSenseConfig(SenseID) == nullptr)
+	{
+		return false;
+	}
+
+	return PerceptionFilter.ShouldRespondToChannel(SenseID);
+}
+
 //----------------------------------------------------------------------//
 // debug
 //----------------------------------------------------------------------//
