Fix bad copy paste
#rnx

diff --git a/Engine/Source/Runtime/AIModule/Classes/Perception/AIPerceptionComponent.h b/Engine/Source/Runtime/AIModule/Classes/Perception/AIPerceptionComponent.h
--- a/Engine/Source/Runtime/AIModule/Classes/Perception/AIPerceptionComponent.h
+++ b/Engine/Source/Runtime/AIModule/Classes/Perception/AIPerceptionComponent.h
@@ -59,7 +59,7 @@ struct FActorPerceptionInfo
 	/** indicates whether this Actor is hostile to perception holder */
 	uint32 bIsHostile : 1;
 
-	/** indicates whether this Actor is hostile to perception holder */
+	/** indicates whether this Actor is friendly to perception holder */
 	uint32 bIsFriendly : 1;
 
 	explicit FActorPerceptionInfo(AActor* InTarget = nullptr)
