Use GetStartsEnabled instead of IsEnabled
#rnx

diff --git a/Engine/Source/Runtime/AIModule/Private/Perception/AIPerceptionComponent.cpp b/Engine/Source/Runtime/AIModule/Private/Perception/AIPerceptionComponent.cpp
--- a/Engine/Source/Runtime/AIModule/Private/Perception/AIPerceptionComponent.cpp
+++ b/Engine/Source/Runtime/AIModule/Private/Perception/AIPerceptionComponent.cpp
@@ -235,7 +235,7 @@ void UAIPerceptionComponent::RegisterSenseConfig(const UAISenseConfig& SenseConf
 		const FAISenseID SenseID = AIPerceptionSys.RegisterSenseClass(SenseImplementation);
 		check(SenseID.IsValid());
 
-		if (SenseConfig.IsEnabled())
+		if (SenseConfig.GetStartsEnabled())
 		{
 			PerceptionFilter.AcceptChannel(SenseID);
 		}
