Do not use deprecated function
Minor fixes in comments
#rnx

diff --git a/Engine/Source/Runtime/AIModule/Classes/Perception/AIPerceptionComponent.h b/Engine/Source/Runtime/AIModule/Classes/Perception/AIPerceptionComponent.h
--- a/Engine/Source/Runtime/AIModule/Classes/Perception/AIPerceptionComponent.h
+++ b/Engine/Source/Runtime/AIModule/Classes/Perception/AIPerceptionComponent.h
@@ -384,6 +384,8 @@ public:
 	UFUNCTION(BlueprintCallable, Category = "AI|Perception")
 	AIMODULE_API void SetSenseEnabled(TSubclassOf<UAISense> SenseClass, const bool bEnable);
 
+	/** Returns if a sense is active. Note that this works only if given sense has been
+	*	already configured for this component instance */
 	UFUNCTION(BlueprintCallable, Category = "AI|Perception")
 	AIMODULE_API bool IsSenseEnabled(TSubclassOf<UAISense> SenseClass) const;
 
diff --git a/Engine/Source/Runtime/AIModule/Classes/Perception/AISenseConfig.h b/Engine/Source/Runtime/AIModule/Classes/Perception/AISenseConfig.h
--- a/Engine/Source/Runtime/AIModule/Classes/Perception/AISenseConfig.h
+++ b/Engine/Source/Runtime/AIModule/Classes/Perception/AISenseConfig.h
@@ -50,7 +50,8 @@ public:
 
 	bool GetStartsEnabled() const { return bStartsEnabled; }
 
-	/** Changes whether the given sense starts off enabled. Note that calling the function after given sense config has already been registered won't have any effect */
+	/** Changes whether the given sense starts off enabled.
+	*	Note that calling the function after given sense config has already been registered won't have any effect */
 	void SetStartsEnabled(bool bEnabled) { bStartsEnabled = bEnabled; }
 
 #if WITH_GAMEPLAY_DEBUGGER_MENU
