Perception: Touch: Support team filtering
#tests PIE

diff --git a/Engine/Source/Runtime/AIModule/Classes/Perception/AIPerceptionTypes.h b/Engine/Source/Runtime/AIModule/Classes/Perception/AIPerceptionTypes.h
--- a/Engine/Source/Runtime/AIModule/Classes/Perception/AIPerceptionTypes.h
+++ b/Engine/Source/Runtime/AIModule/Classes/Perception/AIPerceptionTypes.h
@@ -207,19 +207,21 @@ struct FAISenseAffiliationFilter
 {
 	GENERATED_USTRUCT_BODY()
 
-	FAISenseAffiliationFilter()
-		: bDetectEnemies(false)
-		, bDetectNeutrals(false)
-		, bDetectFriendlies(false) {}
+	FAISenseAffiliationFilter() = default;
+	
+	FAISenseAffiliationFilter(bool bInDetectEnemies, bool bInDetectNeutrals, bool bInDetectFriendlies)
+		: bDetectEnemies(bInDetectEnemies)
+		, bDetectNeutrals(bInDetectNeutrals)
+		, bDetectFriendlies(bInDetectFriendlies) {}
 
 	UPROPERTY(EditAnywhere, BlueprintReadOnly, Category = "Sense")
-	uint32 bDetectEnemies : 1;
+	uint32 bDetectEnemies : 1 = false;
 
 	UPROPERTY(EditAnywhere, BlueprintReadOnly, Category = "Sense")
-	uint32 bDetectNeutrals : 1;
+	uint32 bDetectNeutrals : 1 = false;
 
 	UPROPERTY(EditAnywhere, BlueprintReadOnly, Category = "Sense")
-	uint32 bDetectFriendlies : 1;
+	uint32 bDetectFriendlies : 1 = false;
 	
 	uint8 GetAsFlags() const { return IntCastChecked<uint8>((bDetectEnemies << ETeamAttitude::Hostile) | (bDetectNeutrals << ETeamAttitude::Neutral) | (bDetectFriendlies << ETeamAttitude::Friendly)); }
 	FORCEINLINE bool ShouldDetectAll() const { return (bDetectEnemies && bDetectNeutrals && bDetectFriendlies); }
diff --git a/Engine/Source/Runtime/AIModule/Classes/Perception/AISenseConfig_Touch.h b/Engine/Source/Runtime/AIModule/Classes/Perception/AISenseConfig_Touch.h
--- a/Engine/Source/Runtime/AIModule/Classes/Perception/AISenseConfig_Touch.h
+++ b/Engine/Source/Runtime/AIModule/Classes/Perception/AISenseConfig_Touch.h
@@ -13,6 +13,10 @@ UCLASS(meta = (DisplayName = "AI Touch config"), MinimalAPI)
 class UAISenseConfig_Touch : public UAISenseConfig
 {
 	GENERATED_UCLASS_BODY()
-public:	
+public:
+
+	UPROPERTY(EditDefaultsOnly, BlueprintReadOnly, Category = "Sense", config)
+	FAISenseAffiliationFilter DetectionByAffiliation = {true, true, true};
+	
 	AIMODULE_API virtual TSubclassOf<UAISense> GetSenseImplementation() const override;
 };
diff --git a/Engine/Source/Runtime/AIModule/Classes/Perception/AISense_Hearing.h b/Engine/Source/Runtime/AIModule/Classes/Perception/AISense_Hearing.h
--- a/Engine/Source/Runtime/AIModule/Classes/Perception/AISense_Hearing.h
+++ b/Engine/Source/Runtime/AIModule/Classes/Perception/AISense_Hearing.h
@@ -111,5 +111,5 @@ protected:
 
 	AIMODULE_API void OnNewListenerImpl(const FPerceptionListener& NewListener);
 	AIMODULE_API void OnListenerUpdateImpl(const FPerceptionListener& UpdatedListener);
-	AIMODULE_API void OnListenerRemovedImpl(const FPerceptionListener& UpdatedListener);
+	AIMODULE_API void OnListenerRemovedImpl(const FPerceptionListener& RemovedListener);
 };
diff --git a/Engine/Source/Runtime/AIModule/Classes/Perception/AISense_Touch.h b/Engine/Source/Runtime/AIModule/Classes/Perception/AISense_Touch.h
--- a/Engine/Source/Runtime/AIModule/Classes/Perception/AISense_Touch.h
+++ b/Engine/Source/Runtime/AIModule/Classes/Perception/AISense_Touch.h
@@ -9,6 +9,7 @@
 
 class IAIPerceptionListenerInterface;
 class UAISense_Touch;
+class UAISenseConfig_Touch;
 
 USTRUCT()
 struct FAITouchEvent
@@ -17,19 +18,22 @@ struct FAITouchEvent
 
 	typedef UAISense_Touch FSenseClass;
 
-	FVector Location;
+	FVector Location = FVector::ZeroVector;
 	
 	UPROPERTY()
 	TObjectPtr<AActor> TouchReceiver;
 
 	UPROPERTY()
 	TObjectPtr<AActor> OtherActor;
+
+	FGenericTeamId TeamIdentifier = FGenericTeamId::NoTeam;
 		
-	FAITouchEvent() : TouchReceiver(nullptr), OtherActor(nullptr) {}
+	FAITouchEvent() = default;
 	
 	FAITouchEvent(AActor* InTouchReceiver, AActor* InOtherActor, const FVector& EventLocation)
 		: Location(EventLocation), TouchReceiver(InTouchReceiver), OtherActor(InOtherActor)
 	{
+		TeamIdentifier = FGenericTeamId::GetTeamIdentifier(InOtherActor);
 	}
 
 	AIMODULE_API IAIPerceptionListenerInterface* GetTouchedActorAsPerceptionListener() const;
@@ -50,5 +54,20 @@ public:
 	static AIMODULE_API void ReportTouchEvent(UObject* WorldContextObject, AActor* TouchReceiver, AActor* OtherActor, FVector Location);
 
 protected:
+	
+	struct FDigestedTouchProperties
+	{
+		uint8 AffiliationFlags;
+
+		FDigestedTouchProperties(const UAISenseConfig_Touch& SenseConfig);
+		FDigestedTouchProperties();
+	};
+	TMap<FPerceptionListenerID, FDigestedTouchProperties> DigestedProperties;
+
+	
 	AIMODULE_API virtual float Update() override;
+	
+	AIMODULE_API void OnNewListenerImpl(const FPerceptionListener& NewListener);
+	AIMODULE_API void OnListenerUpdateImpl(const FPerceptionListener& UpdatedListener);
+	AIMODULE_API void OnListenerRemovedImpl(const FPerceptionListener& RemovedListener);
 };
diff --git a/Engine/Source/Runtime/AIModule/Private/Perception/AISense_Hearing.cpp b/Engine/Source/Runtime/AIModule/Private/Perception/AISense_Hearing.cpp
--- a/Engine/Source/Runtime/AIModule/Private/Perception/AISense_Hearing.cpp
+++ b/Engine/Source/Runtime/AIModule/Private/Perception/AISense_Hearing.cpp
@@ -43,9 +43,9 @@ UAISense_Hearing::FDigestedHearingProperties::FDigestedHearingProperties(const U
 }
 
 UAISense_Hearing::FDigestedHearingProperties::FDigestedHearingProperties()
-	: HearingRangeSq(-1.f), AffiliationFlags(-1)
+	: HearingRangeSq(-1.f)
 {
-
+	AffiliationFlags = FAISenseAffiliationFilter::DetectAllFlags();
 }
 
 //----------------------------------------------------------------------//
@@ -115,15 +115,15 @@ void UAISense_Hearing::OnListenerUpdateImpl(const FPerceptionListener& UpdatedLi
 	}
 }
 
-void UAISense_Hearing::OnListenerRemovedImpl(const FPerceptionListener& UpdatedListener)
+void UAISense_Hearing::OnListenerRemovedImpl(const FPerceptionListener& RemovedListener)
 {
-	DigestedProperties.FindAndRemoveChecked(UpdatedListener.GetListenerID());
+	DigestedProperties.FindAndRemoveChecked(RemovedListener.GetListenerID());
 }
 
 float UAISense_Hearing::Update()
 {
 	AIPerception::FListenerMap& ListenersMap = *GetListeners();
-	UAIPerceptionSystem* PerseptionSys = GetPerceptionSystem();
+	UAIPerceptionSystem* PerceptionSys = GetPerceptionSystem();
 	const float SpeedOfSoundSqScalar = SpeedOfSoundSq > 0.f ? 1.f / SpeedOfSoundSq : 0.f;
 
 	for (AIPerception::FListenerMap::TIterator ListenerIt(ListenersMap); ListenerIt; ++ListenerIt)
@@ -161,7 +161,7 @@ float UAISense_Hearing::Update()
 			// calculate delay and fake it with Age, Delay should be pretty small so a static_cast is safe enough here.
 			const float Delay = FloatCastChecked<float>(FMath::Sqrt(DistToSoundSquared * SpeedOfSoundSqScalar), UE::LWC::DefaultFloatPrecision);
 			// pass over to listener to process 			
-			PerseptionSys->RegisterDelayedStimulus(Listener.GetListenerID(), Delay, Event.Instigator
+			PerceptionSys->RegisterDelayedStimulus(Listener.GetListenerID(), Delay, Event.Instigator
 				, FAIStimulus(*this, ClampedLoudness, Event.NoiseLocation, Listener.CachedLocation, FAIStimulus::SensingSucceeded, Event.Tag) );
 		}
 	}
diff --git a/Engine/Source/Runtime/AIModule/Private/Perception/AISense_Sight.cpp b/Engine/Source/Runtime/AIModule/Private/Perception/AISense_Sight.cpp
--- a/Engine/Source/Runtime/AIModule/Private/Perception/AISense_Sight.cpp
+++ b/Engine/Source/Runtime/AIModule/Private/Perception/AISense_Sight.cpp
@@ -42,12 +42,12 @@ DECLARE_CYCLE_STAT(TEXT("Perception Sense: Sight, Remove To Target"), STAT_AI_Se
 DECLARE_CYCLE_STAT(TEXT("Perception Sense: Sight, Process pending result"), STAT_AI_Sense_Sight_ProcessPendingQuery, STATGROUP_AI);
 
 
-static const int32 DefaultMaxTracesPerTick = 6;
-static const int32 DefaultMaxAsyncTracesPerTick = 10;
-static const int32 DefaultMinQueriesPerTimeSliceCheck = 40;
-static const float DefaultPendingQueriesBudgetReductionRatio = 0.5f;
-static const bool bDefaultUseAsynchronousTraceForDefaultSightQueries = false;
-static const float DefaultStimulusStrength = 1.f;
+constexpr int32 DefaultMaxTracesPerTick = 6;
+constexpr int32 DefaultMaxAsyncTracesPerTick = 10;
+constexpr int32 DefaultMinQueriesPerTimeSliceCheck = 40;
+constexpr float DefaultPendingQueriesBudgetReductionRatio = 0.5f;
+constexpr bool bDefaultUseAsynchronousTraceForDefaultSightQueries = false;
+constexpr float DefaultStimulusStrength = 1.f;
 
 enum class EForEachResult : uint8
 {
@@ -94,7 +94,7 @@ EReverseForEachResult ReverseForEach(T& Array, const PREDICATE_CLASS& Predicate)
 const FAISightTarget::FTargetId FAISightTarget::InvalidTargetId = FAISystem::InvalidUnsignedID;
 
 FAISightTarget::FAISightTarget(AActor* InTarget, FGenericTeamId InTeamId)
-	: Target(InTarget), SightTargetInterface(NULL), TeamId(InTeamId)
+	: Target(InTarget), SightTargetInterface(nullptr), TeamId(InTeamId)
 {
 	if (InTarget)
 	{
@@ -122,8 +122,10 @@ UAISense_Sight::FDigestedSightProperties::FDigestedSightProperties(const UAISens
 }
 
 UAISense_Sight::FDigestedSightProperties::FDigestedSightProperties()
-	: PeripheralVisionAngleCos(0.f), SightRadiusSq(-1.f), AutoSuccessRangeSqFromLastSeenLocation(FAISystem::InvalidRange), LoseSightRadiusSq(-1.f), PointOfViewBackwardOffset(0.0f), NearClippingRadiusSq(0.0f), AffiliationFlags(-1)
-{}
+	: PeripheralVisionAngleCos(0.f), SightRadiusSq(-1.f), AutoSuccessRangeSqFromLastSeenLocation(FAISystem::InvalidRange), LoseSightRadiusSq(-1.f), PointOfViewBackwardOffset(0.0f), NearClippingRadiusSq(0.0f)
+{
+	AffiliationFlags = FAISenseAffiliationFilter::DetectAllFlags();
+}
 
 //----------------------------------------------------------------------//
 // UAISense_Sight
@@ -161,7 +163,7 @@ UAISense_Sight::UAISense_Sight(const FObjectInitializer& ObjectInitializer)
 	DefaultSightCollisionChannel = GET_AI_CONFIG_VAR(DefaultSightCollisionChannel);
 }
 
-FORCEINLINE_DEBUGGABLE float UAISense_Sight::CalcQueryImportance(const FPerceptionListener& Listener, const FVector& TargetLocation, const float SightRadiusSq) const
+float UAISense_Sight::CalcQueryImportance(const FPerceptionListener& Listener, const FVector& TargetLocation, const float SightRadiusSq) const
 {
 	const FVector::FReal DistanceSq = FVector::DistSquared(Listener.CachedLocation, TargetLocation);
 	return DistanceSq <= HighImportanceDistanceSquare ? MaxQueryImportance
@@ -295,7 +297,7 @@ float UAISense_Sight::Update()
 #if AISENSE_SIGHT_TIMESLICING_DEBUG
 	UE::AISense_Sight::FTimingSlicingInfo SlicingInfo;
 #endif // AISENSE_SIGHT_TIMESLICING_DEBUG
-	static const int32 InitialInvalidItemsSize = 16;
+	constexpr int32 InitialInvalidItemsSize = 16;
 	enum class EOperationType : uint8
 	{
 		Remove,
@@ -434,7 +436,7 @@ float UAISense_Sight::Update()
 		});
         // Do all the removes first and save the out of range swaps because we will insert them at the right location to prevent sorting
 		TArray<FAISightQuery> SightQueriesOutOfRangeToInsert;
-		for (FQueryOperation& Operation : QueryOperations)
+		for (const FQueryOperation& Operation : QueryOperations)
 		{
 			switch (Operation.OpType)
 			{
@@ -913,7 +915,7 @@ void UAISense_Sight::OnListenerUpdateImpl(const FPerceptionListener& UpdatedList
 	// see if this listener is a Target as well
 	const FAISightTarget::FTargetId AsTargetId = UpdatedListener.GetBodyActorUniqueID();
 	FAISightTarget* AsTarget = ObservedTargets.Find(AsTargetId);
-	if (AsTarget != NULL)
+	if (AsTarget != nullptr)
 	{
 		if (AsTarget->Target.IsValid())
 		{
diff --git a/Engine/Source/Runtime/AIModule/Private/Perception/AISense_Touch.cpp b/Engine/Source/Runtime/AIModule/Private/Perception/AISense_Touch.cpp
--- a/Engine/Source/Runtime/AIModule/Private/Perception/AISense_Touch.cpp
+++ b/Engine/Source/Runtime/AIModule/Private/Perception/AISense_Touch.cpp
@@ -2,13 +2,15 @@
 
 #include "Perception/AISense_Touch.h"
 #include "GameFramework/Pawn.h"
-#include "GameFramework/Controller.h"
 #include "Perception/AIPerceptionListenerInterface.h"
 #include "Perception/AIPerceptionComponent.h"
+#include "Perception/AISenseConfig_Touch.h"
 
 #include UE_INLINE_GENERATED_CPP_BY_NAME(AISense_Touch)
 
-
+//----------------------------------------------------------------------//
+// FAITouchEvent
+//----------------------------------------------------------------------//
 IAIPerceptionListenerInterface* FAITouchEvent::GetTouchedActorAsPerceptionListener() const
 {
 	IAIPerceptionListenerInterface* Listener = nullptr;
@@ -17,7 +19,7 @@ IAIPerceptionListenerInterface* FAITouchEvent::GetTouchedActorAsPerceptionListen
 		Listener = Cast<IAIPerceptionListenerInterface>(TouchReceiver);
 		if (Listener == nullptr)
 		{
-			APawn* ListenerAsPawn = Cast<APawn>(TouchReceiver);
+			const APawn* ListenerAsPawn = Cast<APawn>(TouchReceiver);
 			if (ListenerAsPawn)
 			{
 				Listener = Cast<IAIPerceptionListenerInterface>(ListenerAsPawn->GetController());
@@ -27,9 +29,63 @@ IAIPerceptionListenerInterface* FAITouchEvent::GetTouchedActorAsPerceptionListen
 	return Listener;
 }
 
+//----------------------------------------------------------------------//
+// FDigestedHearingProperties
+//----------------------------------------------------------------------//
+UAISense_Touch::FDigestedTouchProperties::FDigestedTouchProperties(const UAISenseConfig_Touch& SenseConfig)
+{
+	AffiliationFlags = SenseConfig.DetectionByAffiliation.GetAsFlags();
+}
+
+UAISense_Touch::FDigestedTouchProperties::FDigestedTouchProperties()
+{
+	AffiliationFlags = FAISenseAffiliationFilter::DetectAllFlags();
+}
+
+//----------------------------------------------------------------------//
+// UAISense_Touch
+//----------------------------------------------------------------------//
 UAISense_Touch::UAISense_Touch(const FObjectInitializer& ObjectInitializer) :
 	Super(ObjectInitializer)
 {
+	if (HasAnyFlags(RF_ClassDefaultObject) == false)
+	{
+		OnNewListenerDelegate.BindUObject(this, &UAISense_Touch::OnNewListenerImpl);
+		OnListenerUpdateDelegate.BindUObject(this, &UAISense_Touch::OnListenerUpdateImpl);
+		OnListenerRemovedDelegate.BindUObject(this, &UAISense_Touch::OnListenerRemovedImpl);
+	}
+}
+
+void UAISense_Touch::OnNewListenerImpl(const FPerceptionListener& NewListener)
+{
+	UAIPerceptionComponent* ListenerPtr = NewListener.Listener.Get();
+	check(ListenerPtr);
+	const UAISenseConfig_Touch* SenseConfig = Cast<const UAISenseConfig_Touch>(ListenerPtr->GetSenseConfig(GetSenseID()));
+	check(SenseConfig);
+	const FDigestedTouchProperties PropertyDigest(*SenseConfig);
+	DigestedProperties.Add(NewListener.GetListenerID(), PropertyDigest);
+}
+
+void UAISense_Touch::OnListenerUpdateImpl(const FPerceptionListener& UpdatedListener)
+{
+	const FPerceptionListenerID ListenerID = UpdatedListener.GetListenerID();
+	
+	if (UpdatedListener.HasSense(GetSenseID()))
+	{
+		const UAISenseConfig_Touch* SenseConfig = Cast<const UAISenseConfig_Touch>(UpdatedListener.Listener->GetSenseConfig(GetSenseID()));
+		check(SenseConfig);
+		FDigestedTouchProperties& PropertiesDigest = DigestedProperties.FindOrAdd(ListenerID);
+		PropertiesDigest = FDigestedTouchProperties(*SenseConfig);
+	}
+	else
+	{
+		DigestedProperties.Remove(ListenerID);
+	}
+}
+
+void UAISense_Touch::OnListenerRemovedImpl(const FPerceptionListener& RemovedListener)
+{
+	DigestedProperties.FindAndRemoveChecked(RemovedListener.GetListenerID());
 }
 
 float UAISense_Touch::Update()
@@ -38,16 +94,21 @@ float UAISense_Touch::Update()
 
 	for (const FAITouchEvent& Event : RegisteredEvents)
 	{
-		if (Event.TouchReceiver != NULL && Event.OtherActor != NULL)
+		if (Event.TouchReceiver != nullptr && Event.OtherActor != nullptr)
 		{
 			IAIPerceptionListenerInterface* PerceptionListener = Event.GetTouchedActorAsPerceptionListener();
-			if (PerceptionListener != NULL)
+			if (PerceptionListener != nullptr)
 			{
-				UAIPerceptionComponent* PerceptionComponent = PerceptionListener->GetPerceptionComponent();
-				if (PerceptionComponent != NULL && ListenersMap.Contains(PerceptionComponent->GetListenerId()))
+				const UAIPerceptionComponent* PerceptionComponent = PerceptionListener->GetPerceptionComponent();
+				if (PerceptionComponent != nullptr && ListenersMap.Contains(PerceptionComponent->GetListenerId()))
 				{
 					// this has to succeed, will assert a failure
 					FPerceptionListener& Listener = ListenersMap[PerceptionComponent->GetListenerId()];
+					const FDigestedTouchProperties& PropDigest = DigestedProperties[Listener.GetListenerID()];
+					if (FAISenseAffiliationFilter::ShouldSenseTeam(Listener.TeamIdentifier, Event.TeamIdentifier, PropDigest.AffiliationFlags) == false)
+					{
+						continue;
+					}
 					if (Listener.HasSense(GetSenseID()))
 					{
 						Listener.RegisterStimulus(Event.OtherActor, FAIStimulus(*this, 1.f, Event.Location, Event.Location));
@@ -75,7 +136,7 @@ void UAISense_Touch::ReportTouchEvent(UObject* WorldContextObject, AActor* Touch
 	UAIPerceptionSystem* PerceptionSystem = UAIPerceptionSystem::GetCurrent(WorldContextObject);
 	if (PerceptionSystem)
 	{
-		FAITouchEvent Event(TouchReceiver, OtherActor, Location);
+		const FAITouchEvent Event(TouchReceiver, OtherActor, Location);
 		PerceptionSystem->OnEvent(Event);
 	}
 }
