Add missing include
#rnx
#tests PIE

diff --git a/Engine/Source/Runtime/AIModule/Private/Perception/AISense_Touch.cpp b/Engine/Source/Runtime/AIModule/Private/Perception/AISense_Touch.cpp
--- a/Engine/Source/Runtime/AIModule/Private/Perception/AISense_Touch.cpp
+++ b/Engine/Source/Runtime/AIModule/Private/Perception/AISense_Touch.cpp
@@ -1,6 +1,7 @@
 // Copyright Epic Games, Inc. All Rights Reserved.
 
 #include "Perception/AISense_Touch.h"
+#include "GameFramework/Controller.h"
 #include "GameFramework/Pawn.h"
 #include "Perception/AIPerceptionListenerInterface.h"
 #include "Perception/AIPerceptionComponent.h"
