Minor edition improvement for EnvQueryGenerator_OnCircle: Use "EditCondition" to indicate which of SpaceBetween and NumberOfPoints will be used, depending on the PointOnCircleSpacingMethod selected

diff --git a/Engine/Source/Runtime/AIModule/Classes/EnvironmentQuery/Generators/EnvQueryGenerator_OnCircle.h b/Engine/Source/Runtime/AIModule/Classes/EnvironmentQuery/Generators/EnvQueryGenerator_OnCircle.h
--- a/Engine/Source/Runtime/AIModule/Classes/EnvironmentQuery/Generators/EnvQueryGenerator_OnCircle.h
+++ b/Engine/Source/Runtime/AIModule/Classes/EnvironmentQuery/Generators/EnvQueryGenerator_OnCircle.h
@@ -33,18 +33,18 @@ class UEnvQueryGenerator_OnCircle : public UEnvQueryGenerator_ProjectedPoints
 	UPROPERTY(EditDefaultsOnly, Category=Generator)
 	FAIDataProviderFloatValue CircleRadius;
 
-	/** items will be generated on a circle this much apart */
+	/** how we are choosing where the points are in the circle */
 	UPROPERTY(EditDefaultsOnly, Category = Generator)
+	EPointOnCircleSpacingMethod PointOnCircleSpacingMethod;
+
+	/** items will be generated on a circle this much apart */
+	UPROPERTY(EditDefaultsOnly, Category = Generator, meta = (EditCondition = "PointOnCircleSpacingMethod == EPointOnCircleSpacingMethod::BySpaceBetween", EditConditionHides))
 	FAIDataProviderFloatValue SpaceBetween;
 
 	/** this many items will be generated on a circle */
-	UPROPERTY(EditDefaultsOnly, Category = Generator)
+	UPROPERTY(EditDefaultsOnly, Category = Generator, meta = (EditCondition = "PointOnCircleSpacingMethod == EPointOnCircleSpacingMethod::ByNumberOfPoints", EditConditionHides))
 	FAIDataProviderIntValue NumberOfPoints;
 
-	/** how we are choosing where the points are in the circle */
-	UPROPERTY(EditDefaultsOnly, Category = Generator)
-	EPointOnCircleSpacingMethod PointOnCircleSpacingMethod;
-
 	/** If you generate items on a piece of circle you define direction of Arc cut here */
 	UPROPERTY(EditDefaultsOnly, Category=Generator, meta=(EditCondition="bDefineArc"))
 	FEnvDirection ArcDirection;
