Sense config: add units
#rnx

diff --git a/Engine/Source/Runtime/AIModule/Classes/Perception/AISenseConfig.h b/Engine/Source/Runtime/AIModule/Classes/Perception/AISenseConfig.h
--- a/Engine/Source/Runtime/AIModule/Classes/Perception/AISenseConfig.h
+++ b/Engine/Source/Runtime/AIModule/Classes/Perception/AISenseConfig.h
@@ -24,7 +24,7 @@ protected:
 	FColor DebugColor;
 
 	/** specifies age limit after stimuli generated by this sense become forgotten. 0 means "never" */
-	UPROPERTY(EditDefaultsOnly, BlueprintReadOnly, Category = Sense, meta = (ClampMin = "0.0", UIMin = "0.0"))
+	UPROPERTY(EditDefaultsOnly, BlueprintReadOnly, Category = Sense, meta = (ClampMin = "0.0", UIMin = "0.0", Units="Seconds"))
 	float MaxAge;
 
 	/** determines whether given sense starts in an enabled state */
diff --git a/Engine/Source/Runtime/AIModule/Classes/Perception/AISenseConfig_Hearing.h b/Engine/Source/Runtime/AIModule/Classes/Perception/AISenseConfig_Hearing.h
--- a/Engine/Source/Runtime/AIModule/Classes/Perception/AISenseConfig_Hearing.h
+++ b/Engine/Source/Runtime/AIModule/Classes/Perception/AISenseConfig_Hearing.h
@@ -23,11 +23,11 @@ public:
 	UPROPERTY(EditDefaultsOnly, BlueprintReadOnly, Category = "Sense", NoClear, config)
 	TSubclassOf<UAISense_Hearing> Implementation;
 
-	UPROPERTY(EditDefaultsOnly, BlueprintReadOnly, Category = "Sense")
+	UPROPERTY(EditDefaultsOnly, BlueprintReadOnly, Category = "Sense", meta = (UIMin = 0.0, ClampMin = 0.0, Units="Centimeters"))
 	float HearingRange;
 
 	UE_DEPRECATED(5.2, "LoSHearingRange is deprecated. Use HearingRange instead.")
-	UPROPERTY(EditDefaultsOnly, BlueprintReadOnly, Category = "Sense", meta = (EditCondition = "bUseLoSHearing"))
+	UPROPERTY(EditDefaultsOnly, BlueprintReadOnly, Category = "Sense", meta = (UIMin = 0.0, ClampMin = 0.0, Units="Centimeters", EditCondition = "bUseLoSHearing"))
 	float LoSHearingRange;
 
 	UE_DEPRECATED(5.2, "bUseLoSHearing is deprecated.")
@@ -40,6 +40,6 @@ public:
 	AIMODULE_API virtual TSubclassOf<UAISense> GetSenseImplementation() const override;
 
 #if WITH_GAMEPLAY_DEBUGGER_MENU
-	AIMODULE_API virtual void DescribeSelfToGameplayDebugger(const UAIPerceptionComponent* PerceptionComponent, FGameplayDebuggerCategory* DebuggerCategory) const;
+	AIMODULE_API virtual void DescribeSelfToGameplayDebugger(const UAIPerceptionComponent* PerceptionComponent, FGameplayDebuggerCategory* DebuggerCategory) const override;
 #endif // WITH_GAMEPLAY_DEBUGGER_MENU
 };
diff --git a/Engine/Source/Runtime/AIModule/Classes/Perception/AISenseConfig_Sight.h b/Engine/Source/Runtime/AIModule/Classes/Perception/AISenseConfig_Sight.h
--- a/Engine/Source/Runtime/AIModule/Classes/Perception/AISenseConfig_Sight.h
+++ b/Engine/Source/Runtime/AIModule/Classes/Perception/AISenseConfig_Sight.h
@@ -24,16 +24,16 @@ public:
 	TSubclassOf<UAISense_Sight> Implementation;
 
 	/** Maximum sight distance to notice a target. */
-	UPROPERTY(EditDefaultsOnly, BlueprintReadOnly, Category = "Sense", config, meta = (UIMin = 0.0, ClampMin = 0.0))
+	UPROPERTY(EditDefaultsOnly, BlueprintReadOnly, Category = "Sense", config, meta = (UIMin = 0.0, ClampMin = 0.0, Units="Centimeters"))
 	float SightRadius;
 
 	/** Maximum sight distance to see target that has been already seen. */
-	UPROPERTY(EditDefaultsOnly, BlueprintReadOnly, Category = "Sense", config, meta = (UIMin = 0.0, ClampMin = 0.0))
+	UPROPERTY(EditDefaultsOnly, BlueprintReadOnly, Category = "Sense", config, meta = (UIMin = 0.0, ClampMin = 0.0, Units="Centimeters"))
 	float LoseSightRadius;
 
 	/** How far to the side AI can see, in degrees. Use SetPeripheralVisionAngle to change the value at runtime. 
 	 *	The value represents the angle measured in relation to the forward vector, not the whole range. */
-	UPROPERTY(EditDefaultsOnly, BlueprintReadOnly, Category = "Sense", config, meta=(UIMin = 0.0, ClampMin = 0.0, UIMax = 180.0, ClampMax = 180.0, DisplayName="PeripheralVisionHalfAngleDegrees"))
+	UPROPERTY(EditDefaultsOnly, BlueprintReadOnly, Category = "Sense", config, meta=(UIMin = 0.0, ClampMin = 0.0, UIMax = 180.0, ClampMax = 180.0, DisplayName="Peripheral Vision Half Angle", Units="Degrees"))
 	float PeripheralVisionAngleDegrees;
 
 	/** */
@@ -41,15 +41,15 @@ public:
 	FAISenseAffiliationFilter DetectionByAffiliation;
 
 	/** If not an InvalidRange (which is the default), we will always be able to see the target that has already been seen if they are within this range of their last seen location. */
-	UPROPERTY(EditDefaultsOnly, BlueprintReadOnly, Category = "Sense", config)
+	UPROPERTY(EditDefaultsOnly, BlueprintReadOnly, Category = "Sense", config, meta = (Units="Centimeters"))
 	float AutoSuccessRangeFromLastSeenLocation;
 
 	/** Point of view move back distance for cone calculation. In conjunction with near clipping distance, this will act as a close by awareness and peripheral vision. */
-	UPROPERTY(EditDefaultsOnly, BlueprintReadOnly, Category = "Sense", config, meta = (UIMin = 0.0, ClampMin = 0.0))
+	UPROPERTY(EditDefaultsOnly, BlueprintReadOnly, Category = "Sense", config, meta = (UIMin = 0.0, ClampMin = 0.0, Units="Centimeters"))
 	float PointOfViewBackwardOffset;
 
 	/** Near clipping distance, to be used with point of view backward offset. Will act as a close by awareness and peripheral vision */
-	UPROPERTY(EditDefaultsOnly, BlueprintReadOnly, Category = "Sense", config, meta = (UIMin = 0.0, ClampMin = 0.0))
+	UPROPERTY(EditDefaultsOnly, BlueprintReadOnly, Category = "Sense", config, meta = (UIMin = 0.0, ClampMin = 0.0, Units="Centimeters"))
 	float NearClippingRadius;
 
 	AIMODULE_API virtual TSubclassOf<UAISense> GetSenseImplementation() const override;
@@ -59,6 +59,6 @@ public:
 #endif // WITH_EDITOR
 
 #if WITH_GAMEPLAY_DEBUGGER_MENU
-	AIMODULE_API virtual void DescribeSelfToGameplayDebugger(const UAIPerceptionComponent* PerceptionComponent, FGameplayDebuggerCategory* DebuggerCategory) const;
+	AIMODULE_API virtual void DescribeSelfToGameplayDebugger(const UAIPerceptionComponent* PerceptionComponent, FGameplayDebuggerCategory* DebuggerCategory) const override;
 #endif // WITH_GAMEPLAY_DEBUGGER_MENU
 };
