Sense config variables can be edited on Instances
#tests PIE

diff --git a/Engine/Source/Runtime/AIModule/Classes/Perception/AISenseConfig.h b/Engine/Source/Runtime/AIModule/Classes/Perception/AISenseConfig.h
--- a/Engine/Source/Runtime/AIModule/Classes/Perception/AISenseConfig.h
+++ b/Engine/Source/Runtime/AIModule/Classes/Perception/AISenseConfig.h
@@ -24,11 +24,11 @@ protected:
 	FColor DebugColor;
 
 	/** specifies age limit after stimuli generated by this sense become forgotten. 0 means "never" */
-	UPROPERTY(EditDefaultsOnly, BlueprintReadOnly, Category = Sense, meta = (ClampMin = "0.0", UIMin = "0.0", Units="Seconds"))
+	UPROPERTY(EditAnywhere, BlueprintReadOnly, Category = Sense, meta = (ClampMin = "0.0", UIMin = "0.0", Units="Seconds"))
 	float MaxAge;
 
 	/** determines whether given sense starts in an enabled state */
-	UPROPERTY(EditDefaultsOnly, BlueprintReadOnly, Category = Sense)
+	UPROPERTY(EditAnywhere, BlueprintReadOnly, Category = Sense)
 	uint32 bStartsEnabled : 1;
 	
 	// name of sense
diff --git a/Engine/Source/Runtime/AIModule/Classes/Perception/AISenseConfig_Blueprint.h b/Engine/Source/Runtime/AIModule/Classes/Perception/AISenseConfig_Blueprint.h
--- a/Engine/Source/Runtime/AIModule/Classes/Perception/AISenseConfig_Blueprint.h
+++ b/Engine/Source/Runtime/AIModule/Classes/Perception/AISenseConfig_Blueprint.h
@@ -17,7 +17,7 @@ class UAISenseConfig_Blueprint : public UAISenseConfig
 	GENERATED_UCLASS_BODY()
 
 public:
-	UPROPERTY(EditDefaultsOnly, BlueprintReadOnly, Category = "Sense", NoClear, config)
+	UPROPERTY(EditAnywhere, BlueprintReadOnly, Category = "Sense", NoClear, config)
 	TSubclassOf<UAISense_Blueprint> Implementation;
 
 	AIMODULE_API virtual TSubclassOf<UAISense> GetSenseImplementation() const override;
diff --git a/Engine/Source/Runtime/AIModule/Classes/Perception/AISenseConfig_Damage.h b/Engine/Source/Runtime/AIModule/Classes/Perception/AISenseConfig_Damage.h
--- a/Engine/Source/Runtime/AIModule/Classes/Perception/AISenseConfig_Damage.h
+++ b/Engine/Source/Runtime/AIModule/Classes/Perception/AISenseConfig_Damage.h
@@ -16,7 +16,7 @@ class UAISenseConfig_Damage : public UAISenseConfig
 	GENERATED_UCLASS_BODY()
 
 public:
-	UPROPERTY(EditDefaultsOnly, BlueprintReadOnly, Category = "Sense", NoClear, config)
+	UPROPERTY(EditAnywhere, BlueprintReadOnly, Category = "Sense", NoClear, config)
 	TSubclassOf<UAISense_Damage> Implementation;
 
 	AIMODULE_API virtual TSubclassOf<UAISense> GetSenseImplementation() const override;
diff --git a/Engine/Source/Runtime/AIModule/Classes/Perception/AISenseConfig_Hearing.h b/Engine/Source/Runtime/AIModule/Classes/Perception/AISenseConfig_Hearing.h
--- a/Engine/Source/Runtime/AIModule/Classes/Perception/AISenseConfig_Hearing.h
+++ b/Engine/Source/Runtime/AIModule/Classes/Perception/AISenseConfig_Hearing.h
@@ -20,21 +20,21 @@ class UAISenseConfig_Hearing : public UAISenseConfig
 	GENERATED_UCLASS_BODY()
 
 public:
-	UPROPERTY(EditDefaultsOnly, BlueprintReadOnly, Category = "Sense", NoClear, config)
+	UPROPERTY(EditAnywhere, BlueprintReadOnly, Category = "Sense", NoClear, config)
 	TSubclassOf<UAISense_Hearing> Implementation;
 
-	UPROPERTY(EditDefaultsOnly, BlueprintReadOnly, Category = "Sense", meta = (UIMin = 0.0, ClampMin = 0.0, Units="Centimeters"))
+	UPROPERTY(EditAnywhere, BlueprintReadOnly, Category = "Sense", meta = (UIMin = 0.0, ClampMin = 0.0, Units="Centimeters"))
 	float HearingRange;
 
 	UE_DEPRECATED(5.2, "LoSHearingRange is deprecated. Use HearingRange instead.")
-	UPROPERTY(EditDefaultsOnly, BlueprintReadOnly, Category = "Sense", meta = (UIMin = 0.0, ClampMin = 0.0, Units="Centimeters", EditCondition = "bUseLoSHearing"))
+	UPROPERTY(EditAnywhere, BlueprintReadOnly, Category = "Sense", meta = (UIMin = 0.0, ClampMin = 0.0, Units="Centimeters", EditCondition = "bUseLoSHearing"))
 	float LoSHearingRange;
 
 	UE_DEPRECATED(5.2, "bUseLoSHearing is deprecated.")
-	UPROPERTY(EditDefaultsOnly, Category = "Sense", meta = (InlineEditConditionToggle))
+	UPROPERTY(EditAnywhere, Category = "Sense", meta = (InlineEditConditionToggle))
 	uint32 bUseLoSHearing : 1;
 
-	UPROPERTY(EditDefaultsOnly, BlueprintReadOnly, Category = "Sense", config)
+	UPROPERTY(EditAnywhere, BlueprintReadOnly, Category = "Sense", config)
 	FAISenseAffiliationFilter DetectionByAffiliation;
 
 	AIMODULE_API virtual TSubclassOf<UAISense> GetSenseImplementation() const override;
diff --git a/Engine/Source/Runtime/AIModule/Classes/Perception/AISenseConfig_Sight.h b/Engine/Source/Runtime/AIModule/Classes/Perception/AISenseConfig_Sight.h
--- a/Engine/Source/Runtime/AIModule/Classes/Perception/AISenseConfig_Sight.h
+++ b/Engine/Source/Runtime/AIModule/Classes/Perception/AISenseConfig_Sight.h
@@ -20,36 +20,36 @@ class UAISenseConfig_Sight : public UAISenseConfig
 	GENERATED_UCLASS_BODY()
 
 public:
-	UPROPERTY(EditDefaultsOnly, BlueprintReadOnly, Category = "Sense", NoClear, config)
+	UPROPERTY(EditAnywhere, BlueprintReadOnly, Category = "Sense", NoClear, config)
 	TSubclassOf<UAISense_Sight> Implementation;
 
 	/** Maximum sight distance to notice a target. */
-	UPROPERTY(EditDefaultsOnly, BlueprintReadOnly, Category = "Sense", config, meta = (UIMin = 0.0, ClampMin = 0.0, Units="Centimeters"))
+	UPROPERTY(EditAnywhere, BlueprintReadOnly, Category = "Sense", config, meta = (UIMin = 0.0, ClampMin = 0.0, Units="Centimeters"))
 	float SightRadius;
 
 	/** Maximum sight distance to see target that has been already seen. */
-	UPROPERTY(EditDefaultsOnly, BlueprintReadOnly, Category = "Sense", config, meta = (UIMin = 0.0, ClampMin = 0.0, Units="Centimeters"))
+	UPROPERTY(EditAnywhere, BlueprintReadOnly, Category = "Sense", config, meta = (UIMin = 0.0, ClampMin = 0.0, Units="Centimeters"))
 	float LoseSightRadius;
 
 	/** How far to the side AI can see, in degrees. Use SetPeripheralVisionAngle to change the value at runtime. 
 	 *	The value represents the angle measured in relation to the forward vector, not the whole range. */
-	UPROPERTY(EditDefaultsOnly, BlueprintReadOnly, Category = "Sense", config, meta=(UIMin = 0.0, ClampMin = 0.0, UIMax = 180.0, ClampMax = 180.0, DisplayName="Peripheral Vision Half Angle", Units="Degrees"))
+	UPROPERTY(EditAnywhere, BlueprintReadOnly, Category = "Sense", config, meta=(UIMin = 0.0, ClampMin = 0.0, UIMax = 180.0, ClampMax = 180.0, DisplayName="Peripheral Vision Half Angle", Units="Degrees"))
 	float PeripheralVisionAngleDegrees;
 
 	/** */
-	UPROPERTY(EditDefaultsOnly, BlueprintReadOnly, Category = "Sense", config)
+	UPROPERTY(EditAnywhere, BlueprintReadOnly, Category = "Sense", config)
 	FAISenseAffiliationFilter DetectionByAffiliation;
 
 	/** If not an InvalidRange (which is the default), we will always be able to see the target that has already been seen if they are within this range of their last seen location. */
-	UPROPERTY(EditDefaultsOnly, BlueprintReadOnly, Category = "Sense", config, meta = (Units="Centimeters"))
+	UPROPERTY(EditAnywhere, BlueprintReadOnly, Category = "Sense", config, meta = (Units="Centimeters"))
 	float AutoSuccessRangeFromLastSeenLocation;
 
 	/** Point of view move back distance for cone calculation. In conjunction with near clipping distance, this will act as a close by awareness and peripheral vision. */
-	UPROPERTY(EditDefaultsOnly, BlueprintReadOnly, Category = "Sense", config, meta = (UIMin = 0.0, ClampMin = 0.0, Units="Centimeters"))
+	UPROPERTY(EditAnywhere, BlueprintReadOnly, Category = "Sense", config, meta = (UIMin = 0.0, ClampMin = 0.0, Units="Centimeters"))
 	float PointOfViewBackwardOffset;
 
 	/** Near clipping distance, to be used with point of view backward offset. Will act as a close by awareness and peripheral vision */
-	UPROPERTY(EditDefaultsOnly, BlueprintReadOnly, Category = "Sense", config, meta = (UIMin = 0.0, ClampMin = 0.0, Units="Centimeters"))
+	UPROPERTY(EditAnywhere, BlueprintReadOnly, Category = "Sense", config, meta = (UIMin = 0.0, ClampMin = 0.0, Units="Centimeters"))
 	float NearClippingRadius;
 
 	AIMODULE_API virtual TSubclassOf<UAISense> GetSenseImplementation() const override;
diff --git a/Engine/Source/Runtime/AIModule/Classes/Perception/AISenseConfig_Touch.h b/Engine/Source/Runtime/AIModule/Classes/Perception/AISenseConfig_Touch.h
--- a/Engine/Source/Runtime/AIModule/Classes/Perception/AISenseConfig_Touch.h
+++ b/Engine/Source/Runtime/AIModule/Classes/Perception/AISenseConfig_Touch.h
@@ -15,7 +15,7 @@ class UAISenseConfig_Touch : public UAISenseConfig
 	GENERATED_UCLASS_BODY()
 public:
 
-	UPROPERTY(EditDefaultsOnly, BlueprintReadOnly, Category = "Sense", config)
+	UPROPERTY(EditAnywhere, BlueprintReadOnly, Category = "Sense", config)
 	FAISenseAffiliationFilter DetectionByAffiliation = {true, true, true};
 	
 	AIMODULE_API virtual TSubclassOf<UAISense> GetSenseImplementation() const override;
