Added client movement mode and base to the gameplay debugger:
- huge help in debugging the not-so-rare bugs with mismatch between client and server on the character movement's mode and pawn's base

diff --git a/Engine/Source/Runtime/AIModule/Private/GameplayDebugger/GameplayDebuggerCategory_AI.cpp b/Engine/Source/Runtime/AIModule/Private/GameplayDebugger/GameplayDebuggerCategory_AI.cpp
--- a/Engine/Source/Runtime/AIModule/Private/GameplayDebugger/GameplayDebuggerCategory_AI.cpp
+++ b/Engine/Source/Runtime/AIModule/Private/GameplayDebugger/GameplayDebuggerCategory_AI.cpp
@@ -338,7 +338,27 @@ void FGameplayDebuggerCategory_AI::DrawData(APlayerController* OwnerPC, FGamepla
 
 	if (DataPack.bIsUsingCharacter)
 	{
-		CanvasContext.Printf(TEXT("Movement Mode: {yellow}%s{white}, Base: {yellow}%s"), *DataPack.MovementModeInfo, *DataPack.MovementBaseInfo);
+		const bool bIsStandalone = IsCategoryLocal() && IsCategoryAuth();
+		CanvasContext.Printf(TEXT("%sMovement Mode: {yellow}%s{white}, Base: {yellow}%s")
+			, bIsStandalone ? TEXT("") : TEXT("(Server) ")
+			, *DataPack.MovementModeInfo, *DataPack.MovementBaseInfo);
+		
+		if (IsCategoryLocal() && !IsCategoryAuth())
+		{
+			const ACharacter* MyChar = Cast<ACharacter>(SelectedActor);
+			// In principle data collection should take place in CollectData, but we don't call that on clients for this category
+			// and we do want to display client-side data here as well, so we bend the rules a bit and gather and display data instantly here
+			if (const UCharacterMovementComponent* CharMovementComp = MyChar ? MyChar->GetCharacterMovement() : nullptr)
+			{
+				const UPrimitiveComponent* FloorComponent = MyChar->GetMovementBase();
+				const AActor* FloorActor = FloorComponent ? FloorComponent->GetOwner() : nullptr;
+				const FString ClientMovementBaseInfo = FloorComponent ? FString::Printf(TEXT("%s.%s"), *GetNameSafe(FloorActor), *FloorComponent->GetName()) : FString(TEXT("None"));
+				const FString ClientMovementModeInfo = CharMovementComp->GetMovementName();
+
+				CanvasContext.Printf(TEXT("(Client) Movement Mode: {yellow}%s{white}, Base: {yellow}%s"), *ClientMovementModeInfo, *ClientMovementBaseInfo);
+			}
+		}
+		
 		CanvasContext.Printf(TEXT("NavData: {yellow}%s{white}, Path following: {yellow}%s"), *DataPack.NavDataInfo, *DataPack.PathFollowingInfo);
 	}
 
