Fix rare crash that happens if the query finishes in the same frame that the BT is stopped.

diff --git a/Engine/Source/Runtime/AIModule/Private/BehaviorTree/Services/BTService_RunEQS.cpp b/Engine/Source/Runtime/AIModule/Private/BehaviorTree/Services/BTService_RunEQS.cpp
--- a/Engine/Source/Runtime/AIModule/Private/BehaviorTree/Services/BTService_RunEQS.cpp
+++ b/Engine/Source/Runtime/AIModule/Private/BehaviorTree/Services/BTService_RunEQS.cpp
@@ -89,8 +89,16 @@ void UBTService_RunEQS::OnQueryFinished(TSharedPtr<FEnvQueryResult> Result)
 	}
 
 	FBTEQSServiceMemory* MyMemory = CastInstanceNodeMemory<FBTEQSServiceMemory>(BTComp->GetNodeMemory(this, BTComp->FindInstanceContainingNode(this)));
-	check(MyMemory);
-	ensure(MyMemory->RequestID != INDEX_NONE);
+	if (MyMemory == nullptr || MyMemory->RequestID == INDEX_NONE)
+	{
+		// this can happen in hard to repro edge cases if the BT node being notified here has already been canceled
+		// in which case it should be safe to ignore it.
+		// @todo we need to find the exact case and repro it via unit tests
+		const bool bIsActive = BTComp->IsAuxNodeActive(this);
+		ensureMsgf(bIsActive == false, TEXT("%hs called while the BT node is not active or has already been aborted"), __FUNCTION__);
+		UE_CVLOG_UELOG(bIsActive == false, BTComp, LogBehaviorTree, Warning, TEXT("%hs called for %s while it's not active or has already been aborted"), __FUNCTION__, *GetNodeName());
+		return;
+	}
 
 	bool bSuccess = Result->IsSuccessful() && (Result->Items.Num() >= 1);
 	if (bSuccess)
