SmartObjects: Changed hash fgrid form check() to clamp when cell location does not fit into int32 range. Fixes crash when objects get really far away.

#localization none
#tests PIE

diff --git a/Engine/Source/Runtime/AIModule/Public/HierarchicalHashGrid2D.h b/Engine/Source/Runtime/AIModule/Public/HierarchicalHashGrid2D.h
--- a/Engine/Source/Runtime/AIModule/Public/HierarchicalHashGrid2D.h
+++ b/Engine/Source/Runtime/AIModule/Public/HierarchicalHashGrid2D.h
@@ -465,10 +465,10 @@ public:
 	FCellRect CalcQueryBounds(const FBox& Bounds, const int32 Level) const
 	{
 		FCellRect Result;
-		Result.MinX = IntCastChecked<int32>(FMath::FloorToInt(Bounds.Min.X * InvCellSize[Level] - 0.5f));
-		Result.MinY = IntCastChecked<int32>(FMath::FloorToInt(Bounds.Min.Y * InvCellSize[Level] - 0.5f));
-		Result.MaxX = IntCastChecked<int32>(FMath::FloorToInt(Bounds.Max.X * InvCellSize[Level] + 0.5f));
-		Result.MaxY = IntCastChecked<int32>(FMath::FloorToInt(Bounds.Max.Y * InvCellSize[Level] + 0.5f));
+		Result.MinX = ClampInt32(FMath::FloorToInt(Bounds.Min.X * InvCellSize[Level] - 0.5f));
+		Result.MinY = ClampInt32(FMath::FloorToInt(Bounds.Min.Y * InvCellSize[Level] - 0.5f));
+		Result.MaxX = ClampInt32(FMath::FloorToInt(Bounds.Max.X * InvCellSize[Level] + 0.5f));
+		Result.MaxY = ClampInt32(FMath::FloorToInt(Bounds.Max.Y * InvCellSize[Level] + 0.5f));
 		return Result;
 	}
 
@@ -509,12 +509,12 @@ public:
 		const FVector::FReal Diameter = FMath::Max(Bounds.Max.X - Bounds.Min.X, Bounds.Max.Y - Bounds.Min.Y);
 		for (Location.Level = 0; Location.Level < NumLevels; Location.Level++)
 		{
-			const int32 DiameterCells = IntCastChecked<int32>(FMath::CeilToInt(Diameter * InvCellSize[Location.Level]));
+			const int32 DiameterCells = ClampInt32(FMath::CeilToInt(Diameter * InvCellSize[Location.Level]));
 			// note that it's fine for DiameterCells to equal 0 - that would happen for 0-sized items (valid location, no extent).
 			if (DiameterCells <= 1)
 			{
-				Location.X = IntCastChecked<int32>(FMath::FloorToInt(Center.X * InvCellSize[Location.Level]));
-				Location.Y = IntCastChecked<int32>(FMath::FloorToInt(Center.Y * InvCellSize[Location.Level]));
+				Location.X = ClampInt32(FMath::FloorToInt(Center.X * InvCellSize[Location.Level]));
+				Location.Y = ClampInt32(FMath::FloorToInt(Center.Y * InvCellSize[Location.Level]));
 				break;
 			}
 		}
@@ -603,9 +603,16 @@ public:
 	const TSparseArray<FItem>& GetItems() const { return Items; }
 
 	/** @return Index in items of the first item on the spill list. */
-	int32 GetFirstSpillListItem() const { return SpillList; }
+	int32 GetFirstSpillListItem() const { return SpillList; } 
 
 protected:
+
+	/** @return given int64 clamped in int32 range. */
+	static constexpr int32 ClampInt32(const int64 Value)
+	{
+		return FMath::Clamp(Value, (int64)std::numeric_limits<int32>::lowest(), (int64)std::numeric_limits<int32>::max());
+	}
+	
 	TStaticArray<float, NumLevels> CellSize;		/** Lowest level cell size */
 	TStaticArray<float, NumLevels> InvCellSize;		/** 1 / CellSize */
 	TStaticArray<int32, NumLevels> LevelItemCount;	/** Number items per level, can be used to skip certain levels. */
