Create helper function InitializeNodeMemory & CleanupNodeMemory in BTNode to handle creation and clean up of node instance memory
Update bt nodes that use instanced memory to use the new helper to create and destroy their instance memory

diff --git a/Engine/Plugins/Experimental/GameplayBehaviors/Source/GameplayBehaviorsModule/Private/AI/BTDecorator_GameplayTagQuery.cpp b/Engine/Plugins/Experimental/GameplayBehaviors/Source/GameplayBehaviorsModule/Private/AI/BTDecorator_GameplayTagQuery.cpp
--- a/Engine/Plugins/Experimental/GameplayBehaviors/Source/GameplayBehaviorsModule/Private/AI/BTDecorator_GameplayTagQuery.cpp
+++ b/Engine/Plugins/Experimental/GameplayBehaviors/Source/GameplayBehaviorsModule/Private/AI/BTDecorator_GameplayTagQuery.cpp
@@ -66,10 +66,16 @@ FString UBTDecorator_GameplayTagQuery::GetStaticDescription() const
 	return FString::Printf(TEXT("%s: %s"), *Super::GetStaticDescription(), *GameplayTagQuery.GetDescription());
 }
 
+void UBTDecorator_GameplayTagQuery::InitializeMemory(UBehaviorTreeComponent& OwnerComp, uint8* NodeMemory, EBTMemoryInit::Type InitType) const
+{
+	InitializeNodeMemory<FBTDecorator_GameplayTagQueryMemory>(NodeMemory, InitType);
+}
+
 void UBTDecorator_GameplayTagQuery::CleanupMemory(UBehaviorTreeComponent& OwnerComp, uint8* NodeMemory, EBTMemoryClear::Type CleanupType) const
 {
 	const FBTDecorator_GameplayTagQueryMemory* MyMemory = CastInstanceNodeMemory<FBTDecorator_GameplayTagQueryMemory>(NodeMemory);
 	ensureMsgf(MyMemory->GameplayTagEventHandles.Num() == 0, TEXT("Dangling gameplay tag event handles for decorator %s"), *GetStaticDescription());
+	CleanupNodeMemory<FBTDecorator_GameplayTagQueryMemory>(NodeMemory, CleanupType);
 }
 
 void UBTDecorator_GameplayTagQuery::OnBecomeRelevant(UBehaviorTreeComponent& OwnerComp, uint8* NodeMemory)
diff --git a/Engine/Plugins/Experimental/GameplayBehaviors/Source/GameplayBehaviorsModule/Public/AI/BTDecorator_GameplayTagQuery.h b/Engine/Plugins/Experimental/GameplayBehaviors/Source/GameplayBehaviorsModule/Public/AI/BTDecorator_GameplayTagQuery.h
--- a/Engine/Plugins/Experimental/GameplayBehaviors/Source/GameplayBehaviorsModule/Public/AI/BTDecorator_GameplayTagQuery.h
+++ b/Engine/Plugins/Experimental/GameplayBehaviors/Source/GameplayBehaviorsModule/Public/AI/BTDecorator_GameplayTagQuery.h
@@ -28,6 +28,7 @@ class GAMEPLAYBEHAVIORSMODULE_API UBTDecorator_GameplayTagQuery : public UBTDeco
 
 	virtual FString GetStaticDescription() const override;
 
+	virtual void InitializeMemory(UBehaviorTreeComponent& OwnerComp, uint8* NodeMemory, EBTMemoryInit::Type InitType) const override;
 	virtual void CleanupMemory(UBehaviorTreeComponent& OwnerComp, uint8* NodeMemory, EBTMemoryClear::Type CleanupType) const override;
 
 protected:
diff --git a/Engine/Plugins/Runtime/GameplayBehaviorSmartObjects/Source/GameplayBehaviorSmartObjectsModule/Private/AI/BTTask_FindAndUseGameplayBehaviorSmartObject.cpp b/Engine/Plugins/Runtime/GameplayBehaviorSmartObjects/Source/GameplayBehaviorSmartObjectsModule/Private/AI/BTTask_FindAndUseGameplayBehaviorSmartObject.cpp
--- a/Engine/Plugins/Runtime/GameplayBehaviorSmartObjects/Source/GameplayBehaviorSmartObjectsModule/Private/AI/BTTask_FindAndUseGameplayBehaviorSmartObject.cpp
+++ b/Engine/Plugins/Runtime/GameplayBehaviorSmartObjects/Source/GameplayBehaviorSmartObjectsModule/Private/AI/BTTask_FindAndUseGameplayBehaviorSmartObject.cpp
@@ -20,6 +20,16 @@ UBTTask_FindAndUseGameplayBehaviorSmartObject::UBTTask_FindAndUseGameplayBehavio
 	bNotifyTaskFinished = true;
 }
 
+void UBTTask_FindAndUseGameplayBehaviorSmartObject::InitializeMemory(UBehaviorTreeComponent& OwnerComp, uint8* NodeMemory, EBTMemoryInit::Type InitType) const
+{
+	InitializeNodeMemory<FBTUseSOTaskMemory>(NodeMemory, InitType);
+}
+
+void UBTTask_FindAndUseGameplayBehaviorSmartObject::CleanupMemory(UBehaviorTreeComponent& OwnerComp, uint8* NodeMemory, EBTMemoryClear::Type CleanupType) const
+{
+	CleanupNodeMemory<FBTUseSOTaskMemory>(NodeMemory, CleanupType);
+}
+
 EBTNodeResult::Type UBTTask_FindAndUseGameplayBehaviorSmartObject::ExecuteTask(UBehaviorTreeComponent& OwnerComp, uint8* NodeMemory)
 {
 	EBTNodeResult::Type NodeResult = EBTNodeResult::Failed;
diff --git a/Engine/Plugins/Runtime/GameplayBehaviorSmartObjects/Source/GameplayBehaviorSmartObjectsModule/Public/AI/BTTask_FindAndUseGameplayBehaviorSmartObject.h b/Engine/Plugins/Runtime/GameplayBehaviorSmartObjects/Source/GameplayBehaviorSmartObjectsModule/Public/AI/BTTask_FindAndUseGameplayBehaviorSmartObject.h
--- a/Engine/Plugins/Runtime/GameplayBehaviorSmartObjects/Source/GameplayBehaviorSmartObjectsModule/Public/AI/BTTask_FindAndUseGameplayBehaviorSmartObject.h
+++ b/Engine/Plugins/Runtime/GameplayBehaviorSmartObjects/Source/GameplayBehaviorSmartObjectsModule/Public/AI/BTTask_FindAndUseGameplayBehaviorSmartObject.h
@@ -32,6 +32,8 @@ protected:
 	virtual EBTNodeResult::Type AbortTask(UBehaviorTreeComponent& OwnerComp, uint8* NodeMemory) override;
 	virtual void OnTaskFinished(UBehaviorTreeComponent& OwnerComp, uint8* NodeMemory, EBTNodeResult::Type TaskResult) override;
 	virtual uint16 GetInstanceMemorySize() const override { return sizeof(FBTUseSOTaskMemory); }
+	virtual void InitializeMemory(UBehaviorTreeComponent& OwnerComp, uint8* NodeMemory, EBTMemoryInit::Type InitType) const override;
+	virtual void CleanupMemory(UBehaviorTreeComponent& OwnerComp, uint8* NodeMemory, EBTMemoryClear::Type CleanupType) const override;
 
 	virtual FString GetStaticDescription() const override;
 
diff --git a/Engine/Source/Developer/AITestSuite/Classes/BehaviorTree/TestBTDecorator_DelayedAbort.h b/Engine/Source/Developer/AITestSuite/Classes/BehaviorTree/TestBTDecorator_DelayedAbort.h
--- a/Engine/Source/Developer/AITestSuite/Classes/BehaviorTree/TestBTDecorator_DelayedAbort.h
+++ b/Engine/Source/Developer/AITestSuite/Classes/BehaviorTree/TestBTDecorator_DelayedAbort.h
@@ -26,4 +26,6 @@ class UTestBTDecorator_DelayedAbort : public UBTDecorator
 	virtual void OnBecomeRelevant(UBehaviorTreeComponent& OwnerComp, uint8* NodeMemory) override;
 	virtual void TickNode(UBehaviorTreeComponent& OwnerComp, uint8* NodeMemory, float DeltaSeconds) override;
 	virtual uint16 GetInstanceMemorySize() const override;
+	virtual void InitializeMemory(UBehaviorTreeComponent& OwnerComp, uint8* NodeMemory, EBTMemoryInit::Type InitType) const override;
+	virtual void CleanupMemory(UBehaviorTreeComponent& OwnerComp, uint8* NodeMemory, EBTMemoryClear::Type CleanupType) const override;
 };
diff --git a/Engine/Source/Developer/AITestSuite/Classes/BehaviorTree/TestBTTask_LatentWithFlags.h b/Engine/Source/Developer/AITestSuite/Classes/BehaviorTree/TestBTTask_LatentWithFlags.h
--- a/Engine/Source/Developer/AITestSuite/Classes/BehaviorTree/TestBTTask_LatentWithFlags.h
+++ b/Engine/Source/Developer/AITestSuite/Classes/BehaviorTree/TestBTTask_LatentWithFlags.h
@@ -72,6 +72,8 @@ protected:
 	virtual EBTNodeResult::Type ExecuteTask(class UBehaviorTreeComponent& OwnerComp, uint8* NodeMemory) override;
 	virtual EBTNodeResult::Type AbortTask(UBehaviorTreeComponent& OwnerComp, uint8* NodeMemory) override;
 	virtual uint16 GetInstanceMemorySize() const override;
+	virtual void InitializeMemory(UBehaviorTreeComponent& OwnerComp, uint8* NodeMemory, EBTMemoryInit::Type InitType) const override;
+	virtual void CleanupMemory(UBehaviorTreeComponent& OwnerComp, uint8* NodeMemory, EBTMemoryClear::Type CleanupType) const override;
 	void LogExecution(class UBehaviorTreeComponent& OwnerComp, int32 LogNumber);
 
 	virtual void TickTask(class UBehaviorTreeComponent& OwnerComp, uint8* NodeMemory, float DeltaSeconds) override;
diff --git a/Engine/Source/Developer/AITestSuite/Classes/BehaviorTree/TestBTTask_Log.h b/Engine/Source/Developer/AITestSuite/Classes/BehaviorTree/TestBTTask_Log.h
--- a/Engine/Source/Developer/AITestSuite/Classes/BehaviorTree/TestBTTask_Log.h
+++ b/Engine/Source/Developer/AITestSuite/Classes/BehaviorTree/TestBTTask_Log.h
@@ -33,6 +33,8 @@ class UTestBTTask_Log : public UBTTaskNode
 
 	virtual EBTNodeResult::Type ExecuteTask(UBehaviorTreeComponent& OwnerComp, uint8* NodeMemory) override;
 	virtual uint16 GetInstanceMemorySize() const override;
+	virtual void InitializeMemory(UBehaviorTreeComponent& OwnerComp, uint8* NodeMemory, EBTMemoryInit::Type InitType) const override;
+	virtual void CleanupMemory(UBehaviorTreeComponent& OwnerComp, uint8* NodeMemory, EBTMemoryClear::Type CleanupType) const override;
 
 	void LogExecution(UBehaviorTreeComponent& OwnerComp, int32 LogNumber);
 
diff --git a/Engine/Source/Developer/AITestSuite/Classes/BehaviorTree/TestBTTask_SetValuesWithLogs.h b/Engine/Source/Developer/AITestSuite/Classes/BehaviorTree/TestBTTask_SetValuesWithLogs.h
--- a/Engine/Source/Developer/AITestSuite/Classes/BehaviorTree/TestBTTask_SetValuesWithLogs.h
+++ b/Engine/Source/Developer/AITestSuite/Classes/BehaviorTree/TestBTTask_SetValuesWithLogs.h
@@ -73,4 +73,6 @@ protected:
 	virtual EBTNodeResult::Type ExecuteTask(UBehaviorTreeComponent& OwnerComp, uint8* NodeMemory) override;
 	virtual EBTNodeResult::Type AbortTask(UBehaviorTreeComponent& OwnerComp, uint8* NodeMemory) override;
 	virtual uint16 GetInstanceMemorySize() const override;
+	virtual void InitializeMemory(UBehaviorTreeComponent& OwnerComp, uint8* NodeMemory, EBTMemoryInit::Type InitType) const override;
+	virtual void CleanupMemory(UBehaviorTreeComponent& OwnerComp, uint8* NodeMemory, EBTMemoryClear::Type CleanupType) const override;
 };
diff --git a/Engine/Source/Developer/AITestSuite/Private/BehaviorTree/TestBTDecorator_DelayedAbort.cpp b/Engine/Source/Developer/AITestSuite/Private/BehaviorTree/TestBTDecorator_DelayedAbort.cpp
--- a/Engine/Source/Developer/AITestSuite/Private/BehaviorTree/TestBTDecorator_DelayedAbort.cpp
+++ b/Engine/Source/Developer/AITestSuite/Private/BehaviorTree/TestBTDecorator_DelayedAbort.cpp
@@ -41,3 +41,12 @@ uint16 UTestBTDecorator_DelayedAbort::GetInstanceMemorySize() const
 	return sizeof(FBTDelayedAbortMemory);
 }
 
+void UTestBTDecorator_DelayedAbort::InitializeMemory(UBehaviorTreeComponent& OwnerComp, uint8* NodeMemory, EBTMemoryInit::Type InitType) const
+{
+	InitializeNodeMemory<FBTDelayedAbortMemory>(NodeMemory, InitType);
+}
+
+void UTestBTDecorator_DelayedAbort::CleanupMemory(UBehaviorTreeComponent& OwnerComp, uint8* NodeMemory, EBTMemoryClear::Type CleanupType) const
+{
+	CleanupNodeMemory<FBTDelayedAbortMemory>(NodeMemory, CleanupType);
+}
\ No newline at end of file
diff --git a/Engine/Source/Developer/AITestSuite/Private/BehaviorTree/TestBTTask_LatentWithFlags.cpp b/Engine/Source/Developer/AITestSuite/Private/BehaviorTree/TestBTTask_LatentWithFlags.cpp
--- a/Engine/Source/Developer/AITestSuite/Private/BehaviorTree/TestBTTask_LatentWithFlags.cpp
+++ b/Engine/Source/Developer/AITestSuite/Private/BehaviorTree/TestBTTask_LatentWithFlags.cpp
@@ -98,6 +98,16 @@ uint16 UTestBTTask_LatentWithFlags::GetInstanceMemorySize() const
 	return sizeof(FBTLatentTaskMemory);
 }
 
+void UTestBTTask_LatentWithFlags::InitializeMemory(UBehaviorTreeComponent& OwnerComp, uint8* NodeMemory, EBTMemoryInit::Type InitType) const
+{
+	InitializeNodeMemory<FBTLatentTaskMemory>(NodeMemory, InitType);
+}
+
+void UTestBTTask_LatentWithFlags::CleanupMemory(UBehaviorTreeComponent& OwnerComp, uint8* NodeMemory, EBTMemoryClear::Type CleanupType) const
+{
+	CleanupNodeMemory<FBTLatentTaskMemory>(NodeMemory, CleanupType);
+}
+
 void UTestBTTask_LatentWithFlags::LogExecution(UBehaviorTreeComponent& OwnerComp, int32 LogNumber)
 {
 	if (LogNumber >= 0)
diff --git a/Engine/Source/Developer/AITestSuite/Private/BehaviorTree/TestBTTask_Log.cpp b/Engine/Source/Developer/AITestSuite/Private/BehaviorTree/TestBTTask_Log.cpp
--- a/Engine/Source/Developer/AITestSuite/Private/BehaviorTree/TestBTTask_Log.cpp
+++ b/Engine/Source/Developer/AITestSuite/Private/BehaviorTree/TestBTTask_Log.cpp
@@ -52,6 +52,16 @@ uint16 UTestBTTask_Log::GetInstanceMemorySize() const
 	return sizeof(FBTLogTaskMemory);
 }
 
+void UTestBTTask_Log::InitializeMemory(UBehaviorTreeComponent& OwnerComp, uint8* NodeMemory, EBTMemoryInit::Type InitType) const
+{
+	InitializeNodeMemory<FBTLogTaskMemory>(NodeMemory, InitType);
+}
+
+void UTestBTTask_Log::CleanupMemory(UBehaviorTreeComponent& OwnerComp, uint8* NodeMemory, EBTMemoryClear::Type CleanupType) const
+{
+	CleanupNodeMemory<FBTLogTaskMemory>(NodeMemory, CleanupType);
+}
+
 void UTestBTTask_Log::LogExecution(UBehaviorTreeComponent& OwnerComp, int32 LogNumber)
 {
 	if (LogNumber >= 0)
diff --git a/Engine/Source/Developer/AITestSuite/Private/BehaviorTree/TestBTTask_SetValuesWithLogs.cpp b/Engine/Source/Developer/AITestSuite/Private/BehaviorTree/TestBTTask_SetValuesWithLogs.cpp
--- a/Engine/Source/Developer/AITestSuite/Private/BehaviorTree/TestBTTask_SetValuesWithLogs.cpp
+++ b/Engine/Source/Developer/AITestSuite/Private/BehaviorTree/TestBTTask_SetValuesWithLogs.cpp
@@ -74,6 +74,16 @@ uint16 UTestBTTask_SetValuesWithLogs::GetInstanceMemorySize() const
 	return sizeof(FBTSetValueTaskMemory);
 }
 
+void UTestBTTask_SetValuesWithLogs::InitializeMemory(UBehaviorTreeComponent& OwnerComp, uint8* NodeMemory, EBTMemoryInit::Type InitType) const
+{
+	InitializeNodeMemory<FBTSetValueTaskMemory>(NodeMemory, InitType);
+}
+
+void UTestBTTask_SetValuesWithLogs::CleanupMemory(UBehaviorTreeComponent& OwnerComp, uint8* NodeMemory, EBTMemoryClear::Type CleanupType) const
+{
+	CleanupNodeMemory<FBTSetValueTaskMemory>(NodeMemory, CleanupType);
+}
+
 EBTNodeResult::Type UTestBTTask_SetValuesWithLogs::AbortTask(UBehaviorTreeComponent& OwnerComp, uint8* NodeMemory)
 {
 	if (OnAbortKeyName.IsValid())
diff --git a/Engine/Source/Runtime/AIModule/Classes/BehaviorTree/BTCompositeNode.h b/Engine/Source/Runtime/AIModule/Classes/BehaviorTree/BTCompositeNode.h
--- a/Engine/Source/Runtime/AIModule/Classes/BehaviorTree/BTCompositeNode.h
+++ b/Engine/Source/Runtime/AIModule/Classes/BehaviorTree/BTCompositeNode.h
@@ -144,6 +144,8 @@ class UBTCompositeNode : public UBTNode
 
 	/** size of instance memory */
 	AIMODULE_API virtual uint16 GetInstanceMemorySize() const override;
+	virtual void InitializeMemory(UBehaviorTreeComponent& OwnerComp, uint8* NodeMemory, EBTMemoryInit::Type InitType) const override;
+	virtual void CleanupMemory(UBehaviorTreeComponent& OwnerComp, uint8* NodeMemory, EBTMemoryClear::Type CleanupType) const override;
 
 	/** @return child node at given index */
 	AIMODULE_API UBTNode* GetChildNode(int32 Index) const;
diff --git a/Engine/Source/Runtime/AIModule/Classes/BehaviorTree/BTNode.h b/Engine/Source/Runtime/AIModule/Classes/BehaviorTree/BTNode.h
--- a/Engine/Source/Runtime/AIModule/Classes/BehaviorTree/BTNode.h
+++ b/Engine/Source/Runtime/AIModule/Classes/BehaviorTree/BTNode.h
@@ -47,10 +47,10 @@ class UBTNode : public UObject, public IGameplayTaskOwnerInterface
 	/** initialize any asset related data */
 	AIMODULE_API virtual void InitializeFromAsset(UBehaviorTree& Asset);
 	
-	/** initialize memory block */
+	/** initialize memory block. InitializeNodeMemory template function is provided to help initialize the memory. */
 	AIMODULE_API virtual void InitializeMemory(UBehaviorTreeComponent& OwnerComp, uint8* NodeMemory, EBTMemoryInit::Type InitType) const;
 
-	/** cleanup memory block */
+	/** cleanup memory block. CleanupNodeMemory template function is provided to help cleanup the memory. */
 	AIMODULE_API virtual void CleanupMemory(UBehaviorTreeComponent& OwnerComp, uint8* NodeMemory, EBTMemoryClear::Type CleanupType) const;
 
 	/** gathers description of all runtime parameters */
@@ -82,6 +82,12 @@ class UBTNode : public UObject, public IGameplayTaskOwnerInterface
 	AIMODULE_API UBTNode* GetNextNode() const;
 #endif
 
+	template<typename T>
+	T* InitializeNodeMemory(uint8* NodeMemory, EBTMemoryInit::Type InitType) const;
+
+	template <typename T>
+	void CleanupNodeMemory(uint8* NodeMemory, EBTMemoryClear::Type CleanupType) const;
+
 	template<typename T>
 	T* GetNodeMemory(FBehaviorTreeSearchData& SearchData) const;
 
@@ -300,6 +306,28 @@ FORCEINLINE bool UBTNode::IsInstanced() const
 	return bIsInstanced;
 }
 
+template<typename T>
+T* UBTNode::InitializeNodeMemory(uint8* NodeMemory, EBTMemoryInit::Type InitType) const
+{
+	if (InitType == EBTMemoryInit::Initialize)
+	{
+		new(NodeMemory) T();
+	}
+	return CastInstanceNodeMemory<T>(NodeMemory);
+}
+
+template <typename T>
+void UBTNode::CleanupNodeMemory(uint8* NodeMemory, EBTMemoryClear::Type CleanupType) const
+{
+	if constexpr (!TIsTriviallyDestructible<T>::Value)
+	{
+		if (CleanupType == EBTMemoryClear::Destroy)
+		{
+			CastInstanceNodeMemory<T>(NodeMemory)->~T();
+		}
+	}
+}
+
 template<typename T>
 T* UBTNode::GetNodeMemory(FBehaviorTreeSearchData& SearchData) const
 {
diff --git a/Engine/Source/Runtime/AIModule/Classes/BehaviorTree/Composites/BTComposite_SimpleParallel.h b/Engine/Source/Runtime/AIModule/Classes/BehaviorTree/Composites/BTComposite_SimpleParallel.h
--- a/Engine/Source/Runtime/AIModule/Classes/BehaviorTree/Composites/BTComposite_SimpleParallel.h
+++ b/Engine/Source/Runtime/AIModule/Classes/BehaviorTree/Composites/BTComposite_SimpleParallel.h
@@ -68,6 +68,8 @@ class UBTComposite_SimpleParallel : public UBTCompositeNode
 	AIMODULE_API virtual bool CanPushSubtree(UBehaviorTreeComponent& OwnerComp, uint8* NodeMemory, int32 ChildIdx) const override;
 	AIMODULE_API virtual void SetChildOverride(FBehaviorTreeSearchData& SearchData, int8 Index) const override;
 	AIMODULE_API virtual uint16 GetInstanceMemorySize() const override;
+	AIMODULE_API virtual void InitializeMemory(UBehaviorTreeComponent& OwnerComp, uint8* NodeMemory, EBTMemoryInit::Type InitType) const override;
+	AIMODULE_API virtual void CleanupMemory(UBehaviorTreeComponent& OwnerComp, uint8* NodeMemory, EBTMemoryClear::Type CleanupType) const override;
 	AIMODULE_API virtual FString GetStaticDescription() const override;
 	AIMODULE_API virtual void DescribeRuntimeValues(const UBehaviorTreeComponent& OwnerComp, uint8* NodeMemory, EBTDescriptionVerbosity::Type Verbosity, TArray<FString>& Values) const override;
 
diff --git a/Engine/Source/Runtime/AIModule/Classes/BehaviorTree/Decorators/BTDecorator_ConditionalLoop.h b/Engine/Source/Runtime/AIModule/Classes/BehaviorTree/Decorators/BTDecorator_ConditionalLoop.h
--- a/Engine/Source/Runtime/AIModule/Classes/BehaviorTree/Decorators/BTDecorator_ConditionalLoop.h
+++ b/Engine/Source/Runtime/AIModule/Classes/BehaviorTree/Decorators/BTDecorator_ConditionalLoop.h
@@ -29,6 +29,8 @@ class UBTDecorator_ConditionalLoop : public UBTDecorator_Blackboard
 	AIMODULE_API virtual void OnNodeDeactivation(FBehaviorTreeSearchData& SearchData, EBTNodeResult::Type NodeResult) override;
 
 	AIMODULE_API virtual uint16 GetInstanceMemorySize() const override;
+	AIMODULE_API virtual void InitializeMemory(UBehaviorTreeComponent& OwnerComp, uint8* NodeMemory, EBTMemoryInit::Type InitType) const override;
+	AIMODULE_API virtual void CleanupMemory(UBehaviorTreeComponent& OwnerComp, uint8* NodeMemory, EBTMemoryClear::Type CleanupType) const override;
 
 #if WITH_EDITOR
 	AIMODULE_API virtual FName GetNodeIconName() const override;
diff --git a/Engine/Source/Runtime/AIModule/Classes/BehaviorTree/Decorators/BTDecorator_ConeCheck.h b/Engine/Source/Runtime/AIModule/Classes/BehaviorTree/Decorators/BTDecorator_ConeCheck.h
--- a/Engine/Source/Runtime/AIModule/Classes/BehaviorTree/Decorators/BTDecorator_ConeCheck.h
+++ b/Engine/Source/Runtime/AIModule/Classes/BehaviorTree/Decorators/BTDecorator_ConeCheck.h
@@ -48,6 +48,8 @@ class UBTDecorator_ConeCheck : public UBTDecorator
 
 	AIMODULE_API virtual void InitializeFromAsset(UBehaviorTree& Asset) override;
 	AIMODULE_API virtual uint16 GetInstanceMemorySize() const override;
+	AIMODULE_API virtual void InitializeMemory(UBehaviorTreeComponent& OwnerComp, uint8* NodeMemory, EBTMemoryInit::Type InitType) const override;
+	AIMODULE_API virtual void CleanupMemory(UBehaviorTreeComponent& OwnerComp, uint8* NodeMemory, EBTMemoryClear::Type CleanupType) const override;
 	AIMODULE_API virtual void DescribeRuntimeValues(const UBehaviorTreeComponent& OwnerComp, uint8* NodeMemory, EBTDescriptionVerbosity::Type Verbosity, TArray<FString>& Values) const override;
 	AIMODULE_API virtual FString GetStaticDescription() const override;
 
diff --git a/Engine/Source/Runtime/AIModule/Classes/BehaviorTree/Decorators/BTDecorator_Cooldown.h b/Engine/Source/Runtime/AIModule/Classes/BehaviorTree/Decorators/BTDecorator_Cooldown.h
--- a/Engine/Source/Runtime/AIModule/Classes/BehaviorTree/Decorators/BTDecorator_Cooldown.h
+++ b/Engine/Source/Runtime/AIModule/Classes/BehaviorTree/Decorators/BTDecorator_Cooldown.h
@@ -32,6 +32,7 @@ class UBTDecorator_Cooldown : public UBTDecorator
 
 	AIMODULE_API virtual bool CalculateRawConditionValue(UBehaviorTreeComponent& OwnerComp, uint8* NodeMemory) const override;
 	AIMODULE_API virtual void InitializeMemory(UBehaviorTreeComponent& OwnerComp, uint8* NodeMemory, EBTMemoryInit::Type InitType) const override;
+	AIMODULE_API virtual void CleanupMemory(UBehaviorTreeComponent& OwnerComp, uint8* NodeMemory, EBTMemoryClear::Type CleanupType) const override;
 	AIMODULE_API virtual uint16 GetInstanceMemorySize() const override;
 	AIMODULE_API virtual void DescribeRuntimeValues(const UBehaviorTreeComponent& OwnerComp, uint8* NodeMemory, EBTDescriptionVerbosity::Type Verbosity, TArray<FString>& Values) const override;
 	AIMODULE_API virtual FString GetStaticDescription() const override;
diff --git a/Engine/Source/Runtime/AIModule/Classes/BehaviorTree/Decorators/BTDecorator_KeepInCone.h b/Engine/Source/Runtime/AIModule/Classes/BehaviorTree/Decorators/BTDecorator_KeepInCone.h
--- a/Engine/Source/Runtime/AIModule/Classes/BehaviorTree/Decorators/BTDecorator_KeepInCone.h
+++ b/Engine/Source/Runtime/AIModule/Classes/BehaviorTree/Decorators/BTDecorator_KeepInCone.h
@@ -50,6 +50,8 @@ class UBTDecorator_KeepInCone : public UBTDecorator
 
 	AIMODULE_API virtual void InitializeFromAsset(UBehaviorTree& Asset) override;
 	AIMODULE_API virtual uint16 GetInstanceMemorySize() const override;
+	AIMODULE_API virtual void InitializeMemory(UBehaviorTreeComponent& OwnerComp, uint8* NodeMemory, EBTMemoryInit::Type InitType) const override;
+	AIMODULE_API virtual void CleanupMemory(UBehaviorTreeComponent& OwnerComp, uint8* NodeMemory, EBTMemoryClear::Type CleanupType) const override;
 	AIMODULE_API virtual void DescribeRuntimeValues(const UBehaviorTreeComponent& OwnerComp, uint8* NodeMemory, EBTDescriptionVerbosity::Type Verbosity, TArray<FString>& Values) const override;
 	AIMODULE_API virtual FString GetStaticDescription() const override;
 
diff --git a/Engine/Source/Runtime/AIModule/Classes/BehaviorTree/Decorators/BTDecorator_Loop.h b/Engine/Source/Runtime/AIModule/Classes/BehaviorTree/Decorators/BTDecorator_Loop.h
--- a/Engine/Source/Runtime/AIModule/Classes/BehaviorTree/Decorators/BTDecorator_Loop.h
+++ b/Engine/Source/Runtime/AIModule/Classes/BehaviorTree/Decorators/BTDecorator_Loop.h
@@ -35,6 +35,8 @@ class UBTDecorator_Loop : public UBTDecorator
 	float InfiniteLoopTimeoutTime;
 
 	AIMODULE_API virtual uint16 GetInstanceMemorySize() const override;
+	AIMODULE_API virtual void InitializeMemory(UBehaviorTreeComponent& OwnerComp, uint8* NodeMemory, EBTMemoryInit::Type InitType) const override;
+	AIMODULE_API virtual void CleanupMemory(UBehaviorTreeComponent& OwnerComp, uint8* NodeMemory, EBTMemoryClear::Type CleanupType) const override;
 	AIMODULE_API virtual void DescribeRuntimeValues(const UBehaviorTreeComponent& OwnerComp, uint8* NodeMemory, EBTDescriptionVerbosity::Type Verbosity, TArray<FString>& Values) const override;
 	AIMODULE_API virtual FString GetStaticDescription() const override;
 
diff --git a/Engine/Source/Runtime/AIModule/Classes/BehaviorTree/Decorators/BTDecorator_TagCooldown.h b/Engine/Source/Runtime/AIModule/Classes/BehaviorTree/Decorators/BTDecorator_TagCooldown.h
--- a/Engine/Source/Runtime/AIModule/Classes/BehaviorTree/Decorators/BTDecorator_TagCooldown.h
+++ b/Engine/Source/Runtime/AIModule/Classes/BehaviorTree/Decorators/BTDecorator_TagCooldown.h
@@ -44,6 +44,7 @@ class UBTDecorator_TagCooldown : public UBTDecorator
 
 	AIMODULE_API virtual bool CalculateRawConditionValue(UBehaviorTreeComponent& OwnerComp, uint8* NodeMemory) const override;
 	AIMODULE_API virtual void InitializeMemory(UBehaviorTreeComponent& OwnerComp, uint8* NodeMemory, EBTMemoryInit::Type InitType) const override;
+	AIMODULE_API virtual void CleanupMemory(UBehaviorTreeComponent& OwnerComp, uint8* NodeMemory, EBTMemoryClear::Type CleanupType) const override;
 	AIMODULE_API virtual uint16 GetInstanceMemorySize() const override;
 	AIMODULE_API virtual void DescribeRuntimeValues(const UBehaviorTreeComponent& OwnerComp, uint8* NodeMemory, EBTDescriptionVerbosity::Type Verbosity, TArray<FString>& Values) const override;
 	AIMODULE_API virtual FString GetStaticDescription() const override;
diff --git a/Engine/Source/Runtime/AIModule/Classes/BehaviorTree/Services/BTService_DefaultFocus.h b/Engine/Source/Runtime/AIModule/Classes/BehaviorTree/Services/BTService_DefaultFocus.h
--- a/Engine/Source/Runtime/AIModule/Classes/BehaviorTree/Services/BTService_DefaultFocus.h
+++ b/Engine/Source/Runtime/AIModule/Classes/BehaviorTree/Services/BTService_DefaultFocus.h
@@ -43,6 +43,8 @@ protected:
 	AIMODULE_API UBTService_DefaultFocus(const FObjectInitializer& ObjectInitializer = FObjectInitializer::Get());
 
 	virtual uint16 GetInstanceMemorySize() const override { return sizeof(FBTFocusMemory); }
+	virtual void InitializeMemory(UBehaviorTreeComponent& OwnerComp, uint8* NodeMemory, EBTMemoryInit::Type InitType) const override;
+	virtual void CleanupMemory(UBehaviorTreeComponent& OwnerComp, uint8* NodeMemory, EBTMemoryClear::Type CleanupType) const override;
 	AIMODULE_API virtual void OnBecomeRelevant(UBehaviorTreeComponent& OwnerComp, uint8* NodeMemory) override;
 	AIMODULE_API virtual void OnCeaseRelevant(UBehaviorTreeComponent& OwnerComp, uint8* NodeMemory) override;
 
diff --git a/Engine/Source/Runtime/AIModule/Classes/BehaviorTree/Services/BTService_RunEQS.h b/Engine/Source/Runtime/AIModule/Classes/BehaviorTree/Services/BTService_RunEQS.h
--- a/Engine/Source/Runtime/AIModule/Classes/BehaviorTree/Services/BTService_RunEQS.h
+++ b/Engine/Source/Runtime/AIModule/Classes/BehaviorTree/Services/BTService_RunEQS.h
@@ -38,13 +38,12 @@ public:
 	AIMODULE_API virtual void TickNode(UBehaviorTreeComponent& OwnerComp, uint8* NodeMemory, float DeltaSeconds) override;
 	AIMODULE_API virtual void OnCeaseRelevant(UBehaviorTreeComponent& OwnerComp, uint8* NodeMemory) override;
 	AIMODULE_API virtual void InitializeMemory(UBehaviorTreeComponent& OwnerComp, uint8* NodeMemory, EBTMemoryInit::Type InitType) const override;
+	AIMODULE_API virtual void CleanupMemory(UBehaviorTreeComponent& OwnerComp, uint8* NodeMemory, EBTMemoryClear::Type CleanupType) const override;
 	virtual uint16 GetInstanceMemorySize() const override { return sizeof(FBTEQSServiceMemory); }
 
 	AIMODULE_API virtual FString GetStaticDescription() const override;
 
 #if WITH_EDITOR
-	/** We need this only for verification, no need to have it in shipped builds */
-	AIMODULE_API virtual void CleanupMemory(UBehaviorTreeComponent& OwnerComp, uint8* NodeMemory, EBTMemoryClear::Type CleanupType) const override;
 	/** prepare query params */
 	AIMODULE_API virtual void PostEditChangeProperty(struct FPropertyChangedEvent& PropertyChangedEvent) override;
 #endif // WITH_EDITOR
diff --git a/Engine/Source/Runtime/AIModule/Classes/BehaviorTree/Tasks/BTTask_GameplayTaskBase.h b/Engine/Source/Runtime/AIModule/Classes/BehaviorTree/Tasks/BTTask_GameplayTaskBase.h
--- a/Engine/Source/Runtime/AIModule/Classes/BehaviorTree/Tasks/BTTask_GameplayTaskBase.h
+++ b/Engine/Source/Runtime/AIModule/Classes/BehaviorTree/Tasks/BTTask_GameplayTaskBase.h
@@ -31,6 +31,8 @@ class UBTTask_GameplayTaskBase : public UBTTaskNode
 	AIMODULE_API virtual void OnTaskFinished(UBehaviorTreeComponent& OwnerComp, uint8* NodeMemory, EBTNodeResult::Type TaskResult) override;
 	AIMODULE_API virtual void OnGameplayTaskDeactivated(UGameplayTask& Task) override;
 	AIMODULE_API virtual uint16 GetInstanceMemorySize() const override;
+	AIMODULE_API virtual void InitializeMemory(UBehaviorTreeComponent& OwnerComp, uint8* NodeMemory, EBTMemoryInit::Type InitType) const override;
+	AIMODULE_API virtual void CleanupMemory(UBehaviorTreeComponent& OwnerComp, uint8* NodeMemory, EBTMemoryClear::Type CleanupType) const override;
 
 protected:
 
diff --git a/Engine/Source/Runtime/AIModule/Classes/BehaviorTree/Tasks/BTTask_MoveTo.h b/Engine/Source/Runtime/AIModule/Classes/BehaviorTree/Tasks/BTTask_MoveTo.h
--- a/Engine/Source/Runtime/AIModule/Classes/BehaviorTree/Tasks/BTTask_MoveTo.h
+++ b/Engine/Source/Runtime/AIModule/Classes/BehaviorTree/Tasks/BTTask_MoveTo.h
@@ -95,6 +95,8 @@ class UBTTask_MoveTo : public UBTTask_BlackboardBase
 	AIMODULE_API virtual EBTNodeResult::Type AbortTask(UBehaviorTreeComponent& OwnerComp, uint8* NodeMemory) override;
 	AIMODULE_API virtual void OnTaskFinished(UBehaviorTreeComponent& OwnerComp, uint8* NodeMemory, EBTNodeResult::Type TaskResult) override;
 	AIMODULE_API virtual uint16 GetInstanceMemorySize() const override;
+	AIMODULE_API virtual void InitializeMemory(UBehaviorTreeComponent& OwnerComp, uint8* NodeMemory, EBTMemoryInit::Type InitType) const override;
+	AIMODULE_API virtual void CleanupMemory(UBehaviorTreeComponent& OwnerComp, uint8* NodeMemory, EBTMemoryClear::Type CleanupType) const override;
 	AIMODULE_API virtual void PostLoad() override;
 
 	AIMODULE_API virtual void OnGameplayTaskDeactivated(UGameplayTask& Task) override;
diff --git a/Engine/Source/Runtime/AIModule/Classes/BehaviorTree/Tasks/BTTask_RotateToFaceBBEntry.h b/Engine/Source/Runtime/AIModule/Classes/BehaviorTree/Tasks/BTTask_RotateToFaceBBEntry.h
--- a/Engine/Source/Runtime/AIModule/Classes/BehaviorTree/Tasks/BTTask_RotateToFaceBBEntry.h
+++ b/Engine/Source/Runtime/AIModule/Classes/BehaviorTree/Tasks/BTTask_RotateToFaceBBEntry.h
@@ -39,6 +39,8 @@ public:
 	AIMODULE_API virtual FString GetStaticDescription() const override;
 	
 	virtual uint16 GetInstanceMemorySize() const override { return sizeof(FBTFocusMemory); }
+	AIMODULE_API virtual void InitializeMemory(UBehaviorTreeComponent& OwnerComp, uint8* NodeMemory, EBTMemoryInit::Type InitType) const override;
+	AIMODULE_API virtual void CleanupMemory(UBehaviorTreeComponent& OwnerComp, uint8* NodeMemory, EBTMemoryClear::Type CleanupType) const override;
 
 protected:
 
diff --git a/Engine/Source/Runtime/AIModule/Classes/BehaviorTree/Tasks/BTTask_RunEQSQuery.h b/Engine/Source/Runtime/AIModule/Classes/BehaviorTree/Tasks/BTTask_RunEQSQuery.h
--- a/Engine/Source/Runtime/AIModule/Classes/BehaviorTree/Tasks/BTTask_RunEQSQuery.h
+++ b/Engine/Source/Runtime/AIModule/Classes/BehaviorTree/Tasks/BTTask_RunEQSQuery.h
@@ -68,6 +68,8 @@ class UBTTask_RunEQSQuery : public UBTTask_BlackboardBase
 	AIMODULE_API virtual void DescribeRuntimeValues(const UBehaviorTreeComponent& OwnerComp, uint8* NodeMemory, EBTDescriptionVerbosity::Type Verbosity, TArray<FString>& Values) const override;
 	AIMODULE_API virtual FString GetStaticDescription() const override;
 	AIMODULE_API virtual uint16 GetInstanceMemorySize() const override;
+	AIMODULE_API virtual void InitializeMemory(UBehaviorTreeComponent& OwnerComp, uint8* NodeMemory, EBTMemoryInit::Type InitType) const override;
+	AIMODULE_API virtual void CleanupMemory(UBehaviorTreeComponent& OwnerComp, uint8* NodeMemory, EBTMemoryClear::Type CleanupType) const override;
 
 	/** finish task */
 	AIMODULE_API void OnQueryFinished(TSharedPtr<FEnvQueryResult> Result);
diff --git a/Engine/Source/Runtime/AIModule/Private/BehaviorTree/BTCompositeNode.cpp b/Engine/Source/Runtime/AIModule/Private/BehaviorTree/BTCompositeNode.cpp
--- a/Engine/Source/Runtime/AIModule/Private/BehaviorTree/BTCompositeNode.cpp
+++ b/Engine/Source/Runtime/AIModule/Private/BehaviorTree/BTCompositeNode.cpp
@@ -704,3 +704,14 @@ uint16 UBTCompositeNode::GetInstanceMemorySize() const
 	return sizeof(FBTCompositeMemory);
 }
 
+void UBTCompositeNode::InitializeMemory(UBehaviorTreeComponent& OwnerComp, uint8* NodeMemory, EBTMemoryInit::Type InitType) const
+{
+	InitializeNodeMemory<FBTCompositeMemory>(NodeMemory, InitType);
+}
+
+void UBTCompositeNode::CleanupMemory(UBehaviorTreeComponent& OwnerComp, uint8* NodeMemory, EBTMemoryClear::Type CleanupType) const
+{
+	CleanupNodeMemory<FBTCompositeMemory>(NodeMemory, CleanupType);
+}
+
+
diff --git a/Engine/Source/Runtime/AIModule/Private/BehaviorTree/Composites/BTComposite_SimpleParallel.cpp b/Engine/Source/Runtime/AIModule/Private/BehaviorTree/Composites/BTComposite_SimpleParallel.cpp
--- a/Engine/Source/Runtime/AIModule/Private/BehaviorTree/Composites/BTComposite_SimpleParallel.cpp
+++ b/Engine/Source/Runtime/AIModule/Private/BehaviorTree/Composites/BTComposite_SimpleParallel.cpp
@@ -146,6 +146,16 @@ uint16 UBTComposite_SimpleParallel::GetInstanceMemorySize() const
 	return sizeof(FBTParallelMemory);
 }
 
+void UBTComposite_SimpleParallel::InitializeMemory(UBehaviorTreeComponent& OwnerComp, uint8* NodeMemory, EBTMemoryInit::Type InitType) const
+{
+	InitializeNodeMemory<FBTParallelMemory>(NodeMemory, InitType);
+}
+
+void UBTComposite_SimpleParallel::CleanupMemory(UBehaviorTreeComponent& OwnerComp, uint8* NodeMemory, EBTMemoryClear::Type CleanupType) const
+{
+	CleanupNodeMemory<FBTParallelMemory>(NodeMemory, CleanupType);
+}
+
 bool UBTComposite_SimpleParallel::CanPushSubtree(UBehaviorTreeComponent& OwnerComp, uint8* NodeMemory, int32 ChildIdx) const
 {
 	return (ChildIdx != EBTParallelChild::MainTask);
diff --git a/Engine/Source/Runtime/AIModule/Private/BehaviorTree/Decorators/BTDecorator_ConditionalLoop.cpp b/Engine/Source/Runtime/AIModule/Private/BehaviorTree/Decorators/BTDecorator_ConditionalLoop.cpp
--- a/Engine/Source/Runtime/AIModule/Private/BehaviorTree/Decorators/BTDecorator_ConditionalLoop.cpp
+++ b/Engine/Source/Runtime/AIModule/Private/BehaviorTree/Decorators/BTDecorator_ConditionalLoop.cpp
@@ -59,6 +59,16 @@ uint16 UBTDecorator_ConditionalLoop::GetInstanceMemorySize() const
 	return sizeof(FBTConditionalLoopDecoratorMemory);
 }
 
+void UBTDecorator_ConditionalLoop::InitializeMemory(UBehaviorTreeComponent& OwnerComp, uint8* NodeMemory, EBTMemoryInit::Type InitType) const
+{
+	InitializeNodeMemory<FBTConditionalLoopDecoratorMemory>(NodeMemory, InitType);
+}
+
+void UBTDecorator_ConditionalLoop::CleanupMemory(UBehaviorTreeComponent& OwnerComp, uint8* NodeMemory, EBTMemoryClear::Type CleanupType) const
+{
+	CleanupNodeMemory<FBTConditionalLoopDecoratorMemory>(NodeMemory, CleanupType);
+}
+
 #if WITH_EDITOR
 
 FName UBTDecorator_ConditionalLoop::GetNodeIconName() const
diff --git a/Engine/Source/Runtime/AIModule/Private/BehaviorTree/Decorators/BTDecorator_ConeCheck.cpp b/Engine/Source/Runtime/AIModule/Private/BehaviorTree/Decorators/BTDecorator_ConeCheck.cpp
--- a/Engine/Source/Runtime/AIModule/Private/BehaviorTree/Decorators/BTDecorator_ConeCheck.cpp
+++ b/Engine/Source/Runtime/AIModule/Private/BehaviorTree/Decorators/BTDecorator_ConeCheck.cpp
@@ -141,6 +141,16 @@ uint16 UBTDecorator_ConeCheck::GetInstanceMemorySize() const
 	return sizeof(TNodeInstanceMemory);
 }
 
+void UBTDecorator_ConeCheck::InitializeMemory(UBehaviorTreeComponent& OwnerComp, uint8* NodeMemory, EBTMemoryInit::Type InitType) const
+{
+	InitializeNodeMemory<TNodeInstanceMemory>(NodeMemory, InitType);
+}
+
+void UBTDecorator_ConeCheck::CleanupMemory(UBehaviorTreeComponent& OwnerComp, uint8* NodeMemory, EBTMemoryClear::Type CleanupType) const
+{
+	CleanupNodeMemory<TNodeInstanceMemory>(NodeMemory, CleanupType);
+}
+
 #if WITH_EDITOR
 
 FName UBTDecorator_ConeCheck::GetNodeIconName() const
diff --git a/Engine/Source/Runtime/AIModule/Private/BehaviorTree/Decorators/BTDecorator_Cooldown.cpp b/Engine/Source/Runtime/AIModule/Private/BehaviorTree/Decorators/BTDecorator_Cooldown.cpp
--- a/Engine/Source/Runtime/AIModule/Private/BehaviorTree/Decorators/BTDecorator_Cooldown.cpp
+++ b/Engine/Source/Runtime/AIModule/Private/BehaviorTree/Decorators/BTDecorator_Cooldown.cpp
@@ -28,17 +28,6 @@ bool UBTDecorator_Cooldown::CalculateRawConditionValue(UBehaviorTreeComponent& O
 	return RecalcTime >= DecoratorMemory->LastUseTimestamp;
 }
 
-void UBTDecorator_Cooldown::InitializeMemory(UBehaviorTreeComponent& OwnerComp, uint8* NodeMemory, EBTMemoryInit::Type InitType) const
-{
-	FBTCooldownDecoratorMemory* DecoratorMemory = CastInstanceNodeMemory<FBTCooldownDecoratorMemory>(NodeMemory);
-	if (InitType == EBTMemoryInit::Initialize)
-	{
-		DecoratorMemory->LastUseTimestamp = TNumericLimits<double>::Lowest();
-	}
-
-	DecoratorMemory->bRequestedRestart = false;
-}
-
 void UBTDecorator_Cooldown::OnNodeDeactivation(FBehaviorTreeSearchData& SearchData, EBTNodeResult::Type NodeResult)
 {
 	FBTCooldownDecoratorMemory* DecoratorMemory = GetNodeMemory<FBTCooldownDecoratorMemory>(SearchData);
@@ -87,6 +76,22 @@ uint16 UBTDecorator_Cooldown::GetInstanceMemorySize() const
 	return sizeof(FBTCooldownDecoratorMemory);
 }
 
+void UBTDecorator_Cooldown::InitializeMemory(UBehaviorTreeComponent& OwnerComp, uint8* NodeMemory, EBTMemoryInit::Type InitType) const
+{
+	FBTCooldownDecoratorMemory* DecoratorMemory = InitializeNodeMemory<FBTCooldownDecoratorMemory>(NodeMemory, InitType);
+	if (InitType == EBTMemoryInit::Initialize)
+	{
+		DecoratorMemory->LastUseTimestamp = TNumericLimits<double>::Lowest();
+	}
+
+	DecoratorMemory->bRequestedRestart = false; ;
+}
+
+void UBTDecorator_Cooldown::CleanupMemory(UBehaviorTreeComponent& OwnerComp, uint8* NodeMemory, EBTMemoryClear::Type CleanupType) const
+{
+	CleanupNodeMemory<FBTCooldownDecoratorMemory>(NodeMemory, CleanupType);
+}
+
 #if WITH_EDITOR
 
 FName UBTDecorator_Cooldown::GetNodeIconName() const
diff --git a/Engine/Source/Runtime/AIModule/Private/BehaviorTree/Decorators/BTDecorator_KeepInCone.cpp b/Engine/Source/Runtime/AIModule/Private/BehaviorTree/Decorators/BTDecorator_KeepInCone.cpp
--- a/Engine/Source/Runtime/AIModule/Private/BehaviorTree/Decorators/BTDecorator_KeepInCone.cpp
+++ b/Engine/Source/Runtime/AIModule/Private/BehaviorTree/Decorators/BTDecorator_KeepInCone.cpp
@@ -132,6 +132,16 @@ uint16 UBTDecorator_KeepInCone::GetInstanceMemorySize() const
 	return sizeof(TNodeInstanceMemory);
 }
 
+void UBTDecorator_KeepInCone::InitializeMemory(UBehaviorTreeComponent& OwnerComp, uint8* NodeMemory, EBTMemoryInit::Type InitType) const
+{
+	InitializeNodeMemory<TNodeInstanceMemory>(NodeMemory, InitType);
+}
+
+void UBTDecorator_KeepInCone::CleanupMemory(UBehaviorTreeComponent& OwnerComp, uint8* NodeMemory, EBTMemoryClear::Type CleanupType) const
+{
+	CleanupNodeMemory<TNodeInstanceMemory>(NodeMemory, CleanupType);
+}
+
 #if WITH_EDITOR
 
 FName UBTDecorator_KeepInCone::GetNodeIconName() const
diff --git a/Engine/Source/Runtime/AIModule/Private/BehaviorTree/Decorators/BTDecorator_Loop.cpp b/Engine/Source/Runtime/AIModule/Private/BehaviorTree/Decorators/BTDecorator_Loop.cpp
--- a/Engine/Source/Runtime/AIModule/Private/BehaviorTree/Decorators/BTDecorator_Loop.cpp
+++ b/Engine/Source/Runtime/AIModule/Private/BehaviorTree/Decorators/BTDecorator_Loop.cpp
@@ -107,6 +107,16 @@ uint16 UBTDecorator_Loop::GetInstanceMemorySize() const
 	return sizeof(FBTLoopDecoratorMemory);
 }
 
+void UBTDecorator_Loop::InitializeMemory(UBehaviorTreeComponent& OwnerComp, uint8* NodeMemory, EBTMemoryInit::Type InitType) const
+{
+	InitializeNodeMemory<FBTLoopDecoratorMemory>(NodeMemory, InitType);
+}
+
+void UBTDecorator_Loop::CleanupMemory(UBehaviorTreeComponent& OwnerComp, uint8* NodeMemory, EBTMemoryClear::Type CleanupType) const
+{
+	CleanupNodeMemory<FBTLoopDecoratorMemory>(NodeMemory, CleanupType);
+}
+
 #if WITH_EDITOR
 
 FName UBTDecorator_Loop::GetNodeIconName() const
diff --git a/Engine/Source/Runtime/AIModule/Private/BehaviorTree/Decorators/BTDecorator_TagCooldown.cpp b/Engine/Source/Runtime/AIModule/Private/BehaviorTree/Decorators/BTDecorator_TagCooldown.cpp
--- a/Engine/Source/Runtime/AIModule/Private/BehaviorTree/Decorators/BTDecorator_TagCooldown.cpp
+++ b/Engine/Source/Runtime/AIModule/Private/BehaviorTree/Decorators/BTDecorator_TagCooldown.cpp
@@ -43,10 +43,15 @@ bool UBTDecorator_TagCooldown::CalculateRawConditionValue(UBehaviorTreeComponent
 
 void UBTDecorator_TagCooldown::InitializeMemory(UBehaviorTreeComponent& OwnerComp, uint8* NodeMemory, EBTMemoryInit::Type InitType) const
 {
-	FBTTagCooldownDecoratorMemory* DecoratorMemory = CastInstanceNodeMemory<FBTTagCooldownDecoratorMemory>(NodeMemory);
+	FBTTagCooldownDecoratorMemory* DecoratorMemory = InitializeNodeMemory<FBTTagCooldownDecoratorMemory>(NodeMemory, InitType);
 	DecoratorMemory->bRequestedRestart = false;
 }
 
+void UBTDecorator_TagCooldown::CleanupMemory(UBehaviorTreeComponent& OwnerComp, uint8* NodeMemory, EBTMemoryClear::Type CleanupType) const
+{
+	CleanupNodeMemory<FBTTagCooldownDecoratorMemory>(NodeMemory, CleanupType);
+}
+
 void UBTDecorator_TagCooldown::OnNodeDeactivation(FBehaviorTreeSearchData& SearchData, EBTNodeResult::Type NodeResult)
 {
 	FBTTagCooldownDecoratorMemory* DecoratorMemory = GetNodeMemory<FBTTagCooldownDecoratorMemory>(SearchData);
diff --git a/Engine/Source/Runtime/AIModule/Private/BehaviorTree/Services/BTService_DefaultFocus.cpp b/Engine/Source/Runtime/AIModule/Private/BehaviorTree/Services/BTService_DefaultFocus.cpp
--- a/Engine/Source/Runtime/AIModule/Private/BehaviorTree/Services/BTService_DefaultFocus.cpp
+++ b/Engine/Source/Runtime/AIModule/Private/BehaviorTree/Services/BTService_DefaultFocus.cpp
@@ -139,6 +139,16 @@ EBlackboardNotificationResult UBTService_DefaultFocus::OnBlackboardKeyValueChang
 	return EBlackboardNotificationResult::ContinueObserving;
 }
 
+void UBTService_DefaultFocus::InitializeMemory(UBehaviorTreeComponent& OwnerComp, uint8* NodeMemory, EBTMemoryInit::Type InitType) const
+{
+	InitializeNodeMemory<FBTFocusMemory>(NodeMemory, InitType);
+}
+
+void UBTService_DefaultFocus::CleanupMemory(UBehaviorTreeComponent& OwnerComp, uint8* NodeMemory, EBTMemoryClear::Type CleanupType) const
+{
+	CleanupNodeMemory<FBTFocusMemory>(NodeMemory, CleanupType);
+}
+
 #if WITH_EDITOR
 
 FName UBTService_DefaultFocus::GetNodeIconName() const
diff --git a/Engine/Source/Runtime/AIModule/Private/BehaviorTree/Services/BTService_RunEQS.cpp b/Engine/Source/Runtime/AIModule/Private/BehaviorTree/Services/BTService_RunEQS.cpp
--- a/Engine/Source/Runtime/AIModule/Private/BehaviorTree/Services/BTService_RunEQS.cpp
+++ b/Engine/Source/Runtime/AIModule/Private/BehaviorTree/Services/BTService_RunEQS.cpp
@@ -148,20 +148,22 @@ void UBTService_RunEQS::OnCeaseRelevant(UBehaviorTreeComponent& OwnerComp, uint8
 
 void UBTService_RunEQS::InitializeMemory(UBehaviorTreeComponent& OwnerComp, uint8* NodeMemory, EBTMemoryInit::Type InitType) const
 {
-	FBTEQSServiceMemory* MyMemory = CastInstanceNodeMemory<FBTEQSServiceMemory>(NodeMemory);
+	FBTEQSServiceMemory* MyMemory = InitializeNodeMemory<FBTEQSServiceMemory>(NodeMemory, InitType);
 	check(MyMemory);
 	MyMemory->RequestID = INDEX_NONE;
 }
 
-#if WITH_EDITOR
-
 void UBTService_RunEQS::CleanupMemory(UBehaviorTreeComponent& OwnerComp, uint8* NodeMemory, EBTMemoryClear::Type CleanupType) const
 {
+#if WITH_EDITOR
 	const FBTEQSServiceMemory* MyMemory = CastInstanceNodeMemory<FBTEQSServiceMemory>(NodeMemory);
 	check(MyMemory);
 	ensure(MyMemory->RequestID == INDEX_NONE);
+#endif // WITH_EDITOR
+	CleanupNodeMemory<FBTEQSServiceMemory>(NodeMemory, CleanupType);
 }
 
+#if WITH_EDITOR
 void UBTService_RunEQS::PostEditChangeProperty(struct FPropertyChangedEvent& PropertyChangedEvent)
 {
 	Super::PostEditChangeProperty(PropertyChangedEvent);
diff --git a/Engine/Source/Runtime/AIModule/Private/BehaviorTree/Tasks/BTTask_GameplayTaskBase.cpp b/Engine/Source/Runtime/AIModule/Private/BehaviorTree/Tasks/BTTask_GameplayTaskBase.cpp
--- a/Engine/Source/Runtime/AIModule/Private/BehaviorTree/Tasks/BTTask_GameplayTaskBase.cpp
+++ b/Engine/Source/Runtime/AIModule/Private/BehaviorTree/Tasks/BTTask_GameplayTaskBase.cpp
@@ -96,3 +96,12 @@ uint16 UBTTask_GameplayTaskBase::GetInstanceMemorySize() const
 	return sizeof(FBTGameplayTaskMemory);
 }
 
+void UBTTask_GameplayTaskBase::InitializeMemory(UBehaviorTreeComponent& OwnerComp, uint8* NodeMemory, EBTMemoryInit::Type InitType) const
+{
+	InitializeNodeMemory<FBTGameplayTaskMemory>(NodeMemory, InitType);
+}
+
+void UBTTask_GameplayTaskBase::CleanupMemory(UBehaviorTreeComponent& OwnerComp, uint8* NodeMemory, EBTMemoryClear::Type CleanupType) const
+{
+	CleanupNodeMemory<FBTGameplayTaskMemory>(NodeMemory, CleanupType);
+}
\ No newline at end of file
diff --git a/Engine/Source/Runtime/AIModule/Private/BehaviorTree/Tasks/BTTask_MoveTo.cpp b/Engine/Source/Runtime/AIModule/Private/BehaviorTree/Tasks/BTTask_MoveTo.cpp
--- a/Engine/Source/Runtime/AIModule/Private/BehaviorTree/Tasks/BTTask_MoveTo.cpp
+++ b/Engine/Source/Runtime/AIModule/Private/BehaviorTree/Tasks/BTTask_MoveTo.cpp
@@ -337,6 +337,16 @@ uint16 UBTTask_MoveTo::GetInstanceMemorySize() const
 	return sizeof(FBTMoveToTaskMemory);
 }
 
+void UBTTask_MoveTo::InitializeMemory(UBehaviorTreeComponent& OwnerComp, uint8* NodeMemory, EBTMemoryInit::Type InitType) const
+{
+	InitializeNodeMemory<FBTMoveToTaskMemory>(NodeMemory, InitType);
+}
+
+void UBTTask_MoveTo::CleanupMemory(UBehaviorTreeComponent& OwnerComp, uint8* NodeMemory, EBTMemoryClear::Type CleanupType) const
+{
+	CleanupNodeMemory<FBTMoveToTaskMemory>(NodeMemory, CleanupType);
+}
+
 void UBTTask_MoveTo::PostLoad()
 {
 	Super::PostLoad();
diff --git a/Engine/Source/Runtime/AIModule/Private/BehaviorTree/Tasks/BTTask_RotateToFaceBBEntry.cpp b/Engine/Source/Runtime/AIModule/Private/BehaviorTree/Tasks/BTTask_RotateToFaceBBEntry.cpp
--- a/Engine/Source/Runtime/AIModule/Private/BehaviorTree/Tasks/BTTask_RotateToFaceBBEntry.cpp
+++ b/Engine/Source/Runtime/AIModule/Private/BehaviorTree/Tasks/BTTask_RotateToFaceBBEntry.cpp
@@ -233,3 +233,12 @@ FString UBTTask_RotateToFaceBBEntry::GetStaticDescription() const
 	return FString::Printf(TEXT("%s: %s"), *Super::GetStaticDescription(), *KeyDesc);
 }
 
+void UBTTask_RotateToFaceBBEntry::InitializeMemory(UBehaviorTreeComponent& OwnerComp, uint8* NodeMemory, EBTMemoryInit::Type InitType) const
+{
+	InitializeNodeMemory<FBTFocusMemory>(NodeMemory, InitType);
+}
+
+void UBTTask_RotateToFaceBBEntry::CleanupMemory(UBehaviorTreeComponent& OwnerComp, uint8* NodeMemory, EBTMemoryClear::Type CleanupType) const
+{
+	CleanupNodeMemory<FBTFocusMemory>(NodeMemory, CleanupType);
+}
diff --git a/Engine/Source/Runtime/AIModule/Private/BehaviorTree/Tasks/BTTask_RunEQSQuery.cpp b/Engine/Source/Runtime/AIModule/Private/BehaviorTree/Tasks/BTTask_RunEQSQuery.cpp
--- a/Engine/Source/Runtime/AIModule/Private/BehaviorTree/Tasks/BTTask_RunEQSQuery.cpp
+++ b/Engine/Source/Runtime/AIModule/Private/BehaviorTree/Tasks/BTTask_RunEQSQuery.cpp
@@ -104,6 +104,16 @@ uint16 UBTTask_RunEQSQuery::GetInstanceMemorySize() const
 	return sizeof(FBTEnvQueryTaskMemory);
 }
 
+void UBTTask_RunEQSQuery::InitializeMemory(UBehaviorTreeComponent& OwnerComp, uint8* NodeMemory, EBTMemoryInit::Type InitType) const
+{
+	InitializeNodeMemory<FBTEnvQueryTaskMemory>(NodeMemory, InitType);
+}
+
+void UBTTask_RunEQSQuery::CleanupMemory(UBehaviorTreeComponent& OwnerComp, uint8* NodeMemory, EBTMemoryClear::Type CleanupType) const
+{
+	CleanupNodeMemory<FBTEnvQueryTaskMemory>(NodeMemory, CleanupType);
+}
+
 void UBTTask_RunEQSQuery::OnQueryFinished(TSharedPtr<FEnvQueryResult> Result)
 {
 	if (Result->IsAborted())
