Added more Logs BT auxiliary node and converted some EQS log to vlog

diff --git a/Engine/Source/Runtime/AIModule/Private/BehaviorTree/BTAuxiliaryNode.cpp b/Engine/Source/Runtime/AIModule/Private/BehaviorTree/BTAuxiliaryNode.cpp
--- a/Engine/Source/Runtime/AIModule/Private/BehaviorTree/BTAuxiliaryNode.cpp
+++ b/Engine/Source/Runtime/AIModule/Private/BehaviorTree/BTAuxiliaryNode.cpp
@@ -22,6 +22,7 @@ void UBTAuxiliaryNode::WrappedOnBecomeRelevant(UBehaviorTreeComponent& OwnerComp
 		const UBTNode* NodeOb = HasInstance() ? GetNodeInstance(OwnerComp, NodeMemory) : this;
 		if (NodeOb)
 		{
+			UE_VLOG(&OwnerComp, LogBehaviorTree, Verbose, TEXT("OnBecomeRelevant: %s"), *UBehaviorTreeTypes::DescribeNodeHelper(NodeOb));
 			((UBTAuxiliaryNode*)NodeOb)->OnBecomeRelevant(OwnerComp, NodeMemory);
 		}
 	}
@@ -34,6 +35,7 @@ void UBTAuxiliaryNode::WrappedOnCeaseRelevant(UBehaviorTreeComponent& OwnerComp,
 		const UBTNode* NodeOb = HasInstance() ? GetNodeInstance(OwnerComp, NodeMemory) : this;
 		if (NodeOb)
 		{
+			UE_VLOG(&OwnerComp, LogBehaviorTree, Verbose, TEXT("OnCeaseRelevant: %s"), *UBehaviorTreeTypes::DescribeNodeHelper(NodeOb));
 			((UBTAuxiliaryNode*)NodeOb)->OnCeaseRelevant(OwnerComp, NodeMemory);
 		}
 	}
diff --git a/Engine/Source/Runtime/AIModule/Private/EnvironmentQuery/EnvQueryManager.cpp b/Engine/Source/Runtime/AIModule/Private/EnvironmentQuery/EnvQueryManager.cpp
--- a/Engine/Source/Runtime/AIModule/Private/EnvironmentQuery/EnvQueryManager.cpp
+++ b/Engine/Source/Runtime/AIModule/Private/EnvironmentQuery/EnvQueryManager.cpp
@@ -291,7 +291,10 @@ int32 UEnvQueryManager::RunQuery(const TSharedPtr<FEnvQueryInstance>& QueryInsta
 	QueryInstance->FinishDelegate = FinishDelegate;
 	RunningQueries.Add(QueryInstance);
 
-	UE_LOG(LogEQS, Verbose, TEXT("%s: Query: %s - Owner: %s"), ANSI_TO_TCHAR(__FUNCTION__), *QueryInstance->QueryName, *GetNameSafe(QueryInstance->Owner.Get()));
+	UE_VLOG_ALWAYS_UELOG(QueryInstance->Owner.Get(), LogEQS, Verbose, TEXT("%hs: Query: %s - Owner: %s"),
+		__FUNCTION__,
+		*QueryInstance->QueryName,
+		*GetNameSafe(QueryInstance->Owner.Get()));
 
 	return QueryInstance->QueryID;
 }
@@ -316,7 +319,10 @@ void UEnvQueryManager::RunInstantQuery(const TSharedPtr<FEnvQueryInstance>& Quer
 		return;
 	}
 
-	UE_LOG(LogEQS, Verbose, TEXT("%s: Query: %s - Owner: %s"), ANSI_TO_TCHAR(__FUNCTION__), *QueryInstance->QueryName, *GetNameSafe(QueryInstance->Owner.Get()));
+	UE_VLOG_ALWAYS_UELOG(QueryInstance->Owner.Get(), LogEQS, Verbose, TEXT("%hs: Query: %s - Owner: %s"),
+		__FUNCTION__,
+		*QueryInstance->QueryName,
+		*GetNameSafe(QueryInstance->Owner.Get()));
 
 	{
 		CSV_SCOPED_TIMING_STAT_EXCLUSIVE(EnvQueryManager);
@@ -348,7 +354,10 @@ void UEnvQueryManager::RemoveAllQueriesByQuerier(const UObject& Querier, bool bE
 			{
 				QueryInstance->MarkAsAborted();
 
-				UE_LOG(LogEQS, Verbose, TEXT("%s: Query: %s - Owner: %s"), ANSI_TO_TCHAR(__FUNCTION__), *QueryInstance->QueryName, *GetNameSafe(QueryInstance->Owner.Get()));
+				UE_VLOG_ALWAYS_UELOG(QueryInstance->Owner.Get(), LogEQS, Verbose, TEXT("%hs: Query: %s - Owner: %s"),
+					__FUNCTION__,
+					*QueryInstance->QueryName,
+					*GetNameSafe(QueryInstance->Owner.Get()));
 
 				if (bExecuteFinishDelegate)
 				{
@@ -367,10 +376,11 @@ TSharedPtr<FEnvQueryInstance> UEnvQueryManager::PrepareQueryInstance(const FEnvQ
 	TSharedPtr<FEnvQueryInstance> QueryInstance = CreateQueryInstance(Request.QueryTemplate, RunMode);
 	if (!QueryInstance.IsValid())
 	{
-		UE_LOG(LogEQS, Warning, TEXT("Error creating query instance for QueryTemplate: %s - Owner: %s"),
+		UE_VLOG_ALWAYS_UELOG(Request.Owner.Get(), LogEQS, Warning, TEXT("Error creating query instance for QueryTemplate: %s - Owner: %s"),
 			Request.QueryTemplate != nullptr ? *Request.QueryTemplate->QueryName.ToString() : TEXT("NONE"),
 			*GetNameSafe(Request.Owner));
-		return NULL;
+
+		return nullptr;
 	}
 
 	QueryInstance->World = GetWorldFast();
@@ -400,7 +410,10 @@ bool UEnvQueryManager::AbortQuery(int32 RequestID)
 		if (QueryInstance->QueryID == RequestID &&
 			QueryInstance->IsFinished() == false)
 		{
-			UE_LOG(LogEQS, Verbose, TEXT("%s: Query: %s - Owner: %s"), ANSI_TO_TCHAR(__FUNCTION__), *QueryInstance->QueryName, *GetNameSafe(QueryInstance->Owner.Get()));
+			UE_VLOG_ALWAYS_UELOG(QueryInstance->Owner.Get(), LogEQS, Verbose, TEXT("%hs: Query: %s - Owner: %s"),
+				__FUNCTION__,
+				*QueryInstance->QueryName,
+				*GetNameSafe(QueryInstance->Owner.Get()));
 
 			QueryInstance->MarkAsAborted();
 			QueryInstance->FinishDelegate.ExecuteIfBound(QueryInstance);
@@ -466,7 +479,7 @@ void UEnvQueryManager::Tick(float DeltaTime)
 					// Always log that we executed total execution time at the end of the query.
 					if (QueryInstancePtr->TotalExecutionTime > ExecutionTimeWarningSeconds)
 					{
-						UE_LOG(LogEQS, Warning, TEXT("Query %s (Owner: %s) has finished in %.2f ms, exceeding the configured limit of %.2f ms. Execution details: %s"), 
+						UE_VLOG_ALWAYS_UELOG(QueryInstance->Owner.Get(), LogEQS, Warning, TEXT("Query %s (Owner: %s) has finished in %.2f ms, exceeding the configured limit of %.2f ms. Execution details: %s"), 
 							*QueryInstancePtr->QueryName, *GetNameSafe(QueryInstancePtr->Owner.Get()), 
 							1000.f * QueryInstancePtr->TotalExecutionTime, 1000.f * ExecutionTimeWarningSeconds, 
 							*QueryInstancePtr->GetExecutionTimeDescription());
@@ -487,7 +500,10 @@ void UEnvQueryManager::Tick(float DeltaTime)
 						EQSDebugger.StoreQuery(QueryInstance);
 #endif // USE_EQS_DEBUGGER
 
-						UE_LOG(LogEQS, Verbose, TEXT("%s: Finished Query: %s - Owner: %s"), ANSI_TO_TCHAR(__FUNCTION__), *QueryInstance->QueryName, *GetNameSafe(QueryInstance->Owner.Get()));
+						UE_VLOG_ALWAYS_UELOG(QueryInstance->Owner.Get(), LogEQS, Verbose, TEXT("%hs: Finished Query: %s - Owner: %s"),
+							__FUNCTION__,
+							*QueryInstance->QueryName,
+							*GetNameSafe(QueryInstance->Owner.Get()));
 
 						QueryInstancePtr->FinishDelegate.ExecuteIfBound(QueryInstance);
 
@@ -500,7 +516,12 @@ void UEnvQueryManager::Tick(float DeltaTime)
 	#if USE_DELEGATE_TRYGETBOUNDFUNCTIONNAME
 							FunctionName = QueryInstancePtr->FinishDelegate.TryGetBoundFunctionName();
 	#endif // USE_DELEGATE_TRYGETBOUNDFUNCTIONNAME
-							UE_LOG(LogEQS, Warning, TEXT("FinishDelegate for EQS query %s took %f seconds and is over handling time limit warning of %f. Delegate info: object = %s method = %s"), *QueryInstancePtr->QueryName, ResultHandlingDuration, HandlingResultTimeWarningSeconds , *GetNameSafe(QueryInstancePtr->FinishDelegate.GetUObject()), *FunctionName.ToString());
+
+							UE_VLOG_ALWAYS_UELOG(QueryInstance->Owner.Get(), LogEQS, Warning,
+								TEXT("FinishDelegate for EQS query %s took %f seconds and is over handling time limit warning of %f. Delegate info: object = %s method = %s"),
+								*QueryInstancePtr->QueryName, ResultHandlingDuration, HandlingResultTimeWarningSeconds,
+								*GetNameSafe(QueryInstancePtr->FinishDelegate.GetUObject()),
+								*FunctionName.ToString());
 						}
 					}
 
@@ -516,7 +537,7 @@ void UEnvQueryManager::Tick(float DeltaTime)
 
 				if (QueryInstancePtr->TotalExecutionTime > ExecutionTimeWarningSeconds && !QueryInstancePtr->bHasLoggedTimeLimitWarning)
 				{
-					UE_LOG(LogEQS, Warning, TEXT("Query %s (Owner: %s) has taken %.2f ms so far, exceeding the configured limit of %.2f ms. Execution details: %s"), 
+					UE_VLOG_ALWAYS_UELOG(QueryInstance->Owner.Get(), LogEQS, Warning, TEXT("Query %s (Owner: %s) has taken %.2f ms so far, exceeding the configured limit of %.2f ms. Execution details: %s"), 
 						*QueryInstancePtr->QueryName, *GetNameSafe(QueryInstancePtr->Owner.Get()), 
 						1000.f * QueryInstancePtr->TotalExecutionTime, 1000.f * ExecutionTimeWarningSeconds, 
 						*QueryInstancePtr->GetExecutionTimeDescription());
@@ -623,22 +644,25 @@ void UEnvQueryManager::LogQueryInfo(bool bDisplayThresholdWarning) const
 {
 	if (bDisplayThresholdWarning)
 	{
-		UE_LOG(LogEQS, Warning, TEXT("The number of EQS queries (%d) has reached the warning threshold (%d).  Logging queries."), RunningQueries.Num(), QueryCountWarningThreshold);
+		UE_VLOG_ALWAYS_UELOG(this, LogEQS, Warning, TEXT("The number of EQS queries (%d) has reached the warning threshold (%d).  Logging queries."), RunningQueries.Num(), QueryCountWarningThreshold);
 	}
 	else
 	{
-		UE_LOG(LogEQS, Warning, TEXT("The number of EQS queries is (%d).  Logging queries."), RunningQueries.Num());
+		UE_VLOG_ALWAYS_UELOG(this, LogEQS, Warning, TEXT("The number of EQS queries is (%d).  Logging queries."), RunningQueries.Num());
 	}
 
 	for (const TSharedPtr<FEnvQueryInstance>& RunningQuery : RunningQueries)
 	{
 		if (RunningQuery.IsValid())
 		{
-			UE_LOG(LogEQS, Warning, TEXT("Query: %s - Owner: %s - %s"), *RunningQuery->QueryName, RunningQuery->Owner.IsValid() ? *RunningQuery->Owner->GetName() : TEXT("Invalid"), *RunningQuery->GetExecutionTimeDescription());
+			UE_VLOG_ALWAYS_UELOG(this, LogEQS, Warning, TEXT("Query: %s - Owner: %s - %s"),
+				*RunningQuery->QueryName,
+				RunningQuery->Owner.IsValid() ? *RunningQuery->Owner->GetName() : TEXT("Invalid"),
+				*RunningQuery->GetExecutionTimeDescription());
 		}
 		else
 		{
-			UE_LOG(LogEQS, Warning, TEXT("Invalid query found in list!"));
+			UE_VLOG_ALWAYS_UELOG(this, LogEQS, Warning, TEXT("Invalid query found in list!"));
 		}
 	}
 }
@@ -664,7 +688,10 @@ void UEnvQueryManager::OnWorldCleanup()
 			TSharedPtr<FEnvQueryInstance>& QueryInstance = RunningQueriesCopy[Index];
 			if (QueryInstance->IsFinished() == false)
 			{
-				UE_LOG(LogEQS, Verbose, TEXT("Query failed due to world cleanup: Query: %s - Owner: %s"), *QueryInstance->QueryName, *GetNameSafe(QueryInstance->Owner.Get()));
+				UE_VLOG_ALWAYS_UELOG(QueryInstance->Owner.Get(), LogEQS, Verbose, TEXT("Query failed due to world cleanup: Query: %s - Owner: %s"),
+					*QueryInstance->QueryName,
+					*GetNameSafe(QueryInstance->Owner.Get()));
+
 				QueryInstance->MarkAsFailed();
 				QueryInstance->FinishDelegate.ExecuteIfBound(QueryInstance);
 			}
@@ -737,7 +764,7 @@ TSharedPtr<FEnvQueryInstance> UEnvQueryManager::CreateQueryInstance(const UEnvQu
 {
 	if (Template == nullptr || Template->Options.Num() == 0)
 	{
-		UE_CLOG(Template != nullptr && Template->Options.Num() == 0, LogEQS, Warning, TEXT("Query [%s] doesn't have any valid options!"), *Template->GetName());
+		UE_CVLOG_ALWAYS_UELOG(Template != nullptr && Template->Options.Num() == 0, this, LogEQS, Warning, TEXT("Query [%s] doesn't have any valid options!"), *Template->GetName());
 		return nullptr;
 	}
 
@@ -793,17 +820,17 @@ TSharedPtr<FEnvQueryInstance> UEnvQueryManager::CreateQueryInstance(const UEnvQu
 			{
 				if (MyOption == nullptr)
 				{
-					UE_LOG(LogEQS, Error, TEXT("Trying to spawn a query with broken Template (null option): %s, option %d"),
+					UE_VLOG_ALWAYS_UELOG(this, LogEQS, Error, TEXT("Trying to spawn a query with broken Template (null option): %s, option %d"),
 						*GetNameSafe(LocalTemplate), OptionIndex);
 				} 
 				else if (MyOption->Generator == nullptr)
 				{
-					UE_LOG(LogEQS, Error, TEXT("Trying to spawn a query with broken Template (generator:MISSING): %s, option %d"),
+					UE_VLOG_ALWAYS_UELOG(this, LogEQS, Error, TEXT("Trying to spawn a query with broken Template (generator:MISSING): %s, option %d"),
 						*GetNameSafe(LocalTemplate), OptionIndex);
 				}
 				else
 				{
-					UE_LOG(LogEQS, Error, TEXT("Trying to spawn a query with broken Template (generator:%s itemType:%s): %s, option %d"),
+					UE_VLOG_ALWAYS_UELOG(this, LogEQS, Error, TEXT("Trying to spawn a query with broken Template (generator:%s itemType:%s): %s, option %d"),
 						MyOption->Generator->IsValidGenerator() ? TEXT("ok") : TEXT("invalid"),
 						MyOption->Generator->ItemType ? TEXT("ok") : TEXT("MISSING"),
 						*GetNameSafe(LocalTemplate), OptionIndex);
@@ -850,7 +877,7 @@ TSharedPtr<FEnvQueryInstance> UEnvQueryManager::CreateQueryInstance(const UEnvQu
 				UEnvQueryTest* TestOb = SortedTests[TestIndex];
 				if (TestOb == NULL || !TestOb->IsSupportedItem(GeneratedType))
 				{
-					UE_LOG(LogEQS, Warning, TEXT("Query [%s] can't use test [%s] in option %d [%s], removing it"),
+					UE_VLOG_ALWAYS_UELOG(this, LogEQS, Warning, TEXT("Query [%s] can't use test [%s] in option %d [%s], removing it"),
 						*GetNameSafe(LocalTemplate), *GetNameSafe(TestOb), OptionIndex, *MyOption->Generator->OptionName);
 
 					SortedTests.RemoveAt(TestIndex, 1, false);
@@ -859,7 +886,7 @@ TSharedPtr<FEnvQueryInstance> UEnvQueryManager::CreateQueryInstance(const UEnvQu
 					&& TestOb->TestPurpose == EEnvTestPurpose::Filter
 					&& (HighestFilterCost < TestOb->Cost || MostExpensiveFilterIndex == INDEX_NONE))
 				{
-					HighestFilterCost = TestOb->Cost;					
+					HighestFilterCost = TestOb->Cost;
 					MostExpensiveFilterIndex = TestIndex;
 				}
 			}
@@ -1102,7 +1129,7 @@ bool UEnvQueryManager::Exec_Dev(UWorld* Inworld, const TCHAR* Cmd, FOutputDevice
 
 void UEnvQueryManager::Configure(const FEnvQueryManagerConfig& NewConfig)
 {
-	UE_LOG(LogEQS, Log, TEXT("Applying new FEnvQueryManagerConfig: %s"), *NewConfig.ToString());
+	UE_VLOG_ALWAYS_UELOG(this, LogEQS, Log, TEXT("Applying new FEnvQueryManagerConfig: %s"), *NewConfig.ToString());
 
 	MaxAllowedTestingTime = NewConfig.MaxAllowedTestingTime;
 	bTestQueriesUsingBreadth = NewConfig.bTestQueriesUsingBreadth;
