BTService to RunEQS should not be allowed to tick on search start to ensure proper query management.

diff --git a/Engine/Source/Runtime/AIModule/Classes/BehaviorTree/BTService.h b/Engine/Source/Runtime/AIModule/Classes/BehaviorTree/BTService.h
--- a/Engine/Source/Runtime/AIModule/Classes/BehaviorTree/BTService.h
+++ b/Engine/Source/Runtime/AIModule/Classes/BehaviorTree/BTService.h
@@ -56,7 +56,7 @@ protected:
 	float RandomDeviation;
 
 	/** call Tick event when task search enters this node (SearchStart will be called as well) */
-	UPROPERTY(Category = Service, EditAnywhere, AdvancedDisplay)
+	UPROPERTY(Category = Service, EditAnywhere, AdvancedDisplay, meta=(EditCondition = bCanTickOnSearchStartBeExposed, EditConditionHides))
 	uint32 bCallTickOnSearchStart : 1;
 
 	/** if set, next tick time will be always reset to service's interval when node is activated */
@@ -66,6 +66,12 @@ protected:
 	/** if set, service will be notified about search entering underlying branch */
 	uint32 bNotifyOnSearch : 1;
 
+#if WITH_EDITORONLY_DATA
+	/** if set, the service exposes the option to call Tick event when task search enters this node */
+	UPROPERTY(transient)
+	uint32 bCanTickOnSearchStartBeExposed : 1;
+#endif // WITH_EDITORONLY_DATA
+
 	/** update next tick interval
 	 * this function should be considered as const (don't modify state of object) if node is not instanced!
 	 * bNotifyTick must be set to true for this function to be called 
diff --git a/Engine/Source/Runtime/AIModule/Private/BehaviorTree/BTService.cpp b/Engine/Source/Runtime/AIModule/Private/BehaviorTree/BTService.cpp
--- a/Engine/Source/Runtime/AIModule/Private/BehaviorTree/BTService.cpp
+++ b/Engine/Source/Runtime/AIModule/Private/BehaviorTree/BTService.cpp
@@ -1,6 +1,7 @@
 // Copyright Epic Games, Inc. All Rights Reserved.
 
 #include "BehaviorTree/BTService.h"
+#include "VisualLogger/VisualLogger.h"
 
 #include UE_INLINE_GENERATED_CPP_BY_NAME(BTService)
 
@@ -9,6 +10,9 @@ UBTService::UBTService(const FObjectInitializer& ObjectInitializer) : Super(Obje
 	bNotifyTick = true;
 	bNotifyOnSearch = true;
 	bTickIntervals = true;
+#if WITH_EDITORONLY_DATA
+	bCanTickOnSearchStartBeExposed = true;
+#endif // WITH_EDITORONLY_DATA
 	bCallTickOnSearchStart = false;
 	bRestartTimerOnEachActivation = false;
 
@@ -47,11 +51,13 @@ void UBTService::NotifyParentActivation(FBehaviorTreeSearchData& SearchData)
 
 			if (bNotifyOnSearch)
 			{
+				UE_VLOG(SearchData.OwnerComp.GetOwner(), LogBehaviorTree, VeryVerbose, TEXT("OnSearchStart: %s"), *UBehaviorTreeTypes::DescribeNodeHelper(ServiceNodeOb));
 				ServiceNodeOb->OnSearchStart(SearchData);
 			}
 
 			if (bCallTickOnSearchStart)
 			{
+				UE_VLOG(SearchData.OwnerComp.GetOwner(), LogBehaviorTree, VeryVerbose, TEXT("TickNode (bCallTickOnSearchStart): %s"), *UBehaviorTreeTypes::DescribeNodeHelper(ServiceNodeOb));
 				ServiceNodeOb->TickNode(SearchData.OwnerComp, NodeMemory, 0.0f);
 			}
 		}
diff --git a/Engine/Source/Runtime/AIModule/Private/BehaviorTree/Services/BTService_RunEQS.cpp b/Engine/Source/Runtime/AIModule/Private/BehaviorTree/Services/BTService_RunEQS.cpp
--- a/Engine/Source/Runtime/AIModule/Private/BehaviorTree/Services/BTService_RunEQS.cpp
+++ b/Engine/Source/Runtime/AIModule/Private/BehaviorTree/Services/BTService_RunEQS.cpp
@@ -21,6 +21,13 @@ UBTService_RunEQS::UBTService_RunEQS(const FObjectInitializer& ObjectInitializer
 
 	bNotifyBecomeRelevant = false;
 	bNotifyCeaseRelevant = true;
+
+#if WITH_EDITORONLY_DATA
+	// Do not expose the option to tick on search start since the request is async and it
+	// requires the node to be relevant to properly manage the request and its delegate.
+	bCanTickOnSearchStartBeExposed = false;
+	bCallTickOnSearchStart = false;
+#endif // WITH_EDITORONLY_DATA
 	
 	// accept only actors and vectors
 	BlackboardKey.AddObjectFilter(this, GET_MEMBER_NAME_CHECKED(UBTService_RunEQS, BlackboardKey), AActor::StaticClass());
@@ -89,14 +96,8 @@ void UBTService_RunEQS::OnQueryFinished(TSharedPtr<FEnvQueryResult> Result)
 	}
 
 	FBTEQSServiceMemory* MyMemory = CastInstanceNodeMemory<FBTEQSServiceMemory>(BTComp->GetNodeMemory(this, BTComp->FindInstanceContainingNode(this)));
-	if (MyMemory == nullptr || MyMemory->RequestID == INDEX_NONE)
+	if (!ensureMsgf(MyMemory && MyMemory->RequestID != INDEX_NONE, TEXT("%hs called while the BT node is not or no longer active."), __FUNCTION__))
 	{
-		// this can happen in hard to repro edge cases if the BT node being notified here has already been canceled
-		// in which case it should be safe to ignore it.
-		// @todo we need to find the exact case and repro it via unit tests
-		const bool bIsActive = BTComp->IsAuxNodeActive(this);
-		ensureMsgf(bIsActive == false, TEXT("%hs called while the BT node is not active or has already been aborted"), __FUNCTION__);
-		UE_CVLOG_UELOG(bIsActive == false, BTComp, LogBehaviorTree, Warning, TEXT("%hs called for %s while it's not active or has already been aborted"), __FUNCTION__, *GetNodeName());
 		return;
 	}
 
