[BehaviorTree] Fixed TimeLimit decorator when used in composite decorators
Problems was that the decorator was returning a failure after requesting execution which could be the case in a composite decorator using operand `Or`
#11112

diff --git a/Engine/Source/Developer/AITestSuite/Private/Tests/BTTest.cpp b/Engine/Source/Developer/AITestSuite/Private/Tests/BTTest.cpp
--- a/Engine/Source/Developer/AITestSuite/Private/Tests/BTTest.cpp
+++ b/Engine/Source/Developer/AITestSuite/Private/Tests/BTTest.cpp
@@ -7,6 +7,7 @@
 #include "AITestsCommon.h"
 #include "MockAI_BT.h"
 #include "BehaviorTree/TestBTDecorator_CantExecute.h"
+#include "BehaviorTree/Decorators/BTDecorator_TimeLimit.h"
 
 #define LOCTEXT_NAMESPACE "AITestSuite_BTTest"
 
@@ -798,6 +799,69 @@ struct FAITest_BTCompositeFailedDecoratorUnregistersChildren : public FAITest_Si
 };
 IMPLEMENT_AI_LATENT_TEST(FAITest_BTCompositeFailedDecoratorUnregistersChildren, "System.AI.Behavior Trees.Composite failed decorator unregisters child nodes")
 
+struct FAITest_BTTimeLimitDecorator : FAITest_SimpleBT
+{
+	FAITest_BTTimeLimitDecorator()
+	{
+		enum
+		{
+			Task1Execute = 1,
+			Task1Tick,
+			Task2Execute,
+		};
+		
+		UBTCompositeNode& CompNode = FBTBuilder::AddSelector(*BTAsset);
+		{
+			FBTBuilder::AddTask(CompNode, Task1Execute, EBTNodeResult::Succeeded, /*ExecutionTicks*/ 5, Task1Tick);
+			{
+				constexpr float NumTicks = 2.5f;
+				FBTBuilder::WithDecorator<UBTDecorator_TimeLimit>(CompNode).TimeLimit = NumTicks * FAITestHelpers::TickInterval;
+			}
+
+			FBTBuilder::AddTask(CompNode, Task2Execute, EBTNodeResult::Succeeded);
+		}
+
+		ExpectedResult.Add(Task1Execute);
+		ExpectedResult.Add(Task1Tick);
+		ExpectedResult.Add(Task1Tick);
+		ExpectedResult.Add(Task2Execute);
+	}
+};
+IMPLEMENT_AI_LATENT_TEST(FAITest_BTTimeLimitDecorator, "System.AI.Behavior Trees.TimeLimit decorator")
+
+struct FAITest_BTTimeLimitInCompositeDecorator : FAITest_SimpleBT
+{
+	FAITest_BTTimeLimitInCompositeDecorator()
+	{
+		enum
+		{
+			Task1Execute = 1,
+			Task1Tick,
+			Task2Execute,
+		};
+		
+		UBTCompositeNode& CompNode = FBTBuilder::AddSelector(*BTAsset);
+		{
+			FBTBuilder::AddTask(CompNode, Task1Execute, EBTNodeResult::Succeeded, /*ExecutionTicks*/ 5, Task1Tick);
+			{
+				constexpr float NumTicks = 2.5f;
+				FBTBuilder::WithDecorator<UBTDecorator_TimeLimit>(CompNode).TimeLimit = NumTicks * FAITestHelpers::TickInterval;
+
+				TArray<FBTDecoratorLogic>& CompositeOps = CompNode.Children.Last().DecoratorOps;
+				CompositeOps.Add(FBTDecoratorLogic(EBTDecoratorLogic::Test, 0));
+			}
+
+			FBTBuilder::AddTask(CompNode, Task2Execute, EBTNodeResult::Succeeded);
+		}
+
+		ExpectedResult.Add(Task1Execute);
+		ExpectedResult.Add(Task1Tick);
+		ExpectedResult.Add(Task1Tick);
+		ExpectedResult.Add(Task2Execute);
+	}
+};
+IMPLEMENT_AI_LATENT_TEST(FAITest_BTTimeLimitInCompositeDecorator, "System.AI.Behavior Trees.TimeLimit in Composite decorator")
+
 /* All BTAbortingDuringService/BTAbortingDuringTaskService come from UDN case 00317509*/
 struct FAITest_BTAbortingDuringServiceTick : public FAITest_SimpleBT
 {
diff --git a/Engine/Source/Runtime/AIModule/Classes/BehaviorTree/Decorators/BTDecorator_TimeLimit.h b/Engine/Source/Runtime/AIModule/Classes/BehaviorTree/Decorators/BTDecorator_TimeLimit.h
--- a/Engine/Source/Runtime/AIModule/Classes/BehaviorTree/Decorators/BTDecorator_TimeLimit.h
+++ b/Engine/Source/Runtime/AIModule/Classes/BehaviorTree/Decorators/BTDecorator_TimeLimit.h
@@ -27,6 +27,11 @@ class UBTDecorator_TimeLimit : public UBTDecorator
 #endif // WITH_EDITOR
 
 protected:
-	AIMODULE_API virtual void OnNodeActivation(FBehaviorTreeSearchData& SearchData) override;
-	AIMODULE_API virtual void TickNode(UBehaviorTreeComponent& OwnerComp, uint8* NodeMemory, float DeltaSeconds) override;
+	virtual void InitializeMemory(UBehaviorTreeComponent& OwnerComp, uint8* NodeMemory, EBTMemoryInit::Type InitType) const override;
+	virtual void CleanupMemory(UBehaviorTreeComponent& OwnerComp, uint8* NodeMemory, EBTMemoryClear::Type CleanupType) const override;
+	virtual uint16 GetInstanceMemorySize() const override;
+	virtual void OnBecomeRelevant(UBehaviorTreeComponent& OwnerComp, uint8* NodeMemory) override;
+	virtual void OnCeaseRelevant(UBehaviorTreeComponent& OwnerComp, uint8* NodeMemory) override;
+	virtual bool CalculateRawConditionValue(UBehaviorTreeComponent& OwnerComp, uint8* NodeMemory) const override;
+	virtual void TickNode(UBehaviorTreeComponent& OwnerComp, uint8* NodeMemory, float DeltaSeconds) override;
 };
diff --git a/Engine/Source/Runtime/AIModule/Private/BehaviorTree/Decorators/BTDecorator_TimeLimit.cpp b/Engine/Source/Runtime/AIModule/Private/BehaviorTree/Decorators/BTDecorator_TimeLimit.cpp
--- a/Engine/Source/Runtime/AIModule/Private/BehaviorTree/Decorators/BTDecorator_TimeLimit.cpp
+++ b/Engine/Source/Runtime/AIModule/Private/BehaviorTree/Decorators/BTDecorator_TimeLimit.cpp
@@ -4,6 +4,11 @@
 
 #include UE_INLINE_GENERATED_CPP_BY_NAME(BTDecorator_TimeLimit)
 
+struct FBTimeLimitMemory
+{
+	bool bElapsed = false;
+};
+
 UBTDecorator_TimeLimit::UBTDecorator_TimeLimit(const FObjectInitializer& ObjectInitializer) : Super(ObjectInitializer)
 {
 	NodeName = "TimeLimit";
@@ -17,19 +22,45 @@ UBTDecorator_TimeLimit::UBTDecorator_TimeLimit(const FObjectInitializer& ObjectI
 	FlowAbortMode = EBTFlowAbortMode::Self;
 }
 
-void UBTDecorator_TimeLimit::OnNodeActivation(FBehaviorTreeSearchData& SearchData)
+uint16 UBTDecorator_TimeLimit::GetInstanceMemorySize() const
 {
-	uint8* RawMemory = SearchData.OwnerComp.GetNodeMemory(this, SearchData.OwnerComp.FindInstanceContainingNode(this));
-	if (RawMemory)
-	{
-		FBTAuxiliaryMemory* DecoratorMemory = GetSpecialNodeMemory<FBTAuxiliaryMemory>(RawMemory);
-		DecoratorMemory->NextTickRemainingTime = TimeLimit;
-		DecoratorMemory->AccumulatedDeltaTime = 0.0f;
-	}
+	return sizeof(FBTimeLimitMemory);
+}
+
+void UBTDecorator_TimeLimit::InitializeMemory(UBehaviorTreeComponent& OwnerComp, uint8* NodeMemory, const EBTMemoryInit::Type InitType) const
+{
+	InitializeNodeMemory<FBTimeLimitMemory>(NodeMemory, InitType);
+}
+
+void UBTDecorator_TimeLimit::CleanupMemory(UBehaviorTreeComponent& OwnerComp, uint8* NodeMemory, const EBTMemoryClear::Type CleanupType) const
+{
+	CleanupNodeMemory<FBTimeLimitMemory>(NodeMemory, CleanupType);
+}
+
+void UBTDecorator_TimeLimit::OnBecomeRelevant(UBehaviorTreeComponent& OwnerComp, uint8* NodeMemory)
+{
+	Super::OnBecomeRelevant(OwnerComp, NodeMemory);
+
+	SetNextTickTime(NodeMemory, TimeLimit);
 }
 
-void UBTDecorator_TimeLimit::TickNode(UBehaviorTreeComponent& OwnerComp, uint8* NodeMemory, float DeltaSeconds)
+void UBTDecorator_TimeLimit::OnCeaseRelevant(UBehaviorTreeComponent& OwnerComp, uint8* NodeMemory)
 {
+	Super::OnCeaseRelevant(OwnerComp, NodeMemory);
+
+	reinterpret_cast<FBTimeLimitMemory*>(NodeMemory)->bElapsed = false;
+}
+
+bool UBTDecorator_TimeLimit::CalculateRawConditionValue(UBehaviorTreeComponent& OwnerComp, uint8* NodeMemory) const
+{
+	return !reinterpret_cast<FBTimeLimitMemory*>(NodeMemory)->bElapsed;
+}
+
+void UBTDecorator_TimeLimit::TickNode(UBehaviorTreeComponent& OwnerComp, uint8* NodeMemory, const float DeltaSeconds)
+{
+	ensureMsgf(DeltaSeconds >= TimeLimit, TEXT("Using SetNextTickTime in OnBecomeRelevant should guarantee that we are only getting ticked when the time limit is finished"));
+	reinterpret_cast<FBTimeLimitMemory*>(NodeMemory)->bElapsed = true;
+
 	OwnerComp.RequestExecution(this);
 }
 
@@ -47,7 +78,7 @@ void UBTDecorator_TimeLimit::DescribeRuntimeValues(const UBehaviorTreeComponent&
 	const FBTAuxiliaryMemory* DecoratorMemory = GetSpecialNodeMemory<FBTAuxiliaryMemory>(NodeMemory);
 	if (OwnerComp.GetWorld())
 	{
-		const float TimeLeft = DecoratorMemory->NextTickRemainingTime > 0 ? DecoratorMemory->NextTickRemainingTime : 0;
+		const float TimeLeft = DecoratorMemory->NextTickRemainingTime > 0.f ? DecoratorMemory->NextTickRemainingTime : 0.f;
 		Values.Add(FString::Printf(TEXT("%s in %ss"),
 			*UBehaviorTreeTypes::DescribeNodeResult(EBTNodeResult::Failed),
 			*FString::SanitizeFloat(TimeLeft)));
