[BehaviorTree] fixed delta time for tasks using tick intervals
Problem was that a new task to execute using tick intervals could be ticked using the accumulated delta time from the previous task.

diff --git a/Engine/Source/Runtime/AIModule/Classes/BehaviorTree/BehaviorTreeComponent.h b/Engine/Source/Runtime/AIModule/Classes/BehaviorTree/BehaviorTreeComponent.h
--- a/Engine/Source/Runtime/AIModule/Classes/BehaviorTree/BehaviorTreeComponent.h
+++ b/Engine/Source/Runtime/AIModule/Classes/BehaviorTree/BehaviorTreeComponent.h
@@ -246,7 +246,16 @@ public:
 	AIMODULE_API bool IsAuxNodeActive(const UBTAuxiliaryNode* AuxNodeTemplate, int32 InstanceIdx) const;
 
 	/** Returns true if InstanceStack contains any BT runtime instances */
-	bool IsInstanceStackEmpty() const { return (InstanceStack.Num() == 0); }
+	bool IsInstanceStackEmpty() const
+	{
+		return (InstanceStack.Num() == 0);
+	}
+
+	/** Returns the accumulated delta time for the current tick */
+	float GetAccumulatedTickDeltaTime() const
+	{
+		return AccumulatedTickDeltaTime;
+	}
 
 	/** @return status of speficied task */
 	AIMODULE_API EBTTaskStatus::Type GetTaskStatus(const UBTTaskNode* TaskNode) const;
@@ -531,12 +540,18 @@ protected:
 
 	/** Used to tell tickmanager that we want interval ticking */
 	bool bTickedOnce = false;
+
 	/** Predicted next DeltaTime*/
-	float NextTickDeltaTime = 0.;
+	float NextTickDeltaTime = 0.f;
+
 	/** Accumulated DeltaTime if ticked more than predicted next delta time */
-	float AccumulatedTickDeltaTime = 0.0f;
+	float AccumulatedTickDeltaTime = 0.f;
+
+	/** Current frame DeltaTime */
+	float CurrentFrameDeltaTime = 0.f;
+
 	/** GameTime of the last DeltaTime request, used for debugging to output warnings about ticking */
-	double LastRequestedDeltaTimeGameTime = 0.;
+	double LastRequestedDeltaTimeGameTime = 0;
 
 #if CSV_PROFILER
 	/** CSV tick stat name. Can be changed but must point to a static string */
diff --git a/Engine/Source/Runtime/AIModule/Private/BehaviorTree/BTAuxiliaryNode.cpp b/Engine/Source/Runtime/AIModule/Private/BehaviorTree/BTAuxiliaryNode.cpp
--- a/Engine/Source/Runtime/AIModule/Private/BehaviorTree/BTAuxiliaryNode.cpp
+++ b/Engine/Source/Runtime/AIModule/Private/BehaviorTree/BTAuxiliaryNode.cpp
@@ -52,7 +52,7 @@ bool UBTAuxiliaryNode::WrappedTickNode(UBehaviorTreeComponent& OwnerComp, uint8*
 		if (NodeOb != nullptr && NodeOb->bNotifyTick)
 		{
 			float UseDeltaTime = DeltaSeconds;
-
+ 
 			if (NodeOb->bTickIntervals)
 			{
 				FBTAuxiliaryMemory* AuxMemory = GetSpecialNodeMemory<FBTAuxiliaryMemory>(NodeMemory);
@@ -63,10 +63,11 @@ bool UBTAuxiliaryNode::WrappedTickNode(UBehaviorTreeComponent& OwnerComp, uint8*
 				if (bTick)
 				{
 				    UseDeltaTime = AuxMemory->AccumulatedDeltaTime;
-				    AuxMemory->AccumulatedDeltaTime = 0.0f;
-    
-					UE_VLOG(OwnerComp.GetOwner(), LogBehaviorTree, Verbose, TEXT("Ticking aux node: %s"), *UBehaviorTreeTypes::DescribeNodeHelper(this));
+					AuxMemory->AccumulatedDeltaTime = 0.0f;
+					// Clamp next tick remaining time to 0 (i.e. tick on next frame) but node will usually set a new tick time in the call to TickNode
+					AuxMemory->NextTickRemainingTime = 0.0f;
 
+					UE_VLOG(OwnerComp.GetOwner(), LogBehaviorTree, Verbose, TEXT("Ticking aux node: %s"), *UBehaviorTreeTypes::DescribeNodeHelper(this));
 				    const_cast<UBTAuxiliaryNode*>(NodeOb)->TickNode(OwnerComp, NodeMemory, UseDeltaTime);
 				}
 
diff --git a/Engine/Source/Runtime/AIModule/Private/BehaviorTree/BTTaskNode.cpp b/Engine/Source/Runtime/AIModule/Private/BehaviorTree/BTTaskNode.cpp
--- a/Engine/Source/Runtime/AIModule/Private/BehaviorTree/BTTaskNode.cpp
+++ b/Engine/Source/Runtime/AIModule/Private/BehaviorTree/BTTaskNode.cpp
@@ -36,8 +36,29 @@ void UBTTaskNode::SetNextTickTime(uint8* NodeMemory, float RemainingTime) const
 
 EBTNodeResult::Type UBTTaskNode::WrappedExecuteTask(UBehaviorTreeComponent& OwnerComp, uint8* NodeMemory) const
 {
-	const UBTNode* NodeOb = bCreateNodeInstance ? GetNodeInstance(OwnerComp, NodeMemory) : this;
-	return NodeOb ? ((UBTTaskNode*)NodeOb)->ExecuteTask(OwnerComp, NodeMemory) : EBTNodeResult::Failed;
+	EBTNodeResult::Type Result = EBTNodeResult::Failed;
+
+	if (const UBTNode* NodeOb = bCreateNodeInstance ? GetNodeInstance(OwnerComp, NodeMemory) : this)
+	{
+		Result = ((UBTTaskNode*)NodeOb)->ExecuteTask(OwnerComp, NodeMemory);
+
+		// Now that the task was executed we need to adjust times when using tick intervals
+		// since `WrappedTickTask` uses `DeltaSeconds` which is accumulated time + current frame DT.
+		// We don't want to pass previously accumulated time to the new task to execute.
+		if (bTickIntervals)
+		{
+			const float AccumulatedDeltaTime = OwnerComp.GetAccumulatedTickDeltaTime();
+			FBTTaskMemory* TaskMemory = GetSpecialNodeMemory<FBTTaskMemory>(NodeMemory);
+
+			// Add accumulated time to `NextTickRemainingTime` set by the task to compensate for the next decrement in `WrappedTickTask`
+			TaskMemory->NextTickRemainingTime += AccumulatedDeltaTime;
+
+			// Assign task accumulated DT to negative current accumulated time/ for the next increment in `WrappedTickTask`
+			TaskMemory->AccumulatedDeltaTime = -AccumulatedDeltaTime;
+		}
+	}
+
+	return Result;
 }
 
 EBTNodeResult::Type UBTTaskNode::WrappedAbortTask(UBehaviorTreeComponent& OwnerComp, uint8* NodeMemory) const
diff --git a/Engine/Source/Runtime/AIModule/Private/BehaviorTree/BehaviorTreeComponent.cpp b/Engine/Source/Runtime/AIModule/Private/BehaviorTree/BehaviorTreeComponent.cpp
--- a/Engine/Source/Runtime/AIModule/Private/BehaviorTree/BehaviorTreeComponent.cpp
+++ b/Engine/Source/Runtime/AIModule/Private/BehaviorTree/BehaviorTreeComponent.cpp
@@ -1594,9 +1594,6 @@ void UBehaviorTreeComponent::ApplyAllSearchUpdates(const TArray<FBehaviorTreeSea
 
 void UBehaviorTreeComponent::TickNewlyAddedAuxNodesHelper()
 {
-	UWorld* MyWorld = GetWorld();
-	const float CurrentFrameDeltaSeconds = MyWorld ? MyWorld->GetDeltaSeconds() : 0.0f;
-
 	for (const FBehaviorTreeSearchUpdate& UpdateInfo : SearchData.PendingUpdates)
 	{
 		if (!UpdateInfo.bApplySkipped && UpdateInfo.Mode == EBTNodeUpdateMode::Add && UpdateInfo.AuxNode && InstanceStack.IsValidIndex(UpdateInfo.InstanceIndex))
@@ -1606,7 +1603,7 @@ void UBehaviorTreeComponent::TickNewlyAddedAuxNodesHelper()
 
 			// We do not care about the next needed DeltaTime, it will be recalculated in the tick later.
 			float NextNeededDeltaTime = 0.0f;
-			UpdateInfo.AuxNode->WrappedTickNode(*this, NodeMemory, CurrentFrameDeltaSeconds, NextNeededDeltaTime);
+			UpdateInfo.AuxNode->WrappedTickNode(*this, NodeMemory, CurrentFrameDeltaTime, NextNeededDeltaTime);
 		}
 	}
 }
@@ -1674,20 +1671,22 @@ void UBehaviorTreeComponent::ApplyDiscardedSearch()
 	SearchData.PendingNotifies.Reset();
 }
 
-void UBehaviorTreeComponent::TickComponent(float DeltaTime, enum ELevelTick TickType, FActorComponentTickFunction *ThisTickFunction)
+void UBehaviorTreeComponent::TickComponent(float DeltaTime, const ELevelTick TickType, FActorComponentTickFunction *ThisTickFunction)
 {
+	// Tick can be optimized by the tick function to not be called every frame so we need
+	// to set the current frame delta time based on that information for other tick scenarios (e.g. manual ticking in unit tests)
+	const UWorld* World = GetWorld();
+	const bool bUseWorldDTForCurrentFrame = (ThisTickFunction != nullptr && World != nullptr);
+	
+	CurrentFrameDeltaTime = bUseWorldDTForCurrentFrame ? World->GetDeltaSeconds() : DeltaTime;
+
 	// Warn if BT asked to be ticked the next frame and did not.
-	if (bTickedOnce && NextTickDeltaTime == 0.0f)
+	if (bTickedOnce && NextTickDeltaTime == 0.0f && bUseWorldDTForCurrentFrame)
 	{
-		UWorld* MyWorld = GetWorld();
-		if (MyWorld)
+		const double CurrentGameTime = World->GetTimeSeconds();
+		if (CurrentGameTime - LastRequestedDeltaTimeGameTime - CurrentFrameDeltaTime > KINDA_SMALL_NUMBER)
 		{
-			const double CurrentGameTime = MyWorld->GetTimeSeconds();
-			const float CurrentDeltaTime = MyWorld->GetDeltaSeconds();
-			if (CurrentGameTime - LastRequestedDeltaTimeGameTime - CurrentDeltaTime > KINDA_SMALL_NUMBER)
-			{
-				UE_VLOG(GetOwner(), LogBehaviorTree, Error, TEXT("BT(%i) expected to be tick next frame, current deltatime(%f) and calculated deltatime(%f)."), GFrameCounter, CurrentDeltaTime, CurrentGameTime - LastRequestedDeltaTimeGameTime);
-			}
+			UE_VLOG(GetOwner(), LogBehaviorTree, Error, TEXT("BT(%llu) expected to be tick next frame, current deltatime(%f) and calculated deltatime(%f)."), GFrameCounter, CurrentFrameDeltaTime, CurrentGameTime - LastRequestedDeltaTimeGameTime);
 		}
 	}
 
@@ -1705,8 +1704,14 @@ void UBehaviorTreeComponent::TickComponent(float DeltaTime, enum ELevelTick Tick
 		ScheduleNextTick(NextTickDeltaTime);
 		return;
 	}
-	DeltaTime += AccumulatedTickDeltaTime;
-	AccumulatedTickDeltaTime = 0.0f;
+
+	AccumulatedTickDeltaTime += DeltaTime;
+	ON_SCOPE_EXIT
+	{
+		AccumulatedTickDeltaTime = 0.0f;
+	};
+
+	DeltaTime = AccumulatedTickDeltaTime;
 
 	const bool bWasTickedOnce = bTickedOnce;
 	bTickedOnce = true;
@@ -2440,15 +2445,13 @@ void UBehaviorTreeComponent::ExecuteTask(UBTTaskNode* TaskNode)
 	}
 
 	// Services were already ticked for this frame, need to tick the new ones. 
-	UWorld* MyWorld = GetWorld();
-	const float CurrentFrameDeltaSeconds = MyWorld ? MyWorld->GetDeltaSeconds() : 0.0f;
 	for (UBTService* ServiceNode : TaskNode->Services)
 	{
 		uint8* NodeMemory = (uint8*)ServiceNode->GetNodeMemory<uint8>(ActiveInstance);
 
 		// We do not care about the next needed DeltaTime, it will be recalculated in the tick later.
 		float NextNeededDeltaTime = 0.0f;
-		ServiceNode->WrappedTickNode(*this, NodeMemory, CurrentFrameDeltaSeconds, NextNeededDeltaTime);
+		ServiceNode->WrappedTickNode(*this, NodeMemory, CurrentFrameDeltaTime, NextNeededDeltaTime);
 	}
 
 	ActiveInstance.ActiveNode = TaskNode;
