*Added an optional ReachCheck on NavLinkCustomInterface to allow navlinks to override the basic reach test done by the PathFollowingComponent.
*Improved the path ReachCheck for existing custom links.

#tests NPC navigation in PIE

diff --git a/Engine/Source/Runtime/AIModule/Classes/Navigation/PathFollowingComponent.h b/Engine/Source/Runtime/AIModule/Classes/Navigation/PathFollowingComponent.h
--- a/Engine/Source/Runtime/AIModule/Classes/Navigation/PathFollowingComponent.h
+++ b/Engine/Source/Runtime/AIModule/Classes/Navigation/PathFollowingComponent.h
@@ -430,6 +430,9 @@ protected:
 	/** currently traversed custom nav link */
 	FWeakObjectPtr CurrentCustomLinkOb;
 
+	/** the custom link for the next segment if there is one */
+	FWeakObjectPtr MoveSegmentCustomLinkOb;
+
 	/** navigation data for agent described in movement component */
 	UPROPERTY(transient)
 	TObjectPtr<ANavigationData> MyNavData;
@@ -527,6 +530,9 @@ protected:
 	/** True if pathfollowing is doing deceleration at the end of the path. @see FollowPathSegment(). */
 	uint8 bIsDecelerating : 1;
 
+	/** True if the next segment is a custom link that has its own reach conditions. */
+	uint8 bMoveSegmentIsUsingCustomLinkReachCondition : 1;
+
 	/** detect blocked movement when distance between center of location samples and furthest one (centroid radius) is below threshold */
 	float BlockDetectionDistance;
 
@@ -606,6 +612,9 @@ protected:
 	/** check if moving agent has reached goal defined by cylinder */
 	AIMODULE_API bool HasReachedInternal(const FVector& GoalLocation, float GoalRadius, float GoalHalfHeight, const FVector& AgentLocation, float RadiusThreshold, float AgentRadiusMultiplier) const;
 
+	/** reset the cached information about CustomLinks on the next MoveSegment */
+	AIMODULE_API void ResetMoveSegmentCustomLinkCache();
+
 	/** check if agent is on path */
 	AIMODULE_API virtual bool IsOnPath() const;
 
diff --git a/Engine/Source/Runtime/AIModule/Private/Navigation/PathFollowingComponent.cpp b/Engine/Source/Runtime/AIModule/Private/Navigation/PathFollowingComponent.cpp
--- a/Engine/Source/Runtime/AIModule/Private/Navigation/PathFollowingComponent.cpp
+++ b/Engine/Source/Runtime/AIModule/Private/Navigation/PathFollowingComponent.cpp
@@ -144,6 +144,7 @@ UPathFollowingComponent::UPathFollowingComponent(const FObjectInitializer& Objec
 	MoveSegmentEndIndex = 1;
 	MoveSegmentStartRef = INVALID_NAVNODEREF;
 	MoveSegmentEndRef = INVALID_NAVNODEREF;
+	bMoveSegmentIsUsingCustomLinkReachCondition = false;
 
 	CachedBrakingDistance = 100.0f;
 	CachedBrakingMaxSpeed = 0.0f;
@@ -737,6 +738,7 @@ void UPathFollowingComponent::Reset()
 	MoveSegmentStartRef = INVALID_NAVNODEREF;
 	MoveSegmentEndRef = INVALID_NAVNODEREF;
 	DecelerationSegmentIndex = INDEX_NONE;
+	ResetMoveSegmentCustomLinkCache();
 
 	LocationSamples.Reset();
 	LastSampleTime = 0.0f;
@@ -858,6 +860,26 @@ void UPathFollowingComponent::SetMoveSegment(int32 SegmentStartIndex)
 			SegmentEnd = *CurrentDestination;
 		}
 
+		// Check if the next segment is a custom link and has the need to use its own custom reach conditions
+		if (PathPt1.CustomNavLinkId != FNavLinkId::Invalid)
+		{
+			const UNavigationSystemV1* NavSys = FNavigationSystem::GetCurrent<UNavigationSystemV1>(GetWorld());
+			if (const INavLinkCustomInterface* MoveSegmentCustomLink = NavSys->GetCustomLink(PathPt1.CustomNavLinkId))
+			{
+				// Cache the CustomLinkOb for faster access during the update
+				MoveSegmentCustomLinkOb = Cast<const UObject>(MoveSegmentCustomLink);
+				bMoveSegmentIsUsingCustomLinkReachCondition = MoveSegmentCustomLink->IsLinkUsingCustomReachCondition(this);
+			}
+			else
+			{
+				ResetMoveSegmentCustomLinkCache();
+			}
+		}
+		else
+		{
+			ResetMoveSegmentCustomLinkCache();
+		}
+
 		CurrentAcceptanceRadius = (PathInstance->GetPathPoints().Num() == (MoveSegmentEndIndex + 1))
 			? GetFinalAcceptanceRadius(*PathInstance, OriginalMoveRequestGoalLocation)
 			// pick appropriate value base on whether we're going to nav link or not
@@ -1165,6 +1187,27 @@ bool UPathFollowingComponent::HasReachedCurrentTarget(const FVector& CurrentLoca
 		return false;
 	}
 
+	// If the next segment is a link with a custom reach condition, we need to call the HasReachedLinkStart on the link interface.
+	if (bMoveSegmentIsUsingCustomLinkReachCondition)
+	{
+		if (const INavLinkCustomInterface* MoveSegmentCustomLink = Cast<const INavLinkCustomInterface>(MoveSegmentCustomLinkOb.Get()))
+		{
+			if (ensureMsgf(Path.IsValid(), TEXT("%hs: Path should be valid when we get here. Owner [%s]."), __FUNCTION__, *GetNameSafe(GetOwner())))
+			{
+				const FNavPathPoint& LinkStart = Path->GetPathPoints()[MoveSegmentEndIndex];
+				if (Path->GetPathPoints().IsValidIndex(MoveSegmentEndIndex + 1))
+				{
+					const FNavPathPoint& LinkEnd = Path->GetPathPoints()[MoveSegmentEndIndex + 1];
+					return MoveSegmentCustomLink->HasReachedLinkStart(this, CurrentLocation, LinkStart, LinkEnd);
+				}
+				else
+				{
+					UE_LOG(LogPathFollowing, Error, TEXT("%hs: NavLink has a start, but no end. Custom reach condition won't be called. NavLinkID [%llu] - LinkStartPos [%s] - Owner [%s]"), __FUNCTION__, LinkStart.CustomNavLinkId.GetId(), *LinkStart.Location.ToString(), *GetNameSafe(GetOwner()));
+				}
+			}
+		}
+	}
+
 	const FVector CurrentTarget = GetCurrentTargetLocation();
 	const FVector CurrentDirection = GetCurrentDirection();
 
@@ -1404,6 +1447,12 @@ void UPathFollowingComponent::OnActorBump(AActor* SelfActor, AActor* OtherActor,
 	}
 }
 
+void UPathFollowingComponent::ResetMoveSegmentCustomLinkCache()
+{
+	MoveSegmentCustomLinkOb.Reset();
+	bMoveSegmentIsUsingCustomLinkReachCondition = false;
+}
+
 bool UPathFollowingComponent::IsOnPath() const
 {
 	bool bOnPath = false;
diff --git a/Engine/Source/Runtime/NavigationSystem/Public/NavLinkCustomInterface.h b/Engine/Source/Runtime/NavigationSystem/Public/NavLinkCustomInterface.h
--- a/Engine/Source/Runtime/NavigationSystem/Public/NavLinkCustomInterface.h
+++ b/Engine/Source/Runtime/NavigationSystem/Public/NavLinkCustomInterface.h
@@ -76,6 +76,14 @@ class INavLinkCustomInterface
 	/** Notify called when agent finishes using this link for movement */
 	virtual void OnLinkMoveFinished(class UObject* PathComp) {}
 
+	/** Whether or not this link has custom reach conditions that need to override the default reach checks done by the path following component. */
+	virtual bool IsLinkUsingCustomReachCondition(const UObject* PathComp) const { return false; }
+
+	/** Function that replaces the default reach check when IsLinkUsingCustomReachCondition is true. 
+	 *  Returns true if CurrentLocation has reached the start of the link.
+	 */
+	virtual bool HasReachedLinkStart(const UObject* PathComp, const FVector& CurrentLocation, const FNavPathPoint& LinkStart, const FNavPathPoint& LinkEnd) const { return true; }
+
 	UE_DEPRECATED(5.3, "LinkIds are now based on FNavLinkId using FNavLinkId::GenerateUniqueId(). This function will still generate an incremental Id however it does not work well in all circumstances.")
 	static NAVIGATIONSYSTEM_API uint32 GetUniqueId();
 
