Removing Actor.h from AITypes.h
(with misc other cleanup)

diff --git a/Engine/Source/Runtime/AIModule/Classes/AITypes.h b/Engine/Source/Runtime/AIModule/Classes/AITypes.h
--- a/Engine/Source/Runtime/AIModule/Classes/AITypes.h
+++ b/Engine/Source/Runtime/AIModule/Classes/AITypes.h
@@ -11,7 +11,9 @@
 #endif
 #include "AI/Navigation/NavigationTypes.h"
 #include "NavigationSystemTypes.h"
+#if UE_ENABLE_INCLUDE_ORDER_DEPRECATED_IN_5_4
 #include "GameFramework/Actor.h"
+#endif
 #include "AITypes.generated.h"
 
 class AActor;
@@ -527,7 +529,7 @@ struct FAIMoveRequest
 	FVector GetGoalLocation() const { return GoalLocation; }
 	/** retrieves request's requested destination location, GoalActor's location 
 	 *	or GoalLocation, depending on the request itself */
-	FVector GetDestination() const { return bMoveToActor ? (GoalActor.IsValid() ? GoalActor->GetActorLocation() : FAISystem::InvalidLocation) : GoalLocation; }
+	AIMODULE_API FVector GetDestination() const;
 
 	bool IsUsingPathfinding() const { return bUsePathfinding; }
 	bool IsUsingPartialPaths() const { return bAllowPartialPath; }
diff --git a/Engine/Source/Runtime/AIModule/Private/AIHelpers.cpp b/Engine/Source/Runtime/AIModule/Private/AIHelpers.cpp
--- a/Engine/Source/Runtime/AIModule/Private/AIHelpers.cpp
+++ b/Engine/Source/Runtime/AIModule/Private/AIHelpers.cpp
@@ -1,6 +1,7 @@
 // Copyright Epic Games, Inc. All Rights Reserved.
 
 #include "AIHelpers.h"
+#include "GameFramework/Actor.h"
 
 namespace FAISystem
 {
diff --git a/Engine/Source/Runtime/AIModule/Private/AITypes.cpp b/Engine/Source/Runtime/AIModule/Private/AITypes.cpp
--- a/Engine/Source/Runtime/AIModule/Private/AITypes.cpp
+++ b/Engine/Source/Runtime/AIModule/Private/AITypes.cpp
@@ -1,6 +1,7 @@
 // Copyright Epic Games, Inc. All Rights Reserved.
 
 #include "AITypes.h"
+#include "GameFramework/Actor.h"
 #include "UObject/Package.h"
 #include "Navigation/PathFollowingComponent.h"
 #include "NavFilters/NavigationQueryFilter.h"
@@ -207,6 +208,11 @@ void FAIMoveRequest::SetGoalLocation(const FVector& InGoalLocation)
 	}
 }
 
+FVector FAIMoveRequest::GetDestination() const
+{
+	return bMoveToActor ? (GoalActor.IsValid() ? GoalActor->GetActorLocation() : FAISystem::InvalidLocation) : GoalLocation;
+}
+
 bool FAIMoveRequest::UpdateGoalLocation(const FVector& NewLocation) const
 {
 	if (!bMoveToActor)
diff --git a/Engine/Source/Runtime/AIModule/Private/BehaviorTree/BTAuxiliaryNode.cpp b/Engine/Source/Runtime/AIModule/Private/BehaviorTree/BTAuxiliaryNode.cpp
--- a/Engine/Source/Runtime/AIModule/Private/BehaviorTree/BTAuxiliaryNode.cpp
+++ b/Engine/Source/Runtime/AIModule/Private/BehaviorTree/BTAuxiliaryNode.cpp
@@ -3,6 +3,7 @@
 #include "BehaviorTree/BTAuxiliaryNode.h"
 #include "BehaviorTree/BTCompositeNode.h"
 #include "VisualLogger/VisualLogger.h"
+#include "GameFramework/Actor.h"
 
 #include UE_INLINE_GENERATED_CPP_BY_NAME(BTAuxiliaryNode)
 
diff --git a/Engine/Source/Runtime/AIModule/Private/BehaviorTree/BTDecorator.cpp b/Engine/Source/Runtime/AIModule/Private/BehaviorTree/BTDecorator.cpp
--- a/Engine/Source/Runtime/AIModule/Private/BehaviorTree/BTDecorator.cpp
+++ b/Engine/Source/Runtime/AIModule/Private/BehaviorTree/BTDecorator.cpp
@@ -2,6 +2,7 @@
 
 #include "BehaviorTree/BTDecorator.h"
 #include "BehaviorTree/BTCompositeNode.h"
+#include "GameFramework/Actor.h"
 #include "VisualLogger/VisualLogger.h"
 
 #include UE_INLINE_GENERATED_CPP_BY_NAME(BTDecorator)
diff --git a/Engine/Source/Runtime/AIModule/Private/BehaviorTree/BTService.cpp b/Engine/Source/Runtime/AIModule/Private/BehaviorTree/BTService.cpp
--- a/Engine/Source/Runtime/AIModule/Private/BehaviorTree/BTService.cpp
+++ b/Engine/Source/Runtime/AIModule/Private/BehaviorTree/BTService.cpp
@@ -1,6 +1,7 @@
 // Copyright Epic Games, Inc. All Rights Reserved.
 
 #include "BehaviorTree/BTService.h"
+#include "GameFramework/Actor.h"
 #include "VisualLogger/VisualLogger.h"
 
 #include UE_INLINE_GENERATED_CPP_BY_NAME(BTService)
diff --git a/Engine/Source/Runtime/AIModule/Private/BehaviorTree/Tasks/BTTask_RunBehaviorDynamic.cpp b/Engine/Source/Runtime/AIModule/Private/BehaviorTree/Tasks/BTTask_RunBehaviorDynamic.cpp
--- a/Engine/Source/Runtime/AIModule/Private/BehaviorTree/Tasks/BTTask_RunBehaviorDynamic.cpp
+++ b/Engine/Source/Runtime/AIModule/Private/BehaviorTree/Tasks/BTTask_RunBehaviorDynamic.cpp
@@ -1,6 +1,7 @@
 // Copyright Epic Games, Inc. All Rights Reserved.
 
 #include "BehaviorTree/Tasks/BTTask_RunBehaviorDynamic.h"
+#include "GameFramework/Actor.h"
 #include "VisualLogger/VisualLogger.h"
 #include "BehaviorTree/BehaviorTree.h"
 #include "Internationalization/Regex.h"
diff --git a/Engine/Source/Runtime/Engine/Public/AI/NavigationModifier.h b/Engine/Source/Runtime/Engine/Public/AI/NavigationModifier.h
--- a/Engine/Source/Runtime/Engine/Public/AI/NavigationModifier.h
+++ b/Engine/Source/Runtime/Engine/Public/AI/NavigationModifier.h
@@ -2,7 +2,9 @@
 
 #pragma once 
 
+#if UE_ENABLE_INCLUDE_ORDER_DEPRECATED_IN_5_4
 #include "CollisionShape.h"
+#endif
 #include "Templates/SubclassOf.h"
 #include "UObject/WeakObjectPtrTemplates.h"
 #include "AI/Navigation/NavLinkDefinition.h"
@@ -266,6 +268,8 @@ protected:
 	TWeakObjectPtr<UClass> LinkDefinitionClassOb;
 };
 
+struct FCollisionShape;
+
 struct FCompositeNavModifier : public FNavigationModifier
 {
 	FCompositeNavModifier() 
diff --git a/Engine/Source/Runtime/NavigationSystem/Public/NavigationSystemTypes.h b/Engine/Source/Runtime/NavigationSystem/Public/NavigationSystemTypes.h
--- a/Engine/Source/Runtime/NavigationSystem/Public/NavigationSystemTypes.h
+++ b/Engine/Source/Runtime/NavigationSystem/Public/NavigationSystemTypes.h
@@ -112,6 +112,8 @@ struct FNavigationInvokerRaw
 	FNavigationInvokerRaw(const FVector& InLocation, float Min, float Max, const FNavAgentSelector& InSupportedAgents, ENavigationInvokerPriority InPriority);
 };
 
+class AActor;
+
 struct FNavigationInvoker
 {
 	/** The Invoker source should be either an Actor or an Object. Thus only 1 of those member should be set. We'll use IsExplicitlyNull to know which one to use */
