[GraphAStar] Removed shrinking option parameter from PopIndex in FOpenList since we never want to shrink

diff --git a/Engine/Source/Runtime/AIModule/Public/GraphAStar.h b/Engine/Source/Runtime/AIModule/Public/GraphAStar.h
--- a/Engine/Source/Runtime/AIModule/Public/GraphAStar.h
+++ b/Engine/Source/Runtime/AIModule/Public/GraphAStar.h
@@ -299,17 +299,20 @@ struct FGraphAStar
 			check(false); // We should never reach here.
 		}
 
-		int32 PopIndex(EAllowShrinking AllowShrinking = EAllowShrinking::Yes)
+		int32 PopIndex()
 		{
 			int32 SearchNodeIndex = INDEX_NONE;
+
+			// During A* we grow the array as needed and it does not make sense to shrink in the process.
 			Super::HeapPop(SearchNodeIndex, NodeSorter, EAllowShrinking::No);
 			NodePool[SearchNodeIndex].MarkNotOpened();
 			return SearchNodeIndex;
 		}
-		UE_ALLOWSHRINKING_BOOL_DEPRECATED("PopIndex")
+
+		UE_DEPRECATED(5.4, "PopIndex with a boolean bAllowShrinking has been deprecated - please use the version without parameter")
 		FORCEINLINE int32 PopIndex(bool bAllowShrinking)
 		{
-			return PopIndex(bAllowShrinking ? EAllowShrinking::Yes : EAllowShrinking::No);
+			return PopIndex();
 		}
 	};
 
