AI Path Following
 * When navigating towards an Actor, test reachability before updating the final destination to the Actor's location
 * Remove an AIController ensure uselessly triggering when an AIController tries to use its DefaultNavigationFilterClass

diff --git a/Engine/Source/Runtime/AIModule/Classes/Navigation/PathFollowingComponent.h b/Engine/Source/Runtime/AIModule/Classes/Navigation/PathFollowingComponent.h
--- a/Engine/Source/Runtime/AIModule/Classes/Navigation/PathFollowingComponent.h
+++ b/Engine/Source/Runtime/AIModule/Classes/Navigation/PathFollowingComponent.h
@@ -440,6 +440,9 @@ protected:
 	/** requested path */
 	FNavPathSharedPtr Path;
 
+	/** Navigation query filter of the current move request */
+	FSharedConstNavQueryFilter NavigationFilter;
+
 	/** value based on navigation agent's properties that's used for AcceptanceRadius when DefaultAcceptanceRadius is requested */
 	float MyDefaultAcceptanceRadius;
 
@@ -509,6 +512,11 @@ protected:
 	/** if set, target location will be constantly updated to match goal actor while following last segment of full path */
 	uint8 bMoveToGoalOnLastSegment : 1;
 
+	/** Whether to clamp the goal location to reachable navigation data when trying to track the goal actor (Only used when bMoveToGoalOnLastSegment is true)
+	 *  False: (default) while following the last segment, the path is allowed to adjust through obstacles and off navigation data without checks.
+	 *  True: the last segment's destination will be clamped at the furthest reachable location towards the goal actor. */
+	uint8 bMoveToGoalClampedToNavigation : 1;
+
 	/** if set, movement block detection will be used */
 	uint8 bUseBlockDetection : 1;
 
diff --git a/Engine/Source/Runtime/AIModule/Private/AIController.cpp b/Engine/Source/Runtime/AIModule/Private/AIController.cpp
--- a/Engine/Source/Runtime/AIModule/Private/AIController.cpp
+++ b/Engine/Source/Runtime/AIModule/Private/AIController.cpp
@@ -627,8 +627,6 @@ FPathFollowingRequestResult AAIController::MoveTo(const FAIMoveRequest& MoveRequ
 		return ResultData;
 	}
 
-	ensure(MoveRequest.GetNavigationFilter() || !DefaultNavigationFilterClass);
-
 	bool bCanRequestMove = true;
 	bool bAlreadyAtGoal = false;
 	
diff --git a/Engine/Source/Runtime/AIModule/Private/Navigation/PathFollowingComponent.cpp b/Engine/Source/Runtime/AIModule/Private/Navigation/PathFollowingComponent.cpp
--- a/Engine/Source/Runtime/AIModule/Private/Navigation/PathFollowingComponent.cpp
+++ b/Engine/Source/Runtime/AIModule/Private/Navigation/PathFollowingComponent.cpp
@@ -28,7 +28,7 @@
 
 DEFINE_LOG_CATEGORY(LogPathFollowing);
 
-namespace
+namespace UE::Navigation::Private
 {
 	FORCEINLINE FVector FindGoalLocation(const UPathFollowingComponent& Component, const AActor& GoalActor, const INavAgentInterface* GoalNavAgent, float& GoalRadius, float& GoalHalfHeight)
 	{
@@ -45,6 +45,29 @@ namespace
 			return GoalActor.GetActorLocation();
 		}
 	}
+
+	FSharedConstNavQueryFilter ExtractNavigationFilterForRequest(UObject* Owner, FNavPathSharedPtr Path, TObjectPtr<ANavigationData> MyNavData, const FAIMoveRequest& RequestData)
+	{
+		if (!Path->CastPath<FAbstractNavigationPath>() && Path->GetFilter() != nullptr)
+		{
+			return Path->GetFilter();
+		}
+
+		if (RequestData.GetNavigationFilter() != nullptr)
+		{
+			return UNavigationQueryFilter::GetQueryFilter(*MyNavData, Owner, RequestData.GetNavigationFilter());
+		}
+
+		if (const AAIController* AIOwner = Cast<AAIController>(Owner))
+		{
+			if (AIOwner->GetDefaultNavigationFilterClass() != nullptr)
+			{
+				return UNavigationQueryFilter::GetQueryFilter(*MyNavData, AIOwner->GetDefaultNavigationFilterClass());
+			}
+		}
+
+		return MyNavData->GetDefaultQueryFilter();
+	}
 }
 
 FPathFollowingResult::FPathFollowingResult(uint16 InFlags) : Flags(InFlags)
@@ -153,6 +176,7 @@ UPathFollowingComponent::UPathFollowingComponent(const FObjectInitializer& Objec
 	bReachTestIncludesAgentRadius = true;
 	bReachTestIncludesGoalRadius = true;
 	bMoveToGoalOnLastSegment = true;
+	bMoveToGoalClampedToNavigation = false;
 
 	Status = EPathFollowingStatus::Idle;
 
@@ -254,7 +278,7 @@ void UPathFollowingComponent::OnPathEvent(FNavigationPath* InPath, ENavPathEvent
 				else if (InPath->IsPartial() && InPath->GetGoalActor())
 				{
 					float IgnoreGoalRadius, IgnoreGoalHalfHeight;
-					OriginalMoveRequestGoalLocation = FindGoalLocation(*this, *InPath->GetGoalActor(), InPath->GetGoalActorAsNavAgent(), IgnoreGoalRadius, IgnoreGoalHalfHeight);
+					OriginalMoveRequestGoalLocation = UE::Navigation::Private::FindGoalLocation(*this, *InPath->GetGoalActor(), InPath->GetGoalActorAsNavAgent(), IgnoreGoalRadius, IgnoreGoalHalfHeight);
 				}
 			}
 			break;
@@ -288,7 +312,7 @@ bool UPathFollowingComponent::HandlePathUpdateEvent()
 	if (PathGoalActor)
 	{
 		float IgnoreGoalRadius, IgnoreGoalHalfHeight;
-		OriginalMoveRequestGoalLocation = FindGoalLocation(*this, *PathGoalActor, Path->GetGoalActorAsNavAgent(), IgnoreGoalRadius, IgnoreGoalHalfHeight);
+		OriginalMoveRequestGoalLocation = UE::Navigation::Private::FindGoalLocation(*this, *PathGoalActor, Path->GetGoalActorAsNavAgent(), IgnoreGoalRadius, IgnoreGoalHalfHeight);
 	}
 	if (FAISystem::IsValidLocation(OriginalMoveRequestGoalLocation))
 	{
@@ -391,6 +415,8 @@ FAIRequestID UPathFollowingComponent::RequestMove(const FAIMoveRequest& RequestD
 			MetaNavPath->Initialize(CurrentLocation);
 		}
 
+		NavigationFilter = UE::Navigation::Private::ExtractNavigationFilterForRequest(GetOwner(), Path, MyNavData, RequestData);
+
 		PathTimeWhenPaused = 0.;
 		OnPathUpdated();
 
@@ -981,8 +1007,16 @@ void UPathFollowingComponent::UpdatePathSegment()
 			if (DestinationActor.IsValid() && Path->IsPartial() == false)
 			{
 				const FVector AgentLocation = DestinationAgent ? DestinationAgent->GetNavAgentLocation() : DestinationActor->GetActorLocation();
-				// note that the condition below requires GoalLocation to be in world space.
-				const FVector GoalLocation = FQuatRotationTranslationMatrix(DestinationActor->GetActorQuat(), AgentLocation).TransformPosition(MoveOffset);
+				FVector GoalLocation = FQuatRotationTranslationMatrix(DestinationActor->GetActorQuat(), AgentLocation).TransformPosition(MoveOffset);
+
+				if (bMoveToGoalClampedToNavigation && NavigationFilter)
+				{
+					FVector HitLocation;
+					if (MyNavData->Raycast(CurrentLocation, GoalLocation, HitLocation, NavigationFilter))
+					{
+						GoalLocation = HitLocation;
+					}
+				}
 
 				CurrentDestination.Set(NULL, GoalLocation);
 
@@ -1159,19 +1193,26 @@ bool UPathFollowingComponent::HasReachedDestination(const FVector& CurrentLocati
 	float GoalHalfHeight = 0.0f;
 	
 	// take goal's current location, unless path is partial or last segment doesn't reach goal actor (used by tethered AI)
-	if (DestinationActor.IsValid() && !Path->IsPartial() && bMoveToGoalOnLastSegment)
+	if (!Path->IsPartial() && bMoveToGoalOnLastSegment)
 	{
-		if (DestinationAgent)
-		{
-			const AActor* OwnerActor = GetOwner();
-			FVector GoalOffset;
-			DestinationAgent->GetMoveGoalReachTest(OwnerActor, MoveOffset, GoalOffset, GoalRadius, GoalHalfHeight);
+		GoalLocation = *CurrentDestination;
 
-			GoalLocation = FQuatRotationTranslationMatrix(DestinationActor->GetActorQuat(), DestinationAgent->GetNavAgentLocation()).TransformPosition(GoalOffset);
-		}
-		else
+		// Testing IsNearlyZero is not optimal as (0,0,0) could be a valid world position. Tracking the initialization state inside FBasedPosition would be better.
+		if (GoalLocation.IsNearlyZero() && DestinationActor.IsValid() && !bMoveToGoalClampedToNavigation)
 		{
-			GoalLocation = DestinationActor->GetActorLocation();
+			// In case CurrentDestination is not set because we didn't update MoveSegment yet, let's use the goal Actor's location directly
+			if (DestinationAgent)
+			{
+				const AActor* OwnerActor = GetOwner();
+				FVector GoalOffset;
+				DestinationAgent->GetMoveGoalReachTest(OwnerActor, MoveOffset, GoalOffset, GoalRadius, GoalHalfHeight);
+
+				GoalLocation = FQuatRotationTranslationMatrix(DestinationActor->GetActorQuat(), DestinationAgent->GetNavAgentLocation()).TransformPosition(GoalOffset);
+			}
+			else
+			{
+				GoalLocation = DestinationActor->GetActorLocation();
+			}
 		}
 	}
 
