Build fix - Add missing include NavigationQueryFilter

diff --git a/Engine/Source/Runtime/AIModule/Private/Navigation/PathFollowingComponent.cpp b/Engine/Source/Runtime/AIModule/Private/Navigation/PathFollowingComponent.cpp
--- a/Engine/Source/Runtime/AIModule/Private/Navigation/PathFollowingComponent.cpp
+++ b/Engine/Source/Runtime/AIModule/Private/Navigation/PathFollowingComponent.cpp
@@ -1,22 +1,24 @@
 // Copyright Epic Games, Inc. All Rights Reserved.
 
 #include "Navigation/PathFollowingComponent.h"
-#include "UObject/Package.h"
-#include "TimerManager.h"
-#include "GameFramework/Pawn.h"
-#include "GameFramework/Controller.h"
-#include "NavigationSystem.h"
-#include "NavMesh/RecastNavMesh.h"
+
+#include "AbstractNavData.h"
+#include "AIConfig.h"
+#include "AIController.h"
 #include "AISystem.h"
 #include "BrainComponent.h"
 #include "Engine/Canvas.h"
-#include "AIController.h"
-#include "VisualLogger/VisualLoggerTypes.h"
-#include "VisualLogger/VisualLogger.h"
-#include "AbstractNavData.h"
-#include "NavLinkCustomInterface.h"
+#include "GameFramework/Controller.h"
+#include "GameFramework/Pawn.h"
+#include "NavFilters/NavigationQueryFilter.h"
 #include "Navigation/MetaNavMeshPath.h"
-#include "AIConfig.h"
+#include "NavigationSystem.h"
+#include "NavLinkCustomInterface.h"
+#include "NavMesh/RecastNavMesh.h"
+#include "TimerManager.h"
+#include "UObject/Package.h"
+#include "VisualLogger/VisualLogger.h"
+#include "VisualLogger/VisualLoggerTypes.h"
 
 #include UE_INLINE_GENERATED_CPP_BY_NAME(PathFollowingComponent)
 
