[BTService] updated description for case where `bNotifyTick` is not set

diff --git a/Engine/Source/Runtime/AIModule/Private/BehaviorTree/BTService.cpp b/Engine/Source/Runtime/AIModule/Private/BehaviorTree/BTService.cpp
--- a/Engine/Source/Runtime/AIModule/Private/BehaviorTree/BTService.cpp
+++ b/Engine/Source/Runtime/AIModule/Private/BehaviorTree/BTService.cpp
@@ -67,11 +67,20 @@ void UBTService::NotifyParentActivation(FBehaviorTreeSearchData& SearchData)
 
 FString UBTService::GetStaticTickIntervalDescription() const
 {
-	FString IntervalDesc = (RandomDeviation > 0.0f) ?
-		FString::Printf(TEXT("%.2fs..%.2fs"), FMath::Max(0.0f, Interval - RandomDeviation), (Interval + RandomDeviation)) :
-		FString::Printf(TEXT("%.2fs"), Interval);
+	if (bNotifyTick)
+	{
+		FString IntervalDesc(TEXT("frame"));
+		if (bTickIntervals)
+		{
+			IntervalDesc = (RandomDeviation > 0.0f)
+				? FString::Printf(TEXT("%.2fs..%.2fs"), FMath::Max(0.0f, Interval - RandomDeviation), (Interval + RandomDeviation))
+				: FString::Printf(TEXT("%.2fs"), Interval);
+		} 
+
+		return FString::Printf(TEXT("tick every %s"), *IntervalDesc);
+	}
 
-	return FString::Printf(TEXT("tick every %s"), *IntervalDesc);
+	return TEXT("never ticks");
 }
 
 FString UBTService::GetStaticServiceDescription() const
