Fixed TIsValidVariadicFunctionArg failing to catch disallowed types which were convertible to allowed types.
Fixed a bunch of misuses, including TEnumAsByte, FAIRequestID and FPlayInEditorID.

diff --git a/Engine/Source/Runtime/AIModule/Private/BehaviorTree/BehaviorTreeComponent.cpp b/Engine/Source/Runtime/AIModule/Private/BehaviorTree/BehaviorTreeComponent.cpp
--- a/Engine/Source/Runtime/AIModule/Private/BehaviorTree/BehaviorTreeComponent.cpp
+++ b/Engine/Source/Runtime/AIModule/Private/BehaviorTree/BehaviorTreeComponent.cpp
@@ -2552,7 +2552,7 @@ void UBehaviorTreeComponent::RegisterMessageObserver(const UBTTaskNode* TaskNode
 			);
 
 		UE_VLOG(GetOwner(), LogBehaviorTree, Log, TEXT("Message[%s:%d] observer added for %s"),
-			*MessageType.ToString(), RequestID, *UBehaviorTreeTypes::DescribeNodeHelper(TaskNode));
+			*MessageType.ToString(), RequestID.GetID(), *UBehaviorTreeTypes::DescribeNodeHelper(TaskNode));
 	}
 }
 
diff --git a/Engine/Source/Runtime/AIModule/Private/Navigation/GridPathFollowingComponent.cpp b/Engine/Source/Runtime/AIModule/Private/Navigation/GridPathFollowingComponent.cpp
--- a/Engine/Source/Runtime/AIModule/Private/Navigation/GridPathFollowingComponent.cpp
+++ b/Engine/Source/Runtime/AIModule/Private/Navigation/GridPathFollowingComponent.cpp
@@ -178,7 +178,7 @@ void UGridPathFollowingComponent::ResumeMove(FAIRequestID RequestID)
 
 		if (HasActiveGrid())
 		{
-			UE_VLOG(GetOwner(), LogPathFollowing, Log, TEXT("ResumeMove: RequestID(%u) is on grid"), RequestID);
+			UE_VLOG(GetOwner(), LogPathFollowing, Log, TEXT("ResumeMove: RequestID(%u) is on grid"), RequestID.GetID());
 			Status = EPathFollowingStatus::Moving;
 		}
 		else
diff --git a/Engine/Source/Runtime/AIModule/Private/Navigation/PathFollowingComponent.cpp b/Engine/Source/Runtime/AIModule/Private/Navigation/PathFollowingComponent.cpp
--- a/Engine/Source/Runtime/AIModule/Private/Navigation/PathFollowingComponent.cpp
+++ b/Engine/Source/Runtime/AIModule/Private/Navigation/PathFollowingComponent.cpp
@@ -438,7 +438,7 @@ FAIRequestID UPathFollowingComponent::RequestMove(const FAIMoveRequest& RequestD
 			const FVector CurrentLocation = MovementComp ? MovementComp->GetActorFeetLocation() : FVector::ZeroVector;
 			const FVector DestLocation = InPath->GetDestinationLocation();
 			const FVector ToDest = DestLocation - CurrentLocation;
-			UE_VLOG(GetOwner(), LogPathFollowing, Log, TEXT("RequestMove: accepted, ID(%u) dist2D(%.0f) distZ(%.0f)"), MoveId, ToDest.Size2D(), FMath::Abs(ToDest.Z));
+			UE_VLOG(GetOwner(), LogPathFollowing, Log, TEXT("RequestMove: accepted, ID(%u) dist2D(%.0f) distZ(%.0f)"), MoveId.GetID(), ToDest.Size2D(), FMath::Abs(ToDest.Z));
 #endif // ENABLE_VISUAL_LOG
 
 			// with async pathfinding paths can be incomplete, movement will start after receiving path event 
@@ -463,7 +463,7 @@ FAIRequestID UPathFollowingComponent::RequestMove(const FAIMoveRequest& RequestD
 
 void UPathFollowingComponent::AbortMove(const UObject& Instigator, FPathFollowingResultFlags::Type AbortFlags, FAIRequestID RequestID, EPathFollowingVelocityMode VelocityMode)
 {
-	UE_VLOG(GetOwner(), LogPathFollowing, Log, TEXT("AbortMove: RequestID(%u) Instigator(%s)"), RequestID, *Instigator.GetName());
+	UE_VLOG(GetOwner(), LogPathFollowing, Log, TEXT("AbortMove: RequestID(%u) Instigator(%s)"), RequestID.GetID(), *Instigator.GetName());
 
 	if ((Status != EPathFollowingStatus::Idle) && RequestID.IsEquivalent(GetCurrentRequestId()))
 	{
@@ -489,7 +489,7 @@ FAIRequestID UPathFollowingComponent::RequestMoveWithImmediateFinish(EPathFollow
 
 void UPathFollowingComponent::PauseMove(FAIRequestID RequestID, EPathFollowingVelocityMode VelocityMode)
 {
-	UE_VLOG(GetOwner(), LogPathFollowing, Log, TEXT("PauseMove: RequestID(%u)"), RequestID);
+	UE_VLOG(GetOwner(), LogPathFollowing, Log, TEXT("PauseMove: RequestID(%u)"), RequestID.GetID());
 	if (Status == EPathFollowingStatus::Paused)
 	{
 		return;
@@ -514,7 +514,7 @@ void UPathFollowingComponent::ResumeMove(FAIRequestID RequestID)
 {
 	if (RequestID.IsEquivalent(CurrentRequestId) && RequestID.IsValid())
 	{
-		UE_VLOG(GetOwner(), LogPathFollowing, Log, TEXT("ResumeMove: RequestID(%u)"), RequestID);
+		UE_VLOG(GetOwner(), LogPathFollowing, Log, TEXT("ResumeMove: RequestID(%u)"), RequestID.GetID());
 
 		const bool bMovedDuringPause = ShouldCheckPathOnResume();
 		const bool bIsOnPath = IsOnPath();
@@ -546,7 +546,7 @@ void UPathFollowingComponent::ResumeMove(FAIRequestID RequestID)
 	}
 	else
 	{
-		UE_VLOG(GetOwner(), LogPathFollowing, Log, TEXT("ResumeMove: RequestID(%u) is neither \'AnyRequest\' not CurrentRequestId(%u). Ignoring."), RequestID, CurrentRequestId);
+		UE_VLOG(GetOwner(), LogPathFollowing, Log, TEXT("ResumeMove: RequestID(%u) is neither \'AnyRequest\' not CurrentRequestId(%u). Ignoring."), RequestID.GetID(), CurrentRequestId.GetID());
 	}
 }
 
diff --git a/Engine/Source/Runtime/Core/Public/Templates/IsValidVariadicFunctionArg.h b/Engine/Source/Runtime/Core/Public/Templates/IsValidVariadicFunctionArg.h
--- a/Engine/Source/Runtime/Core/Public/Templates/IsValidVariadicFunctionArg.h
+++ b/Engine/Source/Runtime/Core/Public/Templates/IsValidVariadicFunctionArg.h
@@ -4,35 +4,21 @@
 
 #include "CoreTypes.h"
 #include "IsEnum.h"
+#include <type_traits>
 
 /**
  * Tests if a type is a valid argument to a variadic function, e.g. printf.
  */
-template <typename T, bool = TIsEnum<T>::Value>
+template <typename T>
 struct TIsValidVariadicFunctionArg
 {
 private:
-	static uint32 Tester(uint32);
-	static uint32 Tester(uint8);
-	static uint32 Tester(int32);
-	static uint32 Tester(uint64);
-	static uint32 Tester(int64);
-	static uint32 Tester(double);
-	static uint32 Tester(long);
-	static uint32 Tester(unsigned long);
-	static uint32 Tester(TCHAR);
-	static uint32 Tester(bool);
-	static uint32 Tester(const void*);
-	static uint8  Tester(...);
-
-	static T DeclValT();
+	using DecayedT = std::decay_t<T>;
 
 public:
-	enum { Value = sizeof(Tester(DeclValT())) == sizeof(uint32) };
-};
-
-template <typename T>
-struct TIsValidVariadicFunctionArg<T, true>
-{
-	enum { Value = true };
+	static constexpr bool Value =
+		std::is_enum_v      <DecayedT> ||
+		std::is_arithmetic_v<DecayedT> ||
+		std::is_pointer_v   <DecayedT> ||
+		std::is_same_v      <DecayedT, TYPE_OF_NULLPTR>;
 };
diff --git a/Engine/Source/Runtime/Engine/Private/KismetSystemLibrary.cpp b/Engine/Source/Runtime/Engine/Private/KismetSystemLibrary.cpp
--- a/Engine/Source/Runtime/Engine/Private/KismetSystemLibrary.cpp
+++ b/Engine/Source/Runtime/Engine/Private/KismetSystemLibrary.cpp
@@ -391,7 +391,7 @@ void UKismetSystemLibrary::PrintString(const UObject* WorldContextObject, const
 				case NM_Client:
 					// GPlayInEditorID 0 is always the server, so 1 will be first client.
 					// You want to keep this logic in sync with GeneratePIEViewportWindowTitle and UpdatePlayInEditorWorldDebugString
-					Prefix = FString::Printf(TEXT("Client %d: "), GPlayInEditorID);
+					Prefix = FString::Printf(TEXT("Client %d: "), (int32)GPlayInEditorID);
 					break;
 				case NM_DedicatedServer:
 				case NM_ListenServer:
diff --git a/Engine/Source/Runtime/Engine/Private/Materials/MaterialExpressionHLSL.cpp b/Engine/Source/Runtime/Engine/Private/Materials/MaterialExpressionHLSL.cpp
--- a/Engine/Source/Runtime/Engine/Private/Materials/MaterialExpressionHLSL.cpp
+++ b/Engine/Source/Runtime/Engine/Private/Materials/MaterialExpressionHLSL.cpp
@@ -1310,7 +1310,7 @@ bool UMaterialExpressionTextureProperty::GenerateHLSLExpression(FMaterialHLSLGen
 
 	if (Property < 0 || Property >= TMTM_MAX)
 	{
-		return Generator.Errorf(TEXT("Invalid texture property %d"), Property);
+		return Generator.Errorf(TEXT("Invalid texture property %d"), Property.GetIntValue());
 	}
 
 	const FExpression* TextureExpression = TextureObject.AcquireHLSLExpression(Generator, Scope);
diff --git a/Engine/Source/Runtime/Engine/Private/PSOPrecacheValidation.cpp b/Engine/Source/Runtime/Engine/Private/PSOPrecacheValidation.cpp
--- a/Engine/Source/Runtime/Engine/Private/PSOPrecacheValidation.cpp
+++ b/Engine/Source/Runtime/Engine/Private/PSOPrecacheValidation.cpp
@@ -488,7 +488,7 @@ static void LogVertexElement(const FVertexElement& VertexElement, PSOMissStringB
 {
 	StringBuilder.Appendf(TEXT("\n\t\t\t\t\tStreamIndex:\t\t%d"), VertexElement.StreamIndex);
 	StringBuilder.Appendf(TEXT("\n\t\t\t\t\tOffset:\t\t\t\t%d"), VertexElement.Offset);
-	StringBuilder.Appendf(TEXT("\n\t\t\t\t\tType:\t\t\t\t%d"), VertexElement.Type);
+	StringBuilder.Appendf(TEXT("\n\t\t\t\t\tType:\t\t\t\t%d"), VertexElement.Type.GetIntValue());
 	StringBuilder.Appendf(TEXT("\n\t\t\t\t\tAttributeIndex:\t\t%d"), VertexElement.AttributeIndex);
 	StringBuilder.Appendf(TEXT("\n\t\t\t\t\tStride:\t\t\t\t%d"), VertexElement.Stride);
 	StringBuilder.Appendf(TEXT("\n\t\t\t\t\tbUseInstanceIndex:\t%d"), VertexElement.bUseInstanceIndex);
diff --git a/Engine/Source/Runtime/Engine/Private/WorldPartition/DataLayer/ExternalDataLayerHelper.cpp b/Engine/Source/Runtime/Engine/Private/WorldPartition/DataLayer/ExternalDataLayerHelper.cpp
--- a/Engine/Source/Runtime/Engine/Private/WorldPartition/DataLayer/ExternalDataLayerHelper.cpp
+++ b/Engine/Source/Runtime/Engine/Private/WorldPartition/DataLayer/ExternalDataLayerHelper.cpp
@@ -8,7 +8,7 @@
 FString FExternalDataLayerHelper::GetExternalStreamingObjectPackageName(const UExternalDataLayerAsset* InExternalDataLayerAsset)
 {
 	check(InExternalDataLayerAsset);
-	return FString::Printf(TEXT("StreamingObject_%X"), InExternalDataLayerAsset->GetUID());
+	return FString::Printf(TEXT("StreamingObject_%X"), (uint32)InExternalDataLayerAsset->GetUID());
 }
 
 FString FExternalDataLayerHelper::GetExternalStreamingObjectName(const UExternalDataLayerAsset* InExternalDataLayerAsset)
