[BTDecorator_TimeLimit] tweaked an ensure to make sure it is not firing because of value imprecisions

diff --git a/Engine/Source/Runtime/AIModule/Private/BehaviorTree/Decorators/BTDecorator_TimeLimit.cpp b/Engine/Source/Runtime/AIModule/Private/BehaviorTree/Decorators/BTDecorator_TimeLimit.cpp
--- a/Engine/Source/Runtime/AIModule/Private/BehaviorTree/Decorators/BTDecorator_TimeLimit.cpp
+++ b/Engine/Source/Runtime/AIModule/Private/BehaviorTree/Decorators/BTDecorator_TimeLimit.cpp
@@ -58,7 +58,11 @@ bool UBTDecorator_TimeLimit::CalculateRawConditionValue(UBehaviorTreeComponent&
 
 void UBTDecorator_TimeLimit::TickNode(UBehaviorTreeComponent& OwnerComp, uint8* NodeMemory, const float DeltaSeconds)
 {
-	ensureMsgf(DeltaSeconds >= TimeLimit, TEXT("Using SetNextTickTime in OnBecomeRelevant should guarantee that we are only getting ticked when the time limit is finished"));
+	ensureMsgf(DeltaSeconds >= TimeLimit || FMath::IsNearlyEqual(DeltaSeconds, TimeLimit, UE_KINDA_SMALL_NUMBER),
+		TEXT("Using SetNextTickTime in OnBecomeRelevant should guarantee that we are only getting ticked when the time limit is finished. DT=%f, TimeLimit=%f"),
+		DeltaSeconds,
+		TimeLimit);
+
 	reinterpret_cast<FBTimeLimitMemory*>(NodeMemory)->bElapsed = true;
 
 	OwnerComp.RequestExecution(this);
