Fix security issue: Memory Corruption in SimpleCellGrid::Serialize() Due to Missing Bounds Check and Integer Overflow.
Making sure DataBytesCount is computed from width and height and prevent width*height from overflowing.

diff --git a/Engine/Source/Runtime/AIModule/Public/SimpleCellGrid.h b/Engine/Source/Runtime/AIModule/Public/SimpleCellGrid.h
--- a/Engine/Source/Runtime/AIModule/Public/SimpleCellGrid.h
+++ b/Engine/Source/Runtime/AIModule/Public/SimpleCellGrid.h
@@ -4,6 +4,7 @@
 
 #include "CoreMinimal.h"
 #include "Serialization/Archive.h"
+#include "Math/NumericLimits.h"
 
 struct FGridSize2D
 {
@@ -55,7 +56,10 @@ public:
 		Origin = FVector(Bounds.Min.X, Bounds.Min.Y, (Bounds.Min.Z + Bounds.Max.Z) * 0.5f);
 		UpdateWorldBounds();
 
-		Cells.AddDefaulted(GridSize.Width * GridSize.Height);
+		const uint64 TempCellCount = GridSize.Width * GridSize.Height;
+		const typename TArray<FCellType>::SizeType CellCount = FMath::Min(TempCellCount, (uint64)TNumericLimits<typename TArray<FCellType>::SizeType>::Max());
+		ensureMsgf(CellCount == TempCellCount, TEXT("Grid width and height are too big."));
+		Cells.AddDefaulted(CellCount);
 
 		return true;
 	}
@@ -248,7 +252,12 @@ public:
 			{
 				if (Ar.IsLoading())
 				{
-					Cells.SetNum(GridSize.Width * GridSize.Height);
+					const uint64 TempCellCount = GridSize.Width * GridSize.Height;
+					const typename TArray<FCellType>::SizeType CellCount = FMath::Min(TempCellCount, (uint64)TNumericLimits<typename TArray<FCellType>::SizeType>::Max());
+					ensureMsgf(CellCount == TempCellCount, TEXT("Grid width and height are too big."));
+					Cells.SetNum(CellCount);
+					
+					DataBytesCount = FMath::Clamp(DataBytesCount, 0, CellCount * sizeof(FCellType));
 				}
 
 				Ar.Serialize(Cells.GetData(), DataBytesCount);
