Add ForEachChildTask to BehaviorTreeComponent that allows to run a functor on all the child task in the current behavior stack.
Reimplement SetDynamicSubtree with ForEachChildTask
Add OptionalStartNode to SetDynamicSubtree to allow caller to give a hint where the dynamic subtree is and avoid having to parse the entire behavior tree to find it.

diff --git a/Engine/Source/Runtime/AIModule/Classes/BehaviorTree/BehaviorTreeComponent.h b/Engine/Source/Runtime/AIModule/Classes/BehaviorTree/BehaviorTreeComponent.h
--- a/Engine/Source/Runtime/AIModule/Classes/BehaviorTree/BehaviorTreeComponent.h
+++ b/Engine/Source/Runtime/AIModule/Classes/BehaviorTree/BehaviorTreeComponent.h
@@ -272,10 +272,18 @@ public:
 	UFUNCTION(BlueprintCallable, Category = "AI|Logic")
 	AIMODULE_API void AddCooldownTagDuration(FGameplayTag CooldownTag, float CooldownDuration, bool bAddToExistingDuration);
 
-	/** assign subtree to RunBehaviorDynamic task specified by tag */
+	/** assign subtree to RunBehaviorDynamic task specified by tag. */
 	UFUNCTION(BlueprintCallable, Category="AI|Logic")
 	AIMODULE_API virtual void SetDynamicSubtree(FGameplayTag InjectTag, UBehaviorTree* BehaviorAsset);
 
+	/** assign subtree to RunBehaviorDynamic task specified by tag. Optional Starting Node can be given if the location in the tree is known to avoid parsing the whole tree.  */
+	AIMODULE_API virtual void SetDynamicSubtree(FGameplayTag InjectTag, UBehaviorTree* BehaviorAsset, UBTCompositeNode* OptionalStartingNode);
+
+	/** Will call the given functor on each task node in the current instance stacks. */
+	AIMODULE_API void ForEachChildTask(TFunctionRef<void (UBTTaskNode&, const FBehaviorTreeInstance&, int32 InstanceIndex)> Functor);
+
+	/** Will call the given functor on each child node from the given start node. InstanceIndex can be found using FindInstanceContainingNode. */
+	AIMODULE_API void ForEachChildTask(UBTCompositeNode& StartNode, int32 InstanceIndex, TFunctionRef<void(UBTTaskNode&, const FBehaviorTreeInstance&, int32 InstanceIndex)> Functor);
 // Code for timing BT Search for FramePro
 #if !UE_BUILD_SHIPPING
 	static AIMODULE_API void EndFrame();
diff --git a/Engine/Source/Runtime/AIModule/Private/BehaviorTree/BehaviorTreeComponent.cpp b/Engine/Source/Runtime/AIModule/Private/BehaviorTree/BehaviorTreeComponent.cpp
--- a/Engine/Source/Runtime/AIModule/Private/BehaviorTree/BehaviorTreeComponent.cpp
+++ b/Engine/Source/Runtime/AIModule/Private/BehaviorTree/BehaviorTreeComponent.cpp
@@ -3039,77 +3039,85 @@ void UBehaviorTreeComponent::AddCooldownTagDuration(FGameplayTag CooldownTag, fl
 	}
 }
 
-bool SetDynamicSubtreeHelper(const UBTCompositeNode* TestComposite,
-	const FBehaviorTreeInstance& InstanceInfo, const UBehaviorTreeComponent* OwnerComp,
-	const FGameplayTag& InjectTag, UBehaviorTree* BehaviorAsset)
+void UBehaviorTreeComponent::SetDynamicSubtree(FGameplayTag InjectTag, UBehaviorTree* BehaviorAsset)
 {
-	bool bInjected = false;
+	SetDynamicSubtree(InjectTag, BehaviorAsset, nullptr);
+}
 
-	for (int32 Idx = 0; Idx < TestComposite->Children.Num(); Idx++)
-	{
-		const FBTCompositeChild& ChildInfo = TestComposite->Children[Idx];
-		if (ChildInfo.ChildComposite)
-		{
-			bInjected = (SetDynamicSubtreeHelper(ChildInfo.ChildComposite, InstanceInfo, OwnerComp, InjectTag, BehaviorAsset) || bInjected);
-		}
-		else
+void UBehaviorTreeComponent::SetDynamicSubtree(FGameplayTag InjectTag, UBehaviorTree* BehaviorAsset, UBTCompositeNode* OptionalStartingNode)
+{
+	bool bInjected = false;
+	auto ReplaceInjectedBehaviorInTask = [InjectTag, BehaviorAsset, this, &bInjected](UBTTaskNode& TaskNode, const FBehaviorTreeInstance& InstanceInfo, int32 InstanceIndex)
 		{
-			UBTTask_RunBehaviorDynamic* SubtreeTask = Cast<UBTTask_RunBehaviorDynamic>(ChildInfo.ChildTask);
+			UBTTask_RunBehaviorDynamic* SubtreeTask = Cast<UBTTask_RunBehaviorDynamic>(&TaskNode);
 			if (SubtreeTask && SubtreeTask->HasMatchingTag(InjectTag))
 			{
 				const uint8* NodeMemory = SubtreeTask->GetNodeMemory<uint8>(InstanceInfo);
-				UBTTask_RunBehaviorDynamic* InstancedNode = Cast<UBTTask_RunBehaviorDynamic>(SubtreeTask->GetNodeInstance(*OwnerComp, (uint8*)NodeMemory));
-				if (InstancedNode)
+				if (UBTTask_RunBehaviorDynamic* InstancedNode = Cast<UBTTask_RunBehaviorDynamic>(SubtreeTask->GetNodeInstance(*this, (uint8*)NodeMemory)))
 				{
 					const bool bAssetChanged = InstancedNode->SetBehaviorAsset(BehaviorAsset);
 					if (bAssetChanged)
 					{
-						UE_VLOG(OwnerComp->GetOwner(), LogBehaviorTree, Log, TEXT("Replaced subtree in %s with %s (tag: %s)"),
+						UE_VLOG(GetOwner(), LogBehaviorTree, Log, TEXT("Replaced subtree in %s with %s (tag: %s)"),
 							*UBehaviorTreeTypes::DescribeNodeHelper(SubtreeTask), *GetNameSafe(BehaviorAsset), *InjectTag.ToString());
 						bInjected = true;
+						if (InstanceInfo.ActiveNodeType == EBTActiveNode::ActiveTask && SubtreeTask == InstanceInfo.ActiveNode)
+						{
+							UBTCompositeNode* RestartNode = SubtreeTask->GetParentNode();
+							int32 RestartChildIdx = RestartNode->GetChildIndex(*SubtreeTask);
+							RequestExecution(RestartNode, InstanceIndex, SubtreeTask, RestartChildIdx, EBTNodeResult::Aborted);
+						}
 					}
 				}
 			}
-		}
+		};
+
+	if (OptionalStartingNode)
+	{
+		ForEachChildTask(*OptionalStartingNode, FindInstanceContainingNode(OptionalStartingNode), ReplaceInjectedBehaviorInTask);
+	}
+	else
+	{
+		ForEachChildTask(ReplaceInjectedBehaviorInTask);
 	}
 
-	return bInjected;
+	if (!bInjected)
+	{
+		UE_VLOG(GetOwner(), LogBehaviorTree, Log, TEXT("Failed to inject subtree %s at tag %s"), *GetNameSafe(BehaviorAsset), *InjectTag.ToString());
+	}
 }
 
-void UBehaviorTreeComponent::SetDynamicSubtree(FGameplayTag InjectTag, UBehaviorTree* BehaviorAsset)
+void UBehaviorTreeComponent::ForEachChildTask(TFunctionRef<void(class UBTTaskNode&, const FBehaviorTreeInstance&, int32 InstanceIndex)> Functor)
 {
-	bool bInjected = false;
-	// replace at matching injection points
 	for (int32 InstanceIndex = 0; InstanceIndex < InstanceStack.Num(); InstanceIndex++)
 	{
 		const FBehaviorTreeInstance& InstanceInfo = InstanceStack[InstanceIndex];
-		bInjected = (SetDynamicSubtreeHelper(InstanceInfo.RootNode, InstanceInfo, this, InjectTag, BehaviorAsset) || bInjected);
+		ForEachChildTask(*InstanceInfo.RootNode, InstanceIndex, Functor);
 	}
+}
 
-	// restart subtree if it was replaced
-	if (bInjected)
+void UBehaviorTreeComponent::ForEachChildTask(UBTCompositeNode& StartNode, int32 InstanceIndex, TFunctionRef<void(UBTTaskNode&, const FBehaviorTreeInstance&, int32 InstanceIndex)> Functor)
+{
+	if (!InstanceStack.IsValidIndex(InstanceIndex))
 	{
-		for (int32 InstanceIndex = 0; InstanceIndex < InstanceStack.Num(); InstanceIndex++)
+		return;
+	}
+
+	for (int32 Idx = 0; Idx < StartNode.Children.Num(); Idx++)
+	{
+		const FBTCompositeChild& ChildInfo = StartNode.Children[Idx];
+		if (ChildInfo.ChildComposite)
 		{
-			const FBehaviorTreeInstance& InstanceInfo = InstanceStack[InstanceIndex];
-			if (InstanceInfo.ActiveNodeType == EBTActiveNode::ActiveTask)
+			ForEachChildTask(*ChildInfo.ChildComposite, InstanceIndex, Functor);
+		}
+		else
+		{
+			if (ChildInfo.ChildTask)
 			{
-				const UBTTask_RunBehaviorDynamic* SubtreeTask = Cast<const UBTTask_RunBehaviorDynamic>(InstanceInfo.ActiveNode);
-				if (SubtreeTask && SubtreeTask->HasMatchingTag(InjectTag))
-				{
-					UBTCompositeNode* RestartNode = SubtreeTask->GetParentNode();
-					int32 RestartChildIdx = RestartNode->GetChildIndex(*SubtreeTask);
-
-					RequestExecution(RestartNode, InstanceIndex, SubtreeTask, RestartChildIdx, EBTNodeResult::Aborted);
-					break;
-				}
+				Functor(*ChildInfo.ChildTask, InstanceStack[InstanceIndex], InstanceIndex);
 			}
 		}
 	}
-	else
-	{
-		UE_VLOG(GetOwner(), LogBehaviorTree, Log, TEXT("Failed to inject subtree %s at tag %s"), *GetNameSafe(BehaviorAsset), *InjectTag.ToString());
-	}
 }
 
 #if ENABLE_VISUAL_LOG
