Expose RequireNavigableEndLocation to MoveToLocationOrActor BP function.

diff --git a/Engine/Source/Runtime/AIModule/Classes/Tasks/AITask_MoveTo.h b/Engine/Source/Runtime/AIModule/Classes/Tasks/AITask_MoveTo.h
--- a/Engine/Source/Runtime/AIModule/Classes/Tasks/AITask_MoveTo.h
+++ b/Engine/Source/Runtime/AIModule/Classes/Tasks/AITask_MoveTo.h
@@ -31,10 +31,11 @@ public:
 	bool WasMoveSuccessful() const { return MoveResult == EPathFollowingResult::Success; }
 	bool WasMovePartial() const { return Path.IsValid() && Path->IsPartial(); }
 
-	UFUNCTION(BlueprintCallable, Category = "AI|Tasks", meta = (AdvancedDisplay = "AcceptanceRadius,StopOnOverlap,AcceptPartialPath,bUsePathfinding,bUseContinuousGoalTracking,ProjectGoalOnNavigation", DefaultToSelf = "Controller", BlueprintInternalUseOnly = "TRUE", DisplayName = "Move To Location or Actor"))
+	UFUNCTION(BlueprintCallable, Category = "AI|Tasks", meta = (AdvancedDisplay = "AcceptanceRadius,StopOnOverlap,AcceptPartialPath,bUsePathfinding,bUseContinuousGoalTracking,ProjectGoalOnNavigation,RequireNavigableEndLocation", DefaultToSelf = "Controller", BlueprintInternalUseOnly = "TRUE", DisplayName = "Move To Location or Actor"))
 	static AIMODULE_API UAITask_MoveTo* AIMoveTo(AAIController* Controller, FVector GoalLocation, AActor* GoalActor = nullptr,
 		float AcceptanceRadius = -1.f, EAIOptionFlag::Type StopOnOverlap = EAIOptionFlag::Default, EAIOptionFlag::Type AcceptPartialPath = EAIOptionFlag::Default,
-		bool bUsePathfinding = true, bool bLockAILogic = true, bool bUseContinuousGoalTracking = false, EAIOptionFlag::Type ProjectGoalOnNavigation = EAIOptionFlag::Default);
+		bool bUsePathfinding = true, bool bLockAILogic = true, bool bUseContinuousGoalTracking = false, EAIOptionFlag::Type ProjectGoalOnNavigation = EAIOptionFlag::Default,
+		EAIOptionFlag::Type RequireNavigableEndLocation = EAIOptionFlag::Default);
 
 	/** Allows custom move request tweaking. Note that all MoveRequest need to 
 	 *	be performed before PerformMove is called. */
diff --git a/Engine/Source/Runtime/AIModule/Private/Tasks/AITask_MoveTo.cpp b/Engine/Source/Runtime/AIModule/Private/Tasks/AITask_MoveTo.cpp
--- a/Engine/Source/Runtime/AIModule/Private/Tasks/AITask_MoveTo.cpp
+++ b/Engine/Source/Runtime/AIModule/Private/Tasks/AITask_MoveTo.cpp
@@ -31,7 +31,8 @@ UAITask_MoveTo::UAITask_MoveTo(const FObjectInitializer& ObjectInitializer)
 
 UAITask_MoveTo* UAITask_MoveTo::AIMoveTo(AAIController* Controller, FVector InGoalLocation, AActor* InGoalActor,
 	float AcceptanceRadius, EAIOptionFlag::Type StopOnOverlap, EAIOptionFlag::Type AcceptPartialPath,
-	bool bUsePathfinding, bool bLockAILogic, bool bUseContinuousGoalTracking, EAIOptionFlag::Type ProjectGoalOnNavigation)
+	bool bUsePathfinding, bool bLockAILogic, bool bUseContinuousGoalTracking, EAIOptionFlag::Type ProjectGoalOnNavigation,
+	EAIOptionFlag::Type RequireNavigableEndLocation)
 {
 	UAITask_MoveTo* MyTask = Controller ? UAITask::NewAITask<UAITask_MoveTo>(*Controller, EAITaskPriority::High) : nullptr;
 	if (MyTask)
@@ -51,6 +52,7 @@ UAITask_MoveTo* UAITask_MoveTo::AIMoveTo(AAIController* Controller, FVector InGo
 		MoveReq.SetAllowPartialPath(FAISystem::PickAIOption(AcceptPartialPath, MoveReq.IsUsingPartialPaths()));
 		MoveReq.SetUsePathfinding(bUsePathfinding);
 		MoveReq.SetProjectGoalLocation(FAISystem::PickAIOption(ProjectGoalOnNavigation, MoveReq.IsProjectingGoal()));
+		MoveReq.SetRequireNavigableEndLocation(FAISystem::PickAIOption(RequireNavigableEndLocation, MoveReq.IsNavigableEndLocationRequired()));
 		if (Controller)
 		{
 			MoveReq.SetNavigationFilter(Controller->GetDefaultNavigationFilterClass());
