Added support for collision response parameter in EnvQueryTest_Trace
Previously ignored GeometryByProfile mode is now taken into account

#localization none
#tests PIE

diff --git a/Engine/Source/Runtime/AIModule/Classes/EnvironmentQuery/Tests/EnvQueryTest_Trace.h b/Engine/Source/Runtime/AIModule/Classes/EnvironmentQuery/Tests/EnvQueryTest_Trace.h
--- a/Engine/Source/Runtime/AIModule/Classes/EnvironmentQuery/Tests/EnvQueryTest_Trace.h
+++ b/Engine/Source/Runtime/AIModule/Classes/EnvironmentQuery/Tests/EnvQueryTest_Trace.h
@@ -48,14 +48,23 @@ class UEnvQueryTest_Trace : public UEnvQueryTest
 
 protected:
 
+	UE_DEPRECATED(5.4, "This delegate type is not used anymore in UEnvQueryTest_Trace, new code is done in EnvQueryTest_Trace_Helpers namespace helper functions")
 	DECLARE_DELEGATE_RetVal_SevenParams(bool, FRunTraceSignature, const FVector&, const FVector&, AActor*, UWorld*, enum ECollisionChannel, const FCollisionQueryParams&, const FVector&);
 
+	UE_DEPRECATED(5.4, "This function will be deleted")
 	AIMODULE_API bool RunLineTraceTo(const FVector& ItemPos, const FVector& ContextPos, AActor* ItemActor, UWorld* World, enum ECollisionChannel Channel, const FCollisionQueryParams& Params, const FVector& Extent);
+	UE_DEPRECATED(5.4, "This function will be deleted")
 	AIMODULE_API bool RunLineTraceFrom(const FVector& ItemPos, const FVector& ContextPos, AActor* ItemActor, UWorld* World, enum ECollisionChannel Channel, const FCollisionQueryParams& Params, const FVector& Extent);
+	UE_DEPRECATED(5.4, "This function will be deleted")
 	AIMODULE_API bool RunBoxTraceTo(const FVector& ItemPos, const FVector& ContextPos, AActor* ItemActor, UWorld* World, enum ECollisionChannel Channel, const FCollisionQueryParams& Params, const FVector& Extent);
+	UE_DEPRECATED(5.4, "This function will be deleted")
 	AIMODULE_API bool RunBoxTraceFrom(const FVector& ItemPos, const FVector& ContextPos, AActor* ItemActor, UWorld* World, enum ECollisionChannel Channel, const FCollisionQueryParams& Params, const FVector& Extent);
+	UE_DEPRECATED(5.4, "This function will be deleted")
 	AIMODULE_API bool RunSphereTraceTo(const FVector& ItemPos, const FVector& ContextPos, AActor* ItemActor, UWorld* World, enum ECollisionChannel Channel, const FCollisionQueryParams& Params, const FVector& Extent);
+	UE_DEPRECATED(5.4, "This function will be deleted")
 	AIMODULE_API bool RunSphereTraceFrom(const FVector& ItemPos, const FVector& ContextPos, AActor* ItemActor, UWorld* World, enum ECollisionChannel Channel, const FCollisionQueryParams& Params, const FVector& Extent);
+	UE_DEPRECATED(5.4, "This function will be deleted")
 	AIMODULE_API bool RunCapsuleTraceTo(const FVector& ItemPos, const FVector& ContextPos, AActor* ItemActor, UWorld* World, enum ECollisionChannel Channel, const FCollisionQueryParams& Params, const FVector& Extent);
+	UE_DEPRECATED(5.4, "This function will be deleted")
 	AIMODULE_API bool RunCapsuleTraceFrom(const FVector& ItemPos, const FVector& ContextPos, AActor* ItemActor, UWorld* World, enum ECollisionChannel Channel, const FCollisionQueryParams& Params, const FVector& Extent);
 };
diff --git a/Engine/Source/Runtime/AIModule/Private/EnvironmentQuery/Tests/EnvQueryTest_Trace.cpp b/Engine/Source/Runtime/AIModule/Private/EnvironmentQuery/Tests/EnvQueryTest_Trace.cpp
--- a/Engine/Source/Runtime/AIModule/Private/EnvironmentQuery/Tests/EnvQueryTest_Trace.cpp
+++ b/Engine/Source/Runtime/AIModule/Private/EnvironmentQuery/Tests/EnvQueryTest_Trace.cpp
@@ -7,11 +7,105 @@
 #include "Engine/World.h"
 #include "EnvironmentQuery/Items/EnvQueryItemType_VectorBase.h"
 #include "EnvironmentQuery/Contexts/EnvQueryContext_Querier.h"
+#include "VisualLogger/VisualLogger.h"
 
 #include UE_INLINE_GENERATED_CPP_BY_NAME(EnvQueryTest_Trace)
 
 #define LOCTEXT_NAMESPACE "EnvQueryGenerator"
 
+namespace EnvQueryTest_Trace_Helpers
+{
+	template <EEnvTraceShape::Type Shape>
+	bool TraceShape(const FVector& Start, const FVector& End, UWorld* World, enum ECollisionChannel Channel, const FCollisionQueryParams& TraceParams, const FVector3f& Extent, const FCollisionResponseParams& ResponseParams)
+	{
+		if constexpr (Shape == EEnvTraceShape::Line)
+		{
+			return World->LineTraceTestByChannel(Start, End, Channel, TraceParams, ResponseParams);
+		}
+		else if constexpr (Shape == EEnvTraceShape::Box)
+		{
+			return World->SweepTestByChannel(Start, End, FQuat((End - Start).Rotation()), Channel, FCollisionShape::MakeBox(Extent), TraceParams, ResponseParams);
+		}
+		else if constexpr (Shape == EEnvTraceShape::Sphere)
+		{
+			return World->SweepTestByChannel(Start, End, FQuat::Identity, Channel, FCollisionShape::MakeSphere(Extent.X), TraceParams, ResponseParams);
+		}
+		else if constexpr (Shape == EEnvTraceShape::Capsule)
+		{
+			return World->SweepTestByChannel(Start, End, FQuat::Identity, Channel, FCollisionShape::MakeCapsule(Extent.X, Extent.Z), TraceParams, ResponseParams);
+		}
+		else
+		{
+			[] <bool cond = false>() { static_assert(cond, "Unsupported value of EEnvTraceShape received in TraceShape"); }(); // static_assert must be type-dependent to avoid "ill-formed" code
+		}
+	}
+
+	template <bool bTraceToItem, EEnvTraceShape::Type Shape>
+	bool TraceShapeWithDir(const FVector& ItemPos, const FVector& ContextPos, UWorld* World, enum ECollisionChannel Channel, const FCollisionQueryParams& TraceParams, const FVector3f& Extent, const FCollisionResponseParams& ResponseParams)
+	{
+		if constexpr (bTraceToItem)
+		{
+			return TraceShape<Shape>(ContextPos, ItemPos, World, Channel, TraceParams, Extent, ResponseParams);
+		}
+		else
+		{
+			return TraceShape<Shape>(ItemPos, ContextPos, World, Channel, TraceParams, Extent, ResponseParams);
+		}
+	}
+
+	// Use a templated function execute the traces in order to avoid using any branch or function pointer inside the item loop
+	template <bool bTraceToItem, EEnvTraceShape::Type Shape>
+	void RunTraces(const UEnvQueryTest_Trace& Query, const TArrayView<FVector>& ContextLocations, FEnvQueryInstance& QueryInstance, float ContextZ, float ItemZ, const FEnvTraceData& TraceData, EEnvTestPurpose::Type TestPurpose, EEnvTestFilterType::Type FilterType, bool bWantsHit, const FCollisionQueryParams& TraceParams)
+	{
+		ECollisionChannel TraceCollisionChannel = ECC_WorldStatic;
+		FCollisionResponseParams ResponseParams = FCollisionResponseParams::DefaultResponseParam;
+		if (TraceData.TraceMode == EEnvQueryTrace::Type::GeometryByProfile)
+		{
+			if (!UCollisionProfile::GetChannelAndResponseParams(TraceData.TraceProfileName, TraceCollisionChannel, ResponseParams))
+			{
+				UE_VLOG_ALWAYS_UELOG(QueryInstance.Owner.Get(), LogEQS, Error, 
+					TEXT("Unable to fetch collision channel and response from TraceProfileName %s, test %s for query %s will automatically fail"),
+					*TraceData.TraceProfileName.ToString(),
+					*Query.GetName(),
+					*QueryInstance.QueryName);
+
+				for (FEnvQueryInstance::ItemIterator It(&Query, QueryInstance); It; ++It)
+				{
+					It.SetScore(TestPurpose, FilterType, !bWantsHit, bWantsHit);
+				}
+
+				return;
+			}
+		}
+		else if (TraceData.TraceMode == EEnvQueryTrace::Type::GeometryByChannel)
+		{
+			TraceCollisionChannel = UEngineTypes::ConvertToCollisionChannel(TraceData.TraceChannel);
+		}
+
+		FVector3f TraceExtent(TraceData.ExtentX, TraceData.ExtentY, TraceData.ExtentZ);
+
+		for (int32 ContextIndex = 0; ContextIndex < ContextLocations.Num(); ContextIndex++)
+		{
+			ContextLocations[ContextIndex].Z += ContextZ;
+		}
+
+		for (FEnvQueryInstance::ItemIterator It(&Query, QueryInstance); It; ++It)
+		{
+			const FVector ItemLocation = Query.GetItemLocation(QueryInstance, It.GetIndex()) + FVector(0, 0, ItemZ);
+			AActor* ItemActor = Query.GetItemActor(QueryInstance, It.GetIndex());
+
+			FCollisionQueryParams PerItemTraceParams(TraceParams);
+			PerItemTraceParams.AddIgnoredActor(ItemActor);
+
+			for (int32 ContextIndex = 0; ContextIndex < ContextLocations.Num(); ContextIndex++)
+			{
+				const bool bHit = TraceShapeWithDir<bTraceToItem, Shape>(ItemLocation, ContextLocations[ContextIndex], QueryInstance.World, TraceCollisionChannel, PerItemTraceParams, TraceExtent, ResponseParams);
+				It.SetScore(TestPurpose, FilterType, bHit, bWantsHit);
+			}
+		}
+	}
+}
+
 UEnvQueryTest_Trace::UEnvQueryTest_Trace(const FObjectInitializer& ObjectInitializer) : Super(ObjectInitializer)
 {
 	Cost = EEnvTestCost::High;
@@ -48,47 +142,49 @@ void UEnvQueryTest_Trace::RunTest(FEnvQueryInstance& QueryInstance) const
 	{
 		TraceParams.AddIgnoredActors(IgnoredActors);
 	}
-	
-	ECollisionChannel TraceCollisionChannel = UEngineTypes::ConvertToCollisionChannel(TraceData.TraceChannel);	
-	FVector TraceExtent(TraceData.ExtentX, TraceData.ExtentY, TraceData.ExtentZ);
-	FRunTraceSignature TraceFunc;
+
 	switch (TraceData.TraceShape)
 	{
 	case EEnvTraceShape::Line:
-		TraceFunc.BindUObject(const_cast<UEnvQueryTest_Trace*>(this), bTraceToItem ? &UEnvQueryTest_Trace::RunLineTraceTo : &UEnvQueryTest_Trace::RunLineTraceFrom);
+		if (bTraceToItem)
+		{
+			EnvQueryTest_Trace_Helpers::RunTraces<true, EEnvTraceShape::Line>(*this, ContextLocations, QueryInstance, ContextZ, ItemZ, TraceData, TestPurpose, FilterType, bWantsHit, TraceParams);
+		}
+		else
+		{
+			EnvQueryTest_Trace_Helpers::RunTraces<false, EEnvTraceShape::Line>(*this, ContextLocations, QueryInstance, ContextZ, ItemZ, TraceData, TestPurpose, FilterType, bWantsHit, TraceParams);
+		}
 		break;
-
 	case EEnvTraceShape::Box:
-		TraceFunc.BindUObject(const_cast<UEnvQueryTest_Trace*>(this), bTraceToItem ? &UEnvQueryTest_Trace::RunBoxTraceTo : &UEnvQueryTest_Trace::RunBoxTraceFrom);
+		if (bTraceToItem)
+		{
+			EnvQueryTest_Trace_Helpers::RunTraces<true, EEnvTraceShape::Box>(*this, ContextLocations, QueryInstance, ContextZ, ItemZ, TraceData, TestPurpose, FilterType, bWantsHit, TraceParams);
+		}
+		else
+		{
+			EnvQueryTest_Trace_Helpers::RunTraces<false, EEnvTraceShape::Box>(*this, ContextLocations, QueryInstance, ContextZ, ItemZ, TraceData, TestPurpose, FilterType, bWantsHit, TraceParams);
+		}
 		break;
-
 	case EEnvTraceShape::Sphere:
-		TraceFunc.BindUObject(const_cast<UEnvQueryTest_Trace*>(this), bTraceToItem ? &UEnvQueryTest_Trace::RunSphereTraceTo : &UEnvQueryTest_Trace::RunSphereTraceFrom);
+		if (bTraceToItem)
+		{
+			EnvQueryTest_Trace_Helpers::RunTraces<true, EEnvTraceShape::Sphere>(*this, ContextLocations, QueryInstance, ContextZ, ItemZ, TraceData, TestPurpose, FilterType, bWantsHit, TraceParams);
+		}
+		else
+		{
+			EnvQueryTest_Trace_Helpers::RunTraces<false, EEnvTraceShape::Sphere>(*this, ContextLocations, QueryInstance, ContextZ, ItemZ, TraceData, TestPurpose, FilterType, bWantsHit, TraceParams);
+		}
 		break;
-
 	case EEnvTraceShape::Capsule:
-		TraceFunc.BindUObject(const_cast<UEnvQueryTest_Trace*>(this), bTraceToItem ? &UEnvQueryTest_Trace::RunCapsuleTraceTo : &UEnvQueryTest_Trace::RunCapsuleTraceFrom);
-		break;
-
-	default:
-		return;
-	}
-
-	for (int32 ContextIndex = 0; ContextIndex < ContextLocations.Num(); ContextIndex++)
-	{
-		ContextLocations[ContextIndex].Z += ContextZ;
-	}
-
-	for (FEnvQueryInstance::ItemIterator It(this, QueryInstance); It; ++It)
-	{
-		const FVector ItemLocation = GetItemLocation(QueryInstance, It.GetIndex()) + FVector(0, 0, ItemZ);
-		AActor* ItemActor = GetItemActor(QueryInstance, It.GetIndex());
-
-		for (int32 ContextIndex = 0; ContextIndex < ContextLocations.Num(); ContextIndex++)
+		if (bTraceToItem)
 		{
-			const bool bHit = TraceFunc.Execute(ItemLocation, ContextLocations[ContextIndex], ItemActor, QueryInstance.World, TraceCollisionChannel, TraceParams, TraceExtent);
-			It.SetScore(TestPurpose, FilterType, bHit, bWantsHit);
+			EnvQueryTest_Trace_Helpers::RunTraces<true, EEnvTraceShape::Capsule>(*this, ContextLocations, QueryInstance, ContextZ, ItemZ, TraceData, TestPurpose, FilterType, bWantsHit, TraceParams);
 		}
+		else
+		{
+			EnvQueryTest_Trace_Helpers::RunTraces<false, EEnvTraceShape::Capsule>(*this, ContextLocations, QueryInstance, ContextZ, ItemZ, TraceData, TestPurpose, FilterType, bWantsHit, TraceParams);
+		}
+		break;
 	}
 }
 
@@ -117,6 +213,8 @@ FText UEnvQueryTest_Trace::GetDescriptionDetails() const
 		TraceData.ToText(FEnvTraceData::Detailed), DescribeBoolTestParams("hit"));
 }
 
+PRAGMA_DISABLE_DEPRECATION_WARNINGS
+
 bool UEnvQueryTest_Trace::RunLineTraceTo(const FVector& ItemPos, const FVector& ContextPos, AActor* ItemActor, UWorld* World, enum ECollisionChannel Channel, const FCollisionQueryParams& Params, const FVector& Extent)
 {
 	FCollisionQueryParams TraceParams(Params);
@@ -189,5 +287,7 @@ bool UEnvQueryTest_Trace::RunCapsuleTraceFrom(const FVector& ItemPos, const FVec
 	return bHit;
 }
 
+PRAGMA_ENABLE_DEPRECATION_WARNINGS
+
 #undef LOCTEXT_NAMESPACE
 
