PR #11479: Allow soft class pointers to be printed as node body text in behavior tree nodes or conversation nodes

diff --git a/Engine/Source/Runtime/AIModule/Private/BehaviorTree/BlueprintNodeHelpers.cpp b/Engine/Source/Runtime/AIModule/Private/BehaviorTree/BlueprintNodeHelpers.cpp
--- a/Engine/Source/Runtime/AIModule/Private/BehaviorTree/BlueprintNodeHelpers.cpp
+++ b/Engine/Source/Runtime/AIModule/Private/BehaviorTree/BlueprintNodeHelpers.cpp
@@ -120,7 +120,7 @@ namespace BlueprintNodeHelpers
 		}
 		else
 		{
-			Prop->ExportTextItem_Direct(ExportedStringValue, PropertyAddr, NULL, NULL, PPF_PropertyWindow, NULL);
+			Prop->ExportTextItem_Direct(ExportedStringValue, PropertyAddr, nullptr, nullptr, PPF_PropertyWindow | PPF_SimpleObjectText, nullptr);
 		}
 
 		const bool bIsBool = Prop->IsA(FBoolProperty::StaticClass());
@@ -228,7 +228,8 @@ namespace BlueprintNodeHelpers
 				continue;
 			}
 
-			if (TestProperty->IsA(FClassProperty::StaticClass()) ||
+			if (TestProperty->IsA(FSoftClassProperty::StaticClass()) ||
+				TestProperty->IsA(FClassProperty::StaticClass()) ||
 				TestProperty->IsA(FStructProperty::StaticClass()) ||
 				CanUsePropertyType(TestProperty))
 			{
