Environment Query: Add EQS Owner and EQS Asset names in profiling traces

diff --git a/Engine/Source/Runtime/AIModule/Classes/EnvironmentQuery/EnvQueryTypes.h b/Engine/Source/Runtime/AIModule/Classes/EnvironmentQuery/EnvQueryTypes.h
--- a/Engine/Source/Runtime/AIModule/Classes/EnvironmentQuery/EnvQueryTypes.h
+++ b/Engine/Source/Runtime/AIModule/Classes/EnvironmentQuery/EnvQueryTypes.h
@@ -1272,6 +1272,10 @@ public:
 	static AIMODULE_API bool bDebuggingInfoEnabled;
 #endif // USE_EQS_DEBUGGER
 
+#if STATS
+	TStatId StatId;
+#endif
+
 	AIMODULE_API FBox GetBoundingBox() const;
 };
 
diff --git a/Engine/Source/Runtime/AIModule/Private/EnvironmentQuery/EnvQueryManager.cpp b/Engine/Source/Runtime/AIModule/Private/EnvironmentQuery/EnvQueryManager.cpp
--- a/Engine/Source/Runtime/AIModule/Private/EnvironmentQuery/EnvQueryManager.cpp
+++ b/Engine/Source/Runtime/AIModule/Private/EnvironmentQuery/EnvQueryManager.cpp
@@ -470,7 +470,13 @@ void UEnvQueryManager::Tick(float DeltaTime)
 			}
 			else
 			{
+#if STATS
+				FScopeCycleCounterUObject OwnerScopeCounter(QueryInstance->Owner.Get());
+				FScopeCycleCounter QueryScopeCounter(QueryInstance->StatId);
+#endif // STATS
+
 				QueryInstancePtr->ExecuteOneStep(TimeLeft);
+
 #if USE_EQS_DEBUGGER
 				bWorkHasBeenDone = true;
 #endif // USE_EQS_DEBUGGER
@@ -802,6 +808,7 @@ TSharedPtr<FEnvQueryInstance> UEnvQueryManager::CreateQueryInstance(const UEnvQu
 			NewCacheEntry.Instance.UniqueName = LocalTemplate->GetFName();
 			NewCacheEntry.Instance.QueryName = LocalTemplate->GetQueryName().ToString();
 			NewCacheEntry.Instance.Mode = RunMode;
+			STAT(NewCacheEntry.Instance.StatId = FDynamicStats::CreateStatId<FStatGroup_STATGROUP_AI_EQS>(NewCacheEntry.Instance.UniqueName));
 
 			const int32 Idx = InstanceCache.Add(NewCacheEntry);
 			InstanceTemplate = &InstanceCache[Idx].Instance;
