Fix static analyzer error in EnvQueryTest_Trace

diff --git a/Engine/Source/Runtime/AIModule/Private/EnvironmentQuery/Tests/EnvQueryTest_Trace.cpp b/Engine/Source/Runtime/AIModule/Private/EnvironmentQuery/Tests/EnvQueryTest_Trace.cpp
--- a/Engine/Source/Runtime/AIModule/Private/EnvironmentQuery/Tests/EnvQueryTest_Trace.cpp
+++ b/Engine/Source/Runtime/AIModule/Private/EnvironmentQuery/Tests/EnvQueryTest_Trace.cpp
@@ -37,6 +37,7 @@ namespace EnvQueryTest_Trace_Helpers
 		else
 		{
 			[] <bool cond = false>() { static_assert(cond, "Unsupported value of EEnvTraceShape received in TraceShape"); }(); // static_assert must be type-dependent to avoid "ill-formed" code
+			return false;
 		}
 	}
 
