[NavLocalGridManager] fixes crashes when invalid parameters are provided

diff --git a/Engine/Source/Runtime/AIModule/Private/Navigation/NavLocalGridManager.cpp b/Engine/Source/Runtime/AIModule/Private/Navigation/NavLocalGridManager.cpp
--- a/Engine/Source/Runtime/AIModule/Private/Navigation/NavLocalGridManager.cpp
+++ b/Engine/Source/Runtime/AIModule/Private/Navigation/NavLocalGridManager.cpp
@@ -4,6 +4,7 @@
 #include "AISystem.h"
 #include "NavigationSystem.h"
 #include "Engine/Engine.h"
+#include "VisualLogger/VisualLogger.h"
 
 #include UE_INLINE_GENERATED_CPP_BY_NAME(NavLocalGridManager)
 
@@ -323,7 +324,8 @@ int32 UNavLocalGridManager::AddLocalNavigationGridForPoint(UObject* WorldContext
 	int32 GridId = 0;
 
 	UNavLocalGridManager* GridManager = UNavLocalGridManager::GetCurrent(WorldContextObject);
-	if (GridManager)
+	UE_CVLOG_ALWAYS_UELOG(Radius2D <= 0, WorldContextObject, LogNavigation, Warning, TEXT("%hs must be called with a positive non-zero Radius2D but received %d."), __FUNCTION__, Radius2D);
+	if (GridManager && Radius2D > 0)
 	{
 		FNavLocalGridData GridData(Location, UNavLocalGridManager::GridCellSize * Radius2D);
 		GridData.SetHeight(Height);
@@ -340,7 +342,9 @@ int32 UNavLocalGridManager::AddLocalNavigationGridForPoints(UObject* WorldContex
 	int32 GridId = 0;
 
 	UNavLocalGridManager* GridManager = UNavLocalGridManager::GetCurrent(WorldContextObject);
-	if (GridManager)
+	UE_CVLOG_ALWAYS_UELOG(Locations.IsEmpty(), WorldContextObject, LogNavigation, Warning, TEXT("%hs must be called with a non-empty list of locations."), __FUNCTION__);
+	UE_CVLOG_ALWAYS_UELOG(Radius2D <= 0, WorldContextObject, LogNavigation, Warning, TEXT("%hs must be called with a positive non-zero Radius2D but received %d."), __FUNCTION__, Radius2D);
+	if (GridManager && Locations.Num() > 0 && Radius2D > 0)
 	{
 		const FBox Bounds(Locations);
 		const float BoundsSize2D = FloatCastChecked<float>(FMath::Max(Bounds.Max.X - Bounds.Min.X, Bounds.Max.Y - Bounds.Min.Y), UE::LWC::DefaultFloatPrecision);
@@ -364,7 +368,8 @@ int32 UNavLocalGridManager::AddLocalNavigationGridForBox(UObject* WorldContextOb
 	int32 GridId = 0;
 
 	UNavLocalGridManager* GridManager = UNavLocalGridManager::GetCurrent(WorldContextObject);
-	if (GridManager)
+	UE_CVLOG_ALWAYS_UELOG(Radius2D <= 0, WorldContextObject, LogNavigation, Warning, TEXT("%hs must be called with a positive non-zero Radius2D but received %d."), __FUNCTION__, Radius2D);
+	if (GridManager && Radius2D > 0)
 	{
 		FNavLocalGridData GridData(Location, FVector2D(Extent.X + UNavLocalGridManager::GridCellSize * Radius2D, Extent.Y + UNavLocalGridManager::GridCellSize * Radius2D));
 		GridData.SetHeight(FloatCastChecked<float>(Height + Extent.Z, UE::LWC::DefaultFloatPrecision));
@@ -381,7 +386,8 @@ int32 UNavLocalGridManager::AddLocalNavigationGridForCapsule(UObject* WorldConte
 	int32 GridId = 0;
 
 	UNavLocalGridManager* GridManager = UNavLocalGridManager::GetCurrent(WorldContextObject);
-	if (GridManager)
+	UE_CVLOG_ALWAYS_UELOG(Radius2D <= 0, WorldContextObject, LogNavigation, Warning, TEXT("%hs must be called with a positive non-zero Radius2 Dbut received %d."), __FUNCTION__, Radius2D);
+	if (GridManager && Radius2D > 0)
 	{
 		FNavLocalGridData GridData(Location, FVector2D(CapsuleRadius + UNavLocalGridManager::GridCellSize * Radius2D, CapsuleRadius + UNavLocalGridManager::GridCellSize * Radius2D));
 		GridData.SetHeight(Height + CapsuleHalfHeight);
