PR: Fix to make ADetourCrowdAIController inheritable

#11407

diff --git a/Engine/Source/Runtime/AIModule/Classes/DetourCrowdAIController.h b/Engine/Source/Runtime/AIModule/Classes/DetourCrowdAIController.h
--- a/Engine/Source/Runtime/AIModule/Classes/DetourCrowdAIController.h
+++ b/Engine/Source/Runtime/AIModule/Classes/DetourCrowdAIController.h
@@ -9,7 +9,7 @@
 #include "DetourCrowdAIController.generated.h"
 
 UCLASS()
-class ADetourCrowdAIController : public AAIController
+class AIMODULE_API ADetourCrowdAIController : public AAIController
 {
 	GENERATED_BODY()
 public:
