Added a new flag "bStartFromPreviousPath" in FAIMoveRequest and BTTask_MoveTo. This allows the AI to request a new path that will continue from the end of the previous one. This is useful for AIs that need to request the next part of their path before reaching the last point.

diff --git a/Engine/Source/Runtime/AIModule/Classes/AIController.h b/Engine/Source/Runtime/AIModule/Classes/AIController.h
--- a/Engine/Source/Runtime/AIModule/Classes/AIController.h
+++ b/Engine/Source/Runtime/AIModule/Classes/AIController.h
@@ -214,7 +214,10 @@ public:
 	AIMODULE_API virtual void FindPathForMoveRequest(const FAIMoveRequest& MoveRequest, FPathFindingQuery& Query, FNavPathSharedPtr& OutPath) const;
 
 	/** Helper function for creating pathfinding query for this agent from move request data */
-	AIMODULE_API bool BuildPathfindingQuery(const FAIMoveRequest& MoveRequest, FPathFindingQuery& Query) const;
+	AIMODULE_API bool BuildPathfindingQuery(const FAIMoveRequest& MoveRequest, FPathFindingQuery& OutQuery) const;
+
+	/** Helper function for creating pathfinding query for this agent from move request data and starting location */
+	AIMODULE_API bool BuildPathfindingQuery(const FAIMoveRequest& MoveRequest, const FVector& StartLocation, FPathFindingQuery& OutQuery) const;
 
 	UE_DEPRECATED_FORGAME(4.13, "This function is now deprecated, please use FindPathForMoveRequest() for adjusting Query or BuildPathfindingQuery() for getting one.")
 	AIMODULE_API virtual bool PreparePathfinding(const FAIMoveRequest& MoveRequest, FPathFindingQuery& Query);
@@ -269,6 +272,9 @@ public:
 protected:
 	AIMODULE_API virtual void CleanupBrainComponent();
 
+	/** Merges the remaining points of InitialPath, with the points of InOutMergedPath. The resulting merged path is outputted into InOutMergedPath. */
+	AIMODULE_API virtual void MergePaths(const FNavPathSharedPtr& InitialPath, FNavPathSharedPtr& InOutMergedPath) const;
+
 public:
 	/**
 	 * Makes AI use the specified Blackboard asset & creates a Blackboard Component if one does not already exist.
diff --git a/Engine/Source/Runtime/AIModule/Classes/AITypes.h b/Engine/Source/Runtime/AIModule/Classes/AITypes.h
--- a/Engine/Source/Runtime/AIModule/Classes/AITypes.h
+++ b/Engine/Source/Runtime/AIModule/Classes/AITypes.h
@@ -521,6 +521,7 @@ struct FAIMoveRequest
 	FAIMoveRequest& SetAcceptanceRadius(float Radius) { AcceptanceRadius = Radius; return *this; }
 	FAIMoveRequest& SetUserData(const FCustomMoveSharedPtr& InUserData) { UserData = InUserData; return *this; }
 	FAIMoveRequest& SetUserFlags(int32 InUserFlags) { UserFlags = InUserFlags; return *this; }
+	FAIMoveRequest& SetStartFromPreviousPath(const bool bInStartFromPreviousPath) { bStartFromPreviousPath = bInStartFromPreviousPath; return *this; }
 
 	/** the request should be either set up to move to a location, of go to a valid actor */
 	bool IsValid() const { return bInitialized && (!bMoveToActor || GoalActor.IsValid()); }
@@ -548,6 +549,7 @@ struct FAIMoveRequest
 	const FCustomMoveSharedPtr& GetUserData() const { return UserData; }
 	int32 GetUserFlags() const { return UserFlags; }
 
+	AIMODULE_API const bool ShouldStartFromPreviousPath() const { return bStartFromPreviousPath; }
 	AIMODULE_API void SetGoalActor(const AActor* InGoalActor);
 	AIMODULE_API void SetGoalLocation(const FVector& InGoalLocation);
 
@@ -587,6 +589,9 @@ protected:
 	/** pathfinding: goal location will be projected on navigation data before use */
 	uint32 bProjectGoalOnNavigation : 1;
 
+	/** pathfinding: the request will start from the end of the previous path (if any), and the generated path will be merged with the remaining points of the previous path */
+	uint32 bStartFromPreviousPath : 1;
+
 	/** pathfollowing: acceptance radius needs to be increased by agent radius (stop on overlap vs exact point) */
 	uint32 bReachTestIncludesAgentRadius : 1;
 
diff --git a/Engine/Source/Runtime/AIModule/Classes/BehaviorTree/Tasks/BTTask_MoveTo.h b/Engine/Source/Runtime/AIModule/Classes/BehaviorTree/Tasks/BTTask_MoveTo.h
--- a/Engine/Source/Runtime/AIModule/Classes/BehaviorTree/Tasks/BTTask_MoveTo.h
+++ b/Engine/Source/Runtime/AIModule/Classes/BehaviorTree/Tasks/BTTask_MoveTo.h
@@ -90,6 +90,10 @@ class UBTTask_MoveTo : public UBTTask_BlackboardBase
 	UPROPERTY()
 	uint32 bStopOnOverlapNeedsUpdate : 1;
 
+	/** if set, the path request will start from the end of the previous path (if any), and the generated path will be merged with the remaining points of the previous path */
+	UPROPERTY(Category = Node, EditAnywhere)
+	uint32 bStartFromPreviousPath : 1;
+
 	/** if set, move will use pathfinding. Not exposed on purpose, please use BTTask_MoveDirectlyToward */
 	uint32 bUsePathfinding : 1;
 
diff --git a/Engine/Source/Runtime/AIModule/Private/AIController.cpp b/Engine/Source/Runtime/AIModule/Private/AIController.cpp
--- a/Engine/Source/Runtime/AIModule/Private/AIController.cpp
+++ b/Engine/Source/Runtime/AIModule/Private/AIController.cpp
@@ -604,6 +604,37 @@ EPathFollowingRequestResult::Type AAIController::MoveToLocation(const FVector& D
 	return MoveTo(MoveReq);
 }
 
+void AAIController::MergePaths(const FNavPathSharedPtr& InitialPath, FNavPathSharedPtr& InOutMergedPath) const
+{
+	if (!InitialPath.IsValid() || !InitialPath->IsValid())
+	{
+		UE_VLOG_UELOG(this, LogAINavigation, Error, TEXT("%hs: InitialPath is Invalid"), __FUNCTION__);
+		return;
+	}
+
+	if (!InOutMergedPath.IsValid() || !InOutMergedPath->IsValid())
+	{
+		UE_VLOG_UELOG(this, LogAINavigation, Error, TEXT("%hs: InOutMergedPath is Invalid"), __FUNCTION__);
+		return;
+	}
+
+	const TArray<FNavPathPoint>& InitialPathPoints = InitialPath->GetPathPoints();
+	TArray<FNavPathPoint>& InOutPathPoints = InOutMergedPath->GetPathPoints();
+
+	if (!InitialPathPoints.Last().Location.Equals(InOutPathPoints[0].Location))
+	{
+		UE_VLOG_UELOG(this, LogAINavigation, Error, TEXT("%hs: last %s and first %s points don't match."), __FUNCTION__, *InitialPathPoints.Last().Location.ToString(), *InOutPathPoints[0].Location.ToString());
+		return;
+	}
+
+	// We don't want to keep path points that have already been traversed, so only merge the points starting from "CurrentPathIndex".
+	const int32 StartingPointIndex = PathFollowingComponent ? PathFollowingComponent->GetCurrentPathIndex() : 0;
+	if (StartingPointIndex < InitialPathPoints.Num())
+	{
+		InOutPathPoints.Insert(&InitialPathPoints[StartingPointIndex], InitialPathPoints.Num() - StartingPointIndex - 1, 0);
+	}
+}
+
 FPathFollowingRequestResult AAIController::MoveTo(const FAIMoveRequest& MoveRequest, FNavPathSharedPtr* OutPath)
 {
 	// both MoveToActor and MoveToLocation can be called from blueprints/script and should keep only single movement request at the same time.
@@ -679,12 +710,30 @@ FPathFollowingRequestResult AAIController::MoveTo(const FAIMoveRequest& MoveRequ
 	{
 		FPathFindingQuery PFQuery;
 
-		const bool bValidQuery = BuildPathfindingQuery(MoveRequest, PFQuery);
+		bool bShouldMergePaths = false;
+		FVector StartLocation = GetNavAgentLocation();
+		if (MoveRequest.ShouldStartFromPreviousPath())
+		{
+			FNavPathSharedPtr CurrentPath = PathFollowingComponent->GetPath();
+			if (CurrentPath.IsValid() && CurrentPath->IsValid() && CurrentPath->GetPathPoints().Num() > 0)
+			{
+				StartLocation = CurrentPath->GetPathPoints().Last();
+				bShouldMergePaths = true;
+			}
+		}
+
+		const bool bValidQuery = BuildPathfindingQuery(MoveRequest, StartLocation, PFQuery);
 		if (bValidQuery)
 		{
 			FNavPathSharedPtr Path;
 			FindPathForMoveRequest(MoveRequest, PFQuery, Path);
 
+			if (bShouldMergePaths && Path.IsValid())
+			{
+				// Merge the newly generated path with the current one
+				MergePaths(PathFollowingComponent->GetPath(), Path);
+			}
+
 			const FAIRequestID RequestID = Path.IsValid() ? RequestMove(MoveRequest, Path) : FAIRequestID::InvalidRequest;
 			if (RequestID.IsValid())
 			{
@@ -758,7 +807,12 @@ bool AAIController::ShouldPostponePathUpdates() const
 	return Super::ShouldPostponePathUpdates();
 }
 
-bool AAIController::BuildPathfindingQuery(const FAIMoveRequest& MoveRequest, FPathFindingQuery& Query) const
+bool AAIController::BuildPathfindingQuery(const FAIMoveRequest& MoveRequest, FPathFindingQuery& OutQuery) const
+{
+	return BuildPathfindingQuery(MoveRequest, GetNavAgentLocation(), OutQuery);
+}
+
+bool AAIController::BuildPathfindingQuery(const FAIMoveRequest& MoveRequest, const FVector& StartLocation, FPathFindingQuery& OutQuery) const
 {
 	bool bResult = false;
 
@@ -785,18 +839,18 @@ bool AAIController::BuildPathfindingQuery(const FAIMoveRequest& MoveRequest, FPa
 		}
 
 		FSharedConstNavQueryFilter NavFilter = UNavigationQueryFilter::GetQueryFilter(*NavData, this, MoveRequest.GetNavigationFilter());
-		Query = FPathFindingQuery(*this, *NavData, GetNavAgentLocation(), GoalLocation, NavFilter);
-		Query.SetAllowPartialPaths(MoveRequest.IsUsingPartialPaths());
-		Query.SetRequireNavigableEndLocation(MoveRequest.IsNavigableEndLocationRequired());
+		OutQuery = FPathFindingQuery(*this, *NavData, StartLocation, GoalLocation, NavFilter);
+		OutQuery.SetAllowPartialPaths(MoveRequest.IsUsingPartialPaths());
+		OutQuery.SetRequireNavigableEndLocation(MoveRequest.IsNavigableEndLocationRequired());
 		if (MoveRequest.IsApplyingCostLimitFromHeuristic())
 		{
 			const float HeuristicScale = NavFilter->GetHeuristicScale();
-			Query.CostLimit = FPathFindingQuery::ComputeCostLimitFromHeuristic(Query.StartLocation, Query.EndLocation, HeuristicScale, MoveRequest.GetCostLimitFactor(), MoveRequest.GetMinimumCostLimit()); 
+			OutQuery.CostLimit = FPathFindingQuery::ComputeCostLimitFromHeuristic(OutQuery.StartLocation, OutQuery.EndLocation, HeuristicScale, MoveRequest.GetCostLimitFactor(), MoveRequest.GetMinimumCostLimit()); 
 		}
 
 		if (PathFollowingComponent)
 		{
-			PathFollowingComponent->OnPathfindingQuery(Query);
+			PathFollowingComponent->OnPathfindingQuery(OutQuery);
 		}
 
 		bResult = true;
diff --git a/Engine/Source/Runtime/AIModule/Private/AITypes.cpp b/Engine/Source/Runtime/AIModule/Private/AITypes.cpp
--- a/Engine/Source/Runtime/AIModule/Private/AITypes.cpp
+++ b/Engine/Source/Runtime/AIModule/Private/AITypes.cpp
@@ -163,7 +163,7 @@ FAIMoveRequest::FAIMoveRequest() :
 	GoalActor(nullptr), GoalLocation(FAISystem::InvalidLocation), FilterClass(nullptr),
 	bInitialized(false), bMoveToActor(false),
 	bUsePathfinding(true), bAllowPartialPath(true), bRequireNavigableEndLocation(true), bApplyCostLimitFromHeuristic(false), bProjectGoalOnNavigation(true),
-	bReachTestIncludesAgentRadius(true), bReachTestIncludesGoalRadius(true), bCanStrafe(false),
+	bStartFromPreviousPath(false), bReachTestIncludesAgentRadius(true), bReachTestIncludesGoalRadius(true), bCanStrafe(false),
 	AcceptanceRadius(UPathFollowingComponent::DefaultAcceptanceRadius), CostLimitFactor(FLT_MAX), MinimumCostLimit(0.f),
 	UserFlags(0)
 {
@@ -173,7 +173,7 @@ FAIMoveRequest::FAIMoveRequest(const AActor* InGoalActor) :
 	GoalActor(const_cast<AActor*>(InGoalActor)), GoalLocation(FAISystem::InvalidLocation), FilterClass(nullptr),
 	bInitialized(true), bMoveToActor(true),
 	bUsePathfinding(true), bAllowPartialPath(true), bRequireNavigableEndLocation(true), bApplyCostLimitFromHeuristic(false), bProjectGoalOnNavigation(true),
-	bReachTestIncludesAgentRadius(true), bReachTestIncludesGoalRadius(true), bCanStrafe(false),
+	bStartFromPreviousPath(false), bReachTestIncludesAgentRadius(true), bReachTestIncludesGoalRadius(true), bCanStrafe(false),
 	AcceptanceRadius(UPathFollowingComponent::DefaultAcceptanceRadius), CostLimitFactor(FLT_MAX), MinimumCostLimit(0.f),
 	UserFlags(0)
 {
@@ -183,7 +183,7 @@ FAIMoveRequest::FAIMoveRequest(const FVector& InGoalLocation) :
 	GoalActor(nullptr), GoalLocation(InGoalLocation), FilterClass(nullptr),
 	bInitialized(true), bMoveToActor(false),
 	bUsePathfinding(true), bAllowPartialPath(true), bRequireNavigableEndLocation(true), bApplyCostLimitFromHeuristic(false), bProjectGoalOnNavigation(true),
-	bReachTestIncludesAgentRadius(true), bReachTestIncludesGoalRadius(true), bCanStrafe(false),
+	bStartFromPreviousPath(false), bReachTestIncludesAgentRadius(true), bReachTestIncludesGoalRadius(true), bCanStrafe(false),
 	AcceptanceRadius(UPathFollowingComponent::DefaultAcceptanceRadius), CostLimitFactor(FLT_MAX), MinimumCostLimit(0.f),
 	UserFlags(0)
 {
@@ -226,11 +226,12 @@ bool FAIMoveRequest::UpdateGoalLocation(const FVector& NewLocation) const
 
 FString FAIMoveRequest::ToString() const
 {
-	return FString::Printf(TEXT("%s(%s) Mode(%s) Filter(%s) AcceptanceRadius(%.1f%s)"),
+	return FString::Printf(TEXT("%s(%s) Mode(%s) Filter(%s) AcceptanceRadius(%.1f%s) StartsFromPreviousPath(%s))"),
 		bMoveToActor ? TEXT("Actor") : TEXT("Location"), bMoveToActor ? *GetNameSafe(GoalActor.Get()) : *GoalLocation.ToString(),
 		bUsePathfinding ? (bAllowPartialPath ? TEXT("partial path") : TEXT("complete path")) : TEXT("direct"),
 		*GetNameSafe(FilterClass),
-		AcceptanceRadius, bReachTestIncludesAgentRadius || (bMoveToActor && bReachTestIncludesGoalRadius) ? TEXT(" + overlap") : TEXT("")
+		AcceptanceRadius, bReachTestIncludesAgentRadius || (bMoveToActor && bReachTestIncludesGoalRadius) ? TEXT(" + overlap") : TEXT(""),
+		bStartFromPreviousPath ? TEXT("Yes") : TEXT("No")
 		);
 }
 
diff --git a/Engine/Source/Runtime/AIModule/Private/BehaviorTree/Tasks/BTTask_MoveTo.cpp b/Engine/Source/Runtime/AIModule/Private/BehaviorTree/Tasks/BTTask_MoveTo.cpp
--- a/Engine/Source/Runtime/AIModule/Private/BehaviorTree/Tasks/BTTask_MoveTo.cpp
+++ b/Engine/Source/Runtime/AIModule/Private/BehaviorTree/Tasks/BTTask_MoveTo.cpp
@@ -27,8 +27,8 @@ UBTTask_MoveTo::UBTTask_MoveTo(const FObjectInitializer& ObjectInitializer) : Su
 	bRequireNavigableEndLocation = true;
 	bProjectGoalLocation = true;
 	bUsePathfinding = true;
-
 	bStopOnOverlapNeedsUpdate = true;
+	bStartFromPreviousPath = false;
 
 	ObservedBlackboardValueTolerance = AcceptableRadius * 0.95f;
 
@@ -92,6 +92,7 @@ EBTNodeResult::Type UBTTask_MoveTo::PerformMoveTask(UBehaviorTreeComponent& Owne
 		MoveReq.SetRequireNavigableEndLocation(bRequireNavigableEndLocation);
 		MoveReq.SetProjectGoalLocation(bProjectGoalLocation);
 		MoveReq.SetUsePathfinding(bUsePathfinding);
+		MoveReq.SetStartFromPreviousPath(bStartFromPreviousPath);
 
 		if (BlackboardKey.SelectedKeyType == UBlackboardKeyType_Object::StaticClass())
 		{
