Add more information inside the BTDecorator_TimeLimit ensure to help pinpoint the issue.
#rnx

diff --git a/Engine/Source/Runtime/AIModule/Private/BehaviorTree/Decorators/BTDecorator_TimeLimit.cpp b/Engine/Source/Runtime/AIModule/Private/BehaviorTree/Decorators/BTDecorator_TimeLimit.cpp
--- a/Engine/Source/Runtime/AIModule/Private/BehaviorTree/Decorators/BTDecorator_TimeLimit.cpp
+++ b/Engine/Source/Runtime/AIModule/Private/BehaviorTree/Decorators/BTDecorator_TimeLimit.cpp
@@ -1,6 +1,7 @@
 // Copyright Epic Games, Inc. All Rights Reserved.
 
 #include "BehaviorTree/Decorators/BTDecorator_TimeLimit.h"
+#include "BehaviorTree/BehaviorTree.h"
 
 #include UE_INLINE_GENERATED_CPP_BY_NAME(BTDecorator_TimeLimit)
 
@@ -9,7 +10,8 @@ struct FBTimeLimitMemory
 	bool bElapsed = false;
 };
 
-UBTDecorator_TimeLimit::UBTDecorator_TimeLimit(const FObjectInitializer& ObjectInitializer) : Super(ObjectInitializer)
+UBTDecorator_TimeLimit::UBTDecorator_TimeLimit(const FObjectInitializer& ObjectInitializer)
+	: Super(ObjectInitializer)
 {
 	NodeName = "TimeLimit";
 	TimeLimit = 5.0f;
@@ -58,18 +60,28 @@ bool UBTDecorator_TimeLimit::CalculateRawConditionValue(UBehaviorTreeComponent&
 
 void UBTDecorator_TimeLimit::TickNode(UBehaviorTreeComponent& OwnerComp, uint8* NodeMemory, const float DeltaSeconds)
 {
-	ensureMsgf(DeltaSeconds >= TimeLimit || FMath::IsNearlyEqual(DeltaSeconds, TimeLimit, UE_KINDA_SMALL_NUMBER),
-		TEXT("Using SetNextTickTime in OnBecomeRelevant should guarantee that we are only getting ticked when the time limit is finished. DT=%f, TimeLimit=%f"),
-		DeltaSeconds,
-		TimeLimit);
-
+	FBTimeLimitMemory* TimeLimitMemory = CastInstanceNodeMemory<FBTimeLimitMemory>(NodeMemory);
+	if (DeltaSeconds < TimeLimit && !FMath::IsNearlyEqual(DeltaSeconds, TimeLimit, UE_KINDA_SMALL_NUMBER))
+	{
+		const FBTAuxiliaryMemory* AuxMemory = GetSpecialNodeMemory<FBTAuxiliaryMemory>(NodeMemory);
+		ensureMsgf(false,
+			TEXT("Using SetNextTickTime in OnBecomeRelevant should guarantee that we are only getting ticked when the time limit is finished. DT=%f, TimeLimit=%f Node:%d, BT: %s, NextTickRemainingTime=%f, AccumulatedDeltaTime=%f, bElapsed=%s"),
+			DeltaSeconds,
+			TimeLimit,
+			static_cast<int32>(GetExecutionIndex()),
+			*GetNameSafe(GetTreeAsset()),
+			AuxMemory->NextTickRemainingTime,
+			AuxMemory->AccumulatedDeltaTime,
+			TimeLimitMemory->bElapsed ? TEXT("True") : TEXT("False"));
+	}
+	
 	// Mark this decorator instance as Elapsed for calls to CalculateRawConditionValue
-	reinterpret_cast<FBTimeLimitMemory*>(NodeMemory)->bElapsed = true;
+	TimeLimitMemory->bElapsed = true;
 
 	// Set our next tick time to large value so we don't get ticked again in case the decorator
 	// is still active after requesting execution (e.g. latent abort)
 	SetNextTickTime(NodeMemory, FLT_MAX);
-	
+
 	OwnerComp.RequestExecution(this);
 }
 
@@ -101,5 +113,4 @@ FName UBTDecorator_TimeLimit::GetNodeIconName() const
 	return FName("BTEditor.Graph.BTNode.Decorator.TimeLimit.Icon");
 }
 
-#endif	// WITH_EDITOR
-
+#endif // WITH_EDITOR
