Add memory lifecycle methods for node to module for export

diff --git a/Engine/Source/Runtime/AIModule/Classes/BehaviorTree/BTCompositeNode.h b/Engine/Source/Runtime/AIModule/Classes/BehaviorTree/BTCompositeNode.h
--- a/Engine/Source/Runtime/AIModule/Classes/BehaviorTree/BTCompositeNode.h
+++ b/Engine/Source/Runtime/AIModule/Classes/BehaviorTree/BTCompositeNode.h
@@ -144,8 +144,8 @@ class UBTCompositeNode : public UBTNode
 
 	/** size of instance memory */
 	AIMODULE_API virtual uint16 GetInstanceMemorySize() const override;
-	virtual void InitializeMemory(UBehaviorTreeComponent& OwnerComp, uint8* NodeMemory, EBTMemoryInit::Type InitType) const override;
-	virtual void CleanupMemory(UBehaviorTreeComponent& OwnerComp, uint8* NodeMemory, EBTMemoryClear::Type CleanupType) const override;
+	AIMODULE_API virtual void InitializeMemory(UBehaviorTreeComponent& OwnerComp, uint8* NodeMemory, EBTMemoryInit::Type InitType) const override;
+	AIMODULE_API virtual void CleanupMemory(UBehaviorTreeComponent& OwnerComp, uint8* NodeMemory, EBTMemoryClear::Type CleanupType) const override;
 
 	/** @return child node at given index */
 	AIMODULE_API UBTNode* GetChildNode(int32 Index) const;
