Mass EQS Integration

diff --git a/Engine/Plugins/Runtime/MassGameplay/MassGameplay.uplugin b/Engine/Plugins/Runtime/MassGameplay/MassGameplay.uplugin
--- a/Engine/Plugins/Runtime/MassGameplay/MassGameplay.uplugin
+++ b/Engine/Plugins/Runtime/MassGameplay/MassGameplay.uplugin
@@ -24,6 +24,10 @@
 			"Name": "MassActors",
 			"Type": "Runtime"
 		},
+		{
+			"Name": "MassEQS",
+			"Type": "Runtime"
+		},
 		{
 			"Name": "MassSignals",
 			"Type": "Runtime"
diff --git a/Engine/Plugins/Runtime/MassGameplay/Source/MassEQS/MassEQS.Build.cs b/Engine/Plugins/Runtime/MassGameplay/Source/MassEQS/MassEQS.Build.cs
new file mode 100644
--- /dev/null
+++ b/Engine/Plugins/Runtime/MassGameplay/Source/MassEQS/MassEQS.Build.cs
@@ -0,0 +1,32 @@
+// Copyright Epic Games, Inc. All Rights Reserved.
+
+using UnrealBuildTool;
+
+namespace UnrealBuildTool.Rules
+{
+	public class MassEQS : ModuleRules
+	{
+		public MassEQS(ReadOnlyTargetRules Target) : base(Target)
+		{
+			PCHUsage = ModuleRules.PCHUsageMode.UseExplicitOrSharedPCHs;
+
+			PublicIncludePaths.AddRange(
+				new string[] {
+				}
+			);
+
+			PublicDependencyModuleNames.AddRange(
+				new string[] {
+					"Core",
+					"CoreUObject",
+					"Engine",
+					"AIModule",
+					"MassCommon",
+					"MassEntity",
+					"MassSignals",
+					"StructUtils"
+				}
+			);
+		}
+	}
+}
\ No newline at end of file
diff --git a/Engine/Plugins/Runtime/MassGameplay/Source/MassEQS/Private/Generators/MassEnvQueryGenerator.cpp b/Engine/Plugins/Runtime/MassGameplay/Source/MassEQS/Private/Generators/MassEnvQueryGenerator.cpp
new file mode 100644
--- /dev/null
+++ b/Engine/Plugins/Runtime/MassGameplay/Source/MassEQS/Private/Generators/MassEnvQueryGenerator.cpp
@@ -0,0 +1,15 @@
+// Copyright Epic Games, Inc. All Rights Reserved.
+
+#include "Generators/MassEnvQueryGenerator.h"
+#include "MassEQSSubsystem.h"
+
+UMassEnvQueryGenerator::UMassEnvQueryGenerator(const FObjectInitializer& ObjectInitializer)
+	: Super(ObjectInitializer)
+{
+	bCanRunAsync = true;
+}
+
+void UMassEnvQueryGenerator::GenerateItems(FEnvQueryInstance& QueryInstance) const
+{
+	MassEQSRequestHandler.SendOrRecieveRequest(QueryInstance, *this);
+}
diff --git a/Engine/Plugins/Runtime/MassGameplay/Source/MassEQS/Private/Generators/MassEnvQueryGenerator_MassEntityHandles.cpp b/Engine/Plugins/Runtime/MassGameplay/Source/MassEQS/Private/Generators/MassEnvQueryGenerator_MassEntityHandles.cpp
new file mode 100644
--- /dev/null
+++ b/Engine/Plugins/Runtime/MassGameplay/Source/MassEQS/Private/Generators/MassEnvQueryGenerator_MassEntityHandles.cpp
@@ -0,0 +1,41 @@
+// Copyright Epic Games, Inc. All Rights Reserved.
+
+#include "Generators/MassEnvQueryGenerator_MassEntityHandles.h"
+
+#include "EnvironmentQuery/Contexts/EnvQueryContext_Querier.h"
+#include "Items/EnvQueryItemType_MassEntityHandle.h"
+#include "MassEQSSubsystem.h"
+#include "MassEQSUtils.h"
+
+#include "MassEntityTypes.h"
+
+UMassEnvQueryGenerator_MassEntityHandles::UMassEnvQueryGenerator_MassEntityHandles(const FObjectInitializer& ObjectInitializer)
+	: Super(ObjectInitializer)
+{
+	ItemType = UEnvQueryItemType_MassEntityHandle::StaticClass();
+
+	SearchRadius.DefaultValue = 500.0f;
+	SearchCenter = UEnvQueryContext_Querier::StaticClass();
+}
+
+bool UMassEnvQueryGenerator_MassEntityHandles::TryAcquireResults(FEnvQueryInstance& QueryInstance) const
+{
+	check(MassEQSRequestHandler.MassEQSSubsystem);
+
+	TUniquePtr<FMassEQSRequestData> RawRequestData = MassEQSRequestHandler.MassEQSSubsystem->TryAcquireResults(MassEQSRequestHandler.RequestHandle);
+	if (FMassEnvQueryResultData_MassEntityHandles* RequestData = FMassEQSUtils::TryAndEnsureCast<FMassEnvQueryResultData_MassEntityHandles>(RawRequestData))
+	{
+		QueryInstance.AddItemData<UEnvQueryItemType_MassEntityHandle>(RequestData->GeneratedEntityInfo);
+		return true;
+	}
+
+	return false;
+}
+
+TUniquePtr<FMassEQSRequestData> UMassEnvQueryGenerator_MassEntityHandles::GetRequestData(FEnvQueryInstance& QueryInstance) const
+{
+	TArray<FVector> PreparedContextPositions = {};
+	QueryInstance.PrepareContext(SearchCenter, PreparedContextPositions);
+
+	return MakeUnique<FMassEQSRequestData_MassEntityHandles>(MoveTemp(PreparedContextPositions), SearchRadius.GetValue());
+}
\ No newline at end of file
diff --git a/Engine/Plugins/Runtime/MassGameplay/Source/MassEQS/Private/Items/EnvQueryItemType_MassEntityHandle.cpp b/Engine/Plugins/Runtime/MassGameplay/Source/MassEQS/Private/Items/EnvQueryItemType_MassEntityHandle.cpp
new file mode 100644
--- /dev/null
+++ b/Engine/Plugins/Runtime/MassGameplay/Source/MassEQS/Private/Items/EnvQueryItemType_MassEntityHandle.cpp
@@ -0,0 +1,34 @@
+// Copyright Epic Games, Inc. All Rights Reserved.
+
+#include "Items/EnvQueryItemType_MassEntityHandle.h"
+
+#include "EnvironmentQuery/EnvQueryTypes.h"
+
+#include "MassEntitySubsystem.h"
+
+UEnvQueryItemType_MassEntityHandle::UEnvQueryItemType_MassEntityHandle()
+{
+	ValueSize = sizeof(FMassEnvQueryEntityInfo);
+}
+
+const FMassEnvQueryEntityInfo& UEnvQueryItemType_MassEntityHandle::GetValue(const uint8* RawData)
+{
+	return GetValueFromMemory<FMassEnvQueryEntityInfo>(RawData);
+}
+
+void UEnvQueryItemType_MassEntityHandle::SetValue(uint8* RawData, const FMassEnvQueryEntityInfo& Value)
+{
+	return SetValueInMemory<FMassEnvQueryEntityInfo>(RawData, Value);
+}
+
+FVector UEnvQueryItemType_MassEntityHandle::GetItemLocation(const uint8* RawData) const
+{
+	const FMassEnvQueryEntityInfo& EntityInfo = GetValue(RawData);	
+	return EntityInfo.CachedTransform.GetLocation();
+}
+
+FRotator UEnvQueryItemType_MassEntityHandle::GetItemRotation(const uint8* RawData) const
+{
+	const FMassEnvQueryEntityInfo& EntityInfo = GetValue(RawData);
+	return EntityInfo.CachedTransform.GetRotation().Rotator();
+}
diff --git a/Engine/Plugins/Runtime/MassGameplay/Source/MassEQS/Private/MassEQS.cpp b/Engine/Plugins/Runtime/MassGameplay/Source/MassEQS/Private/MassEQS.cpp
new file mode 100644
--- /dev/null
+++ b/Engine/Plugins/Runtime/MassGameplay/Source/MassEQS/Private/MassEQS.cpp
@@ -0,0 +1,7 @@
+// Copyright Epic Games, Inc. All Rights Reserved.
+
+#include "MassEQS.h"
+
+DEFINE_LOG_CATEGORY(LogMassEQS);
+
+IMPLEMENT_GAME_MODULE(FMassEQSModule, MassEQS);
diff --git a/Engine/Plugins/Runtime/MassGameplay/Source/MassEQS/Private/MassEQSBlueprintLibrary.cpp b/Engine/Plugins/Runtime/MassGameplay/Source/MassEQS/Private/MassEQSBlueprintLibrary.cpp
new file mode 100644
--- /dev/null
+++ b/Engine/Plugins/Runtime/MassGameplay/Source/MassEQS/Private/MassEQSBlueprintLibrary.cpp
@@ -0,0 +1,136 @@
+// Copyright Epic Games, Inc. All Rights Reserved.
+
+#include "MassEQSBlueprintLibrary.h"
+
+#include "MassCommonFragments.h"
+#include "MassSignalSubsystem.h"
+#include "MassEntitySubsystem.h"
+
+#include "Items/EnvQueryItemType_MassEntityHandle.h"
+#include "EnvironmentQuery/EnvQueryInstanceBlueprintWrapper.h"
+
+UMassEQSBlueprintLibrary::UMassEQSBlueprintLibrary(const FObjectInitializer& ObjectInitializer)
+	: UBlueprintFunctionLibrary(ObjectInitializer)
+{
+}
+
+//----------------------------------------------------------------------//
+// Commands
+//----------------------------------------------------------------------//
+
+void UMassEQSBlueprintLibrary::SendSignalToEntity(const AActor* Owner, const FMassEnvQueryEntityInfoBlueprintWrapper& EntityInfo, const FName Signal)
+{
+	if (!ensureMsgf(Owner, TEXT("Must supply owner to function. Reference to self should suffice.")))
+	{
+		return;
+	}
+
+	const UWorld* World = Owner->GetWorld();
+	check(World);
+
+	UMassEntitySubsystem* EntitySubsystem = World->GetSubsystem<UMassEntitySubsystem>();
+	check(EntitySubsystem);
+
+	UMassSignalSubsystem* SignalSubsystem = World->GetSubsystem<UMassSignalSubsystem>();
+	check(SignalSubsystem);
+
+	FMassEntityManager& EntityManager = EntitySubsystem->GetMutableEntityManager();
+
+	FMassEntityHandle EntityHandle = EntityInfo.GetEntityHandle();
+	if (!EntityManager.IsEntityValid(EntityHandle))
+	{
+		return;
+	}
+
+	SignalSubsystem->SignalEntity(Signal, EntityHandle);
+}
+
+//----------------------------------------------------------------------//
+// Utils
+//----------------------------------------------------------------------//
+
+FVector UMassEQSBlueprintLibrary::GetCurrentEntityPosition(const AActor* Owner, const FMassEnvQueryEntityInfoBlueprintWrapper& EntityInfo)
+{
+	if (!ensureMsgf(Owner, TEXT("Must supply owner to function. Reference to self should suffice.")))
+	{
+		return FVector::ZeroVector;
+	}
+
+	const UWorld* World = Owner->GetWorld();
+	check(World);
+
+	UMassEntitySubsystem* EntitySubsystem = World->GetSubsystem<UMassEntitySubsystem>();
+	check(EntitySubsystem);
+
+	FMassEntityManager& EntityManager = EntitySubsystem->GetMutableEntityManager();
+
+	const FTransformFragment* TransformFragment = EntityManager.GetFragmentDataPtr<FTransformFragment>(EntityInfo.GetEntityHandle());
+	if (TransformFragment)
+	{
+		return TransformFragment->GetTransform().GetLocation();
+	}
+
+	return FVector::ZeroVector;
+}
+
+bool UMassEQSBlueprintLibrary::ContainsEntity(const TArray<FMassEnvQueryEntityInfoBlueprintWrapper>& EntityList, const FMassEnvQueryEntityInfoBlueprintWrapper& EntityInfo)
+{
+	for (const FMassEnvQueryEntityInfoBlueprintWrapper& Entity : EntityList)
+	{
+		if (Entity == EntityInfo)
+		{
+			return true;
+		}
+	}
+
+	return false;
+}
+
+TArray<FMassEnvQueryEntityInfoBlueprintWrapper> UMassEQSBlueprintLibrary::GetEnviromentQueryResultAsEntityInfo(const UEnvQueryInstanceBlueprintWrapper* QueryInstance)
+{
+	check(QueryInstance);
+
+	TArray<FMassEnvQueryEntityInfoBlueprintWrapper> ResultInfo = {};
+
+	const FEnvQueryResult* QueryResult = QueryInstance->GetQueryResult();
+	const TSubclassOf<UEnvQueryItemType> ItemType = QueryResult->ItemType;
+	const EEnvQueryRunMode::Type RunMode = QueryInstance->GetRunMode();
+
+	if (QueryResult && (QueryResult->GetRawStatus() == EEnvQueryStatus::Success) && (ItemType.Get() != nullptr) && ItemType->IsChildOf(UEnvQueryItemType_MassEntityHandle::StaticClass()))
+	{
+		if (RunMode != EEnvQueryRunMode::AllMatching)
+		{
+			ResultInfo.Add(GetItemAsEntityInfoBPWrapper(QueryResult, 0));
+		}
+		else
+		{
+			GetAllAsEntityInfoBPWrappers(QueryResult, ResultInfo);
+		}
+	}
+
+	return ResultInfo;
+}
+
+FMassEnvQueryEntityInfoBlueprintWrapper UMassEQSBlueprintLibrary::GetItemAsEntityInfoBPWrapper(const FEnvQueryResult* QueryResult, int32 Index)
+{
+	UEnvQueryItemType_MassEntityHandle* DefTypeOb = QueryResult->ItemType->GetDefaultObject<UEnvQueryItemType_MassEntityHandle>();
+
+	check(DefTypeOb != nullptr);
+	FMassEnvQueryEntityInfo EntityInfo = DefTypeOb->GetValue(QueryResult->RawData.GetData() + QueryResult->Items[Index].DataOffset);
+
+	return FMassEnvQueryEntityInfoBlueprintWrapper(EntityInfo);
+}
+
+void UMassEQSBlueprintLibrary::GetAllAsEntityInfoBPWrappers(const FEnvQueryResult* QueryResult, TArray<FMassEnvQueryEntityInfoBlueprintWrapper>& OutInfo)
+{
+	UEnvQueryItemType_MassEntityHandle* DefTypeOb = QueryResult->ItemType->GetDefaultObject<UEnvQueryItemType_MassEntityHandle>();
+	check(DefTypeOb != nullptr);
+
+	OutInfo.Reserve(OutInfo.Num() + QueryResult->Items.Num());
+
+	for (const FEnvQueryItem& Item : QueryResult->Items)
+	{
+		FMassEnvQueryEntityInfo EntityInfo = DefTypeOb->GetValue(QueryResult->RawData.GetData() + Item.DataOffset);
+		OutInfo.Add(FMassEnvQueryEntityInfoBlueprintWrapper(EntityInfo));
+	}
+}
\ No newline at end of file
diff --git a/Engine/Plugins/Runtime/MassGameplay/Source/MassEQS/Private/MassEQSSubsystem.cpp b/Engine/Plugins/Runtime/MassGameplay/Source/MassEQS/Private/MassEQSSubsystem.cpp
new file mode 100644
--- /dev/null
+++ b/Engine/Plugins/Runtime/MassGameplay/Source/MassEQS/Private/MassEQSSubsystem.cpp
@@ -0,0 +1,146 @@
+// Copyright Epic Games, Inc. All Rights Reserved.
+
+#include "MassEQSSubsystem.h"
+
+#include "MassCommandBuffer.h"
+#include "MassEntitySubsystem.h"
+#include "MassEQS.h"
+
+#include "Items/EnvQueryItemType_MassEntityHandle.h"
+#include "Generators/MassEnvQueryGenerator.h"
+#include "Tests/MassEnvQueryTest.h"
+#include "MassEQSUtils.h"
+
+//----------------------------------------------------------------------//
+// RequestQueue
+//----------------------------------------------------------------------//
+
+FMassEQSRequestHandle UMassEQSSubsystem::PushRequest(const FEnvQueryInstance& QueryInstance, int32 RequestQueueIndex, TUniquePtr<FMassEQSRequestData>&& Request)
+{
+	UE_MT_SCOPED_WRITE_ACCESS(RequestAccessDetector);
+
+	check(RequestQueueIndex >= 0);
+
+	FMassEQSRequestHandle RequestHandle = HandleManager.GetNextHandle();
+	Request.Get()->RequestHandle = RequestHandle;
+
+	// Fill up the Requests EntityHandle List with all Entities previously generated. 
+	// Won't do anything for Generators.
+	FMassEQSUtils::GetAllAsEntityHandles(QueryInstance, Request->EntityHandles);
+
+	RequestQueues[RequestQueueIndex].Enqueue(MoveTemp(Request));
+
+	return RequestHandle;
+}
+
+
+TUniquePtr<FMassEQSRequestData> UMassEQSSubsystem::PopRequest(int32 RequestQueueIndex)
+{
+	UE_MT_SCOPED_WRITE_ACCESS(RequestAccessDetector);
+
+	check(RequestQueueIndex >= 0);
+
+	TUniquePtr<FMassEQSRequestData> OutRequest = nullptr;
+	RequestQueues[RequestQueueIndex].Dequeue(OutRequest);
+
+	return OutRequest;
+}
+
+//----------------------------------------------------------------------//
+// Results
+//----------------------------------------------------------------------//
+
+void UMassEQSSubsystem::SubmitResults(FMassEQSRequestHandle RequestHandle, TUniquePtr<FMassEQSRequestData>&& Result)
+{
+	UE_MT_SCOPED_WRITE_ACCESS(ResultAccessDetector);
+
+	check(RequestHandle.IsValid());
+
+	int32 CancelledRequestIndex = CancelledRequests.Find(RequestHandle);
+	if (CancelledRequestIndex >= 0)
+	{
+		CancelledRequests.RemoveAt(CancelledRequestIndex);
+		FinalizeRequest(RequestHandle);
+		return;
+	}
+	
+	Result->RequestHandle = RequestHandle;
+	Results.Add(RequestHandle, MoveTemp(Result));
+}
+
+TUniquePtr<FMassEQSRequestData> UMassEQSSubsystem::TryAcquireResults(FMassEQSRequestHandle RequestHandle)
+{
+	UE_MT_SCOPED_WRITE_ACCESS(ResultAccessDetector);
+
+	// Should never try to Acquire invalid RequestHandle
+	check(RequestHandle.IsValid());
+
+	TUniquePtr<FMassEQSRequestData> Result = nullptr;
+	if (Results.RemoveAndCopyValue(RequestHandle, Result))
+	{
+		FinalizeRequest(RequestHandle);
+	}
+	
+	return Result;
+}
+
+//----------------------------------------------------------------------//
+// Misc
+//----------------------------------------------------------------------//
+
+void UMassEQSSubsystem::FinalizeRequest(FMassEQSRequestHandle RequestHandle)
+{
+	HandleManager.RemoveHandle(RequestHandle);
+}
+
+void UMassEQSSubsystem::Initialize(FSubsystemCollectionBase& Collection)
+{
+	Super::Initialize(Collection);
+	bIsInitialized = true;
+}
+
+void UMassEQSSubsystem::Deinitialize()
+{
+	UE_MT_SCOPED_WRITE_ACCESS(RequestAccessDetector);
+	UE_MT_SCOPED_WRITE_ACCESS(ResultAccessDetector);
+
+	Super::Deinitialize();
+
+	Results.Empty();
+	RequestQueues.Empty();
+	RequestQueueLookup.Reset();
+
+	bIsInitialized = false;
+}
+
+int32 UMassEQSSubsystem::GetRequestQueueIndex(TSubclassOf<UEnvQueryNode> RequestClass)
+{
+	if (int32* IndexPtr = RequestQueueLookup.Find(RequestClass))
+	{
+		return *IndexPtr;
+	}
+
+	RequestQueues.Emplace();
+	int32 NewIndex = RequestQueues.Num() - 1;
+	
+	return RequestQueueLookup.Add(RequestClass, NewIndex);
+}
+
+void UMassEQSSubsystem::CancelRequest(FMassEQSRequestHandle RequestHandle)
+{
+	// Not necessary after De-initialization
+	if (!bIsInitialized)
+	{
+		return;
+	}
+
+	UE_MT_SCOPED_WRITE_ACCESS(ResultAccessDetector);
+
+	int32 NumRemoved = Results.Remove(RequestHandle);
+
+	// Log this query if its results are not yet present, we need to ignore results once they are submitted.
+	if (NumRemoved == 0u)
+	{
+		CancelledRequests.Add(RequestHandle);
+	}
+}
diff --git a/Engine/Plugins/Runtime/MassGameplay/Source/MassEQS/Private/MassEQSTypes.cpp b/Engine/Plugins/Runtime/MassGameplay/Source/MassEQS/Private/MassEQSTypes.cpp
new file mode 100644
--- /dev/null
+++ b/Engine/Plugins/Runtime/MassGameplay/Source/MassEQS/Private/MassEQSTypes.cpp
@@ -0,0 +1,83 @@
+// Copyright Epic Games, Inc. All Rights Reserved.
+
+#include "MassEQSTypes.h"
+
+#include "EnvironmentQuery/EnvQueryTypes.h"
+#include "MassEQSSubsystem.h"
+#include "MassEQS.h"
+
+
+// Creating a static representation for invalid handles
+FMassEQSRequestHandle FMassEQSRequestHandle::Invalid = FMassEQSRequestHandle(INDEX_NONE, 0);
+
+FMassEnvQueryEntityInfo::FMassEnvQueryEntityInfo(int32 Index, int32 SerialNumber, const FTransform& Transform)
+	: EntityHandle(Index, SerialNumber)
+	, CachedTransform(Transform)
+{
+}
+
+FString FMassEQSRequestHandle::ToString() const
+{
+	return FString::Printf(TEXT("[%d,%u]"), Index, SerialNumber);
+}
+
+void IMassEQSRequestInterface::FMassEQSRequestHandler::SendOrRecieveRequest(FEnvQueryInstance& QueryInstance, const IMassEQSRequestInterface& MassEQSRequestInterface)
+{
+#if WITH_EDITOR
+	if (QueryInstance.World->IsEditorWorld() && !QueryInstance.World->IsPlayInEditor())
+	{
+		UE_LOG(LogEQS, Warning, TEXT("Asynchronous request type [%s] is only available when simulating the game."), *MassEQSRequestInterface.GetRequestClass()->GetName());
+		return;
+	}
+#endif
+	if (!MassEQSSubsystem)
+	{
+		check(QueryInstance.World)
+
+		MassEQSSubsystem = QueryInstance.World->GetSubsystem<UMassEQSSubsystem>();
+		CachedRequestQueueIndex = MassEQSSubsystem->GetRequestQueueIndex(MassEQSRequestInterface.GetRequestClass());
+		check(MassEQSSubsystem);
+	}
+	
+	if (RequestHandle > 0u)
+	{
+		float CurrentTime = QueryInstance.World->GetTimeSeconds();
+		if (CurrentTime - RequestStartTime > MaxRequestTime)
+		{
+			CancelRequest();
+			Reset();
+		}
+		else if (MassEQSRequestInterface.TryAcquireResults(QueryInstance))
+		{
+			Reset();
+		}
+	}
+	else
+	{
+		RequestHandle = MassEQSSubsystem->PushRequest(QueryInstance, CachedRequestQueueIndex, MassEQSRequestInterface.GetRequestData(QueryInstance));
+		RequestStartTime = QueryInstance.World->GetTimeSeconds();
+	}
+}
+
+void IMassEQSRequestInterface::FMassEQSRequestHandler::CancelRequest() const
+{
+	if (!IsPendingResults())
+	{
+		// No active request
+		return;
+	}
+
+	check(MassEQSSubsystem)
+	MassEQSSubsystem->CancelRequest(RequestHandle);
+}
+
+void IMassEQSRequestInterface::FMassEQSRequestHandler::Reset()
+{
+	RequestHandle = FMassEQSRequestHandle::Invalid;
+	RequestStartTime = -1.f;
+}
+
+IMassEQSRequestInterface::FMassEQSRequestHandler::~FMassEQSRequestHandler()
+{
+	CancelRequest();
+}
diff --git a/Engine/Plugins/Runtime/MassGameplay/Source/MassEQS/Private/MassEQSUtils.cpp b/Engine/Plugins/Runtime/MassGameplay/Source/MassEQS/Private/MassEQSUtils.cpp
new file mode 100644
--- /dev/null
+++ b/Engine/Plugins/Runtime/MassGameplay/Source/MassEQS/Private/MassEQSUtils.cpp
@@ -0,0 +1,70 @@
+// Copyright Epic Games, Inc. All Rights Reserved.
+
+#include "MassEQSUtils.h"
+
+#include "MassEQSTypes.h"
+#include "EnvironmentQuery/EnvQueryTypes.h"
+#include "Items/EnvQueryItemType_MassEntityHandle.h"
+
+#include "MassEntityTypes.h"
+
+FMassEnvQueryEntityInfo FMassEQSUtils::GetItemAsEntityInfo(const FEnvQueryInstance& QueryInstance, int32 Index)
+{
+	UEnvQueryItemType_MassEntityHandle* DefTypeOb = QueryInstance.ItemType->GetDefaultObject<UEnvQueryItemType_MassEntityHandle>();
+	check(DefTypeOb != nullptr);
+
+	return DefTypeOb->GetValue(QueryInstance.RawData.GetData() + QueryInstance.Items[Index].DataOffset);
+}
+
+FMassEnvQueryEntityInfo FMassEQSUtils::GetItemAsEntityInfo(const FEnvQueryResult& QueryResult, int32 Index)
+{
+	UEnvQueryItemType_MassEntityHandle* DefTypeOb = QueryResult.ItemType->GetDefaultObject<UEnvQueryItemType_MassEntityHandle>();
+	check(DefTypeOb != nullptr);
+
+	return DefTypeOb->GetValue(QueryResult.RawData.GetData() + QueryResult.Items[Index].DataOffset);
+}
+
+void FMassEQSUtils::GetAllAsEntityInfo(const FEnvQueryInstance& QueryInstance, TArray<FMassEnvQueryEntityInfo>& OutEntityInfo)
+{
+	UEnvQueryItemType_MassEntityHandle* DefTypeOb = QueryInstance.ItemType->GetDefaultObject<UEnvQueryItemType_MassEntityHandle>();
+	check(DefTypeOb != nullptr);
+
+	OutEntityInfo.Reserve(OutEntityInfo.Num() + QueryInstance.Items.Num());
+	for (const FEnvQueryItem& Item : QueryInstance.Items)
+	{
+		OutEntityInfo.Add(DefTypeOb->GetValue(QueryInstance.RawData.GetData() + Item.DataOffset));
+	}
+}
+
+void FMassEQSUtils::GetAllAsEntityInfo(const FEnvQueryResult& QueryResult, TArray<FMassEnvQueryEntityInfo>& OutEntityInfo)
+{
+	UEnvQueryItemType_MassEntityHandle* DefTypeOb = QueryResult.ItemType->GetDefaultObject<UEnvQueryItemType_MassEntityHandle>();
+	check(DefTypeOb != nullptr);
+
+	OutEntityInfo.Reserve(OutEntityInfo.Num() + QueryResult.Items.Num());
+	for (const FEnvQueryItem& Item : QueryResult.Items)
+	{
+		OutEntityInfo.Add(DefTypeOb->GetValue(QueryResult.RawData.GetData() + Item.DataOffset));
+	}
+}
+
+void FMassEQSUtils::GetEntityHandles(const TArray<FMassEnvQueryEntityInfo>& EntityInfo, TArray<FMassEntityHandle>& OutHandles)
+{
+	OutHandles.Reserve(OutHandles.Num() + EntityInfo.Num());
+	for (const FMassEnvQueryEntityInfo& Info : EntityInfo)
+	{
+		OutHandles.Add(Info.EntityHandle);
+	}
+}
+
+void FMassEQSUtils::GetAllAsEntityHandles(const FEnvQueryInstance& QueryInstance, TArray<FMassEntityHandle>& OutHandles)
+{
+	UEnvQueryItemType_MassEntityHandle* DefTypeOb = QueryInstance.ItemType->GetDefaultObject<UEnvQueryItemType_MassEntityHandle>();
+	check(DefTypeOb != nullptr);
+
+	OutHandles.Reserve(OutHandles.Num() + QueryInstance.Items.Num());
+	for (const FEnvQueryItem& Item : QueryInstance.Items)
+	{
+		OutHandles.Add(DefTypeOb->GetValue(QueryInstance.RawData.GetData() + Item.DataOffset).EntityHandle);
+	}
+}
\ No newline at end of file
diff --git a/Engine/Plugins/Runtime/MassGameplay/Source/MassEQS/Private/Processors/MassEnvQueryGeneratorProcessor_MassEntityHandles.cpp b/Engine/Plugins/Runtime/MassGameplay/Source/MassEQS/Private/Processors/MassEnvQueryGeneratorProcessor_MassEntityHandles.cpp
new file mode 100644
--- /dev/null
+++ b/Engine/Plugins/Runtime/MassGameplay/Source/MassEQS/Private/Processors/MassEnvQueryGeneratorProcessor_MassEntityHandles.cpp
@@ -0,0 +1,87 @@
+// Copyright Epic Games, Inc. All Rights Reserved.
+
+#include "Processors/MassEnvQueryGeneratorProcessor_MassEntityHandles.h"
+
+#include "MassExecutionContext.h"
+#include "MassCommonFragments.h"
+#include "MassObserverRegistry.h"
+
+#include "MassEQSUtils.h"
+#include "MassEQSSubsystem.h"
+
+#include "EnvironmentQuery/EnvQueryTypes.h"
+#include "Items/EnvQueryItemType_MassEntityHandle.h"
+#include "Generators/MassEnvQueryGenerator_MassEntityHandles.h"
+
+UMassEnvQueryGeneratorProcessor_MassEntityHandles::UMassEnvQueryGeneratorProcessor_MassEntityHandles()
+	: EntityQuery(*this)
+{
+	CorrespondingRequestClass = UMassEnvQueryGenerator_MassEntityHandles::StaticClass();
+}
+
+void UMassEnvQueryGeneratorProcessor_MassEntityHandles::ConfigureQueries()
+{
+	EntityQuery.AddRequirement<FTransformFragment>(EMassFragmentAccess::ReadOnly);
+	ProcessorRequirements.AddSubsystemRequirement<UMassEQSSubsystem>(EMassFragmentAccess::ReadWrite);
+}
+
+void UMassEnvQueryGeneratorProcessor_MassEntityHandles::Execute(FMassEntityManager& EntityManager, FMassExecutionContext& Context)
+{
+	const UWorld* World = GetWorld();
+	check(World);
+
+	UMassEQSSubsystem* MassEQSSubsystem = Context.GetMutableSubsystem<UMassEQSSubsystem>();
+	check(MassEQSSubsystem);
+
+	// Check for any requests of this type from MassEQSSubsystem, complete one if found.
+	TUniquePtr<FMassEQSRequestData> GeneratorDataUniquePtr = MassEQSSubsystem->PopRequest(CachedRequestQueryIndex);
+	FMassEQSRequestData_MassEntityHandles* GeneratorData = FMassEQSUtils::TryAndEnsureCast<FMassEQSRequestData_MassEntityHandles>(GeneratorDataUniquePtr);
+	if (!GeneratorData)
+	{
+		return;
+	}
+
+	const float SearchRadius = GeneratorData->SearchRadius;
+	const TArray<FVector>& ContextPositions = GeneratorData->ContextPositions;
+
+	TArray<FMassEnvQueryEntityInfo> Items = {};
+	if (SearchRadius <= 0)
+	{
+		EntityQuery.ForEachEntityChunk(EntityManager, Context, [&GeneratorData, &Items, &SearchRadius, &ContextPositions](FMassExecutionContext& ChunkContext)
+		{
+			const TConstArrayView<FTransformFragment> TransformFragmentList = ChunkContext.GetFragmentView<FTransformFragment>();
+			const int32 NumEntities = ChunkContext.GetNumEntities();
+			for (int32 EntityIndex = 0; EntityIndex < NumEntities; ++EntityIndex)
+			{
+				FMassEntityHandle EntityHandle = ChunkContext.GetEntity(EntityIndex);
+				Items.Emplace(EntityHandle.Index, EntityHandle.SerialNumber, TransformFragmentList[EntityIndex].GetTransform());
+			}
+		});
+	}
+	else
+	{
+		FVector::FReal SearchRadiusSqr = FMath::Square(SearchRadius);
+		EntityQuery.ForEachEntityChunk(EntityManager, Context, [&GeneratorData, &Items, &SearchRadiusSqr, &ContextPositions](FMassExecutionContext& ChunkContext)
+		{
+			const TConstArrayView<FTransformFragment> TransformFragmentList = ChunkContext.GetFragmentView<FTransformFragment>();
+			const int32 NumEntities = ChunkContext.GetNumEntities();
+			for (int32 EntityIndex = 0; EntityIndex < NumEntities; ++EntityIndex)
+			{
+				const FTransformFragment& TransformFragment = TransformFragmentList[EntityIndex];
+				const FVector& EntityPosition = TransformFragment.GetTransform().GetTranslation();
+				for (const FVector& ContextPosition : ContextPositions)
+				{
+					const FVector::FReal ContextDistanceToEntitySqr = FVector::DistSquared(EntityPosition, ContextPosition);
+
+					if (ContextDistanceToEntitySqr <= SearchRadiusSqr)
+					{
+						FMassEntityHandle EntityHandle = ChunkContext.GetEntity(EntityIndex);
+						Items.Emplace(EntityHandle.Index, EntityHandle.SerialNumber, TransformFragmentList[EntityIndex].GetTransform());
+						break;
+					}
+				}
+			}
+		});
+	}
+	MassEQSSubsystem->SubmitResults(GeneratorData->RequestHandle, MakeUnique<FMassEnvQueryResultData_MassEntityHandles>(MoveTemp(Items)));
+}
\ No newline at end of file
diff --git a/Engine/Plugins/Runtime/MassGameplay/Source/MassEQS/Private/Processors/MassEnvQueryProcessorBase.cpp b/Engine/Plugins/Runtime/MassGameplay/Source/MassEQS/Private/Processors/MassEnvQueryProcessorBase.cpp
new file mode 100644
--- /dev/null
+++ b/Engine/Plugins/Runtime/MassGameplay/Source/MassEQS/Private/Processors/MassEnvQueryProcessorBase.cpp
@@ -0,0 +1,17 @@
+// Copyright Epic Games, Inc. All Rights Reserved.
+
+#include "Processors/MassEnvQueryProcessorBase.h"
+
+#include "MassEQSSubsystem.h"
+
+void UMassEnvQueryProcessorBase::Initialize(UObject& Owner)
+{
+	Super::Initialize(Owner);
+	UMassEQSSubsystem* MassEQSSubsystem = Owner.GetWorld()->GetSubsystem<UMassEQSSubsystem>();
+	check(MassEQSSubsystem)
+
+	if (CorrespondingRequestClass)
+	{
+		CachedRequestQueryIndex = MassEQSSubsystem->GetRequestQueueIndex(CorrespondingRequestClass);
+	}
+}
\ No newline at end of file
diff --git a/Engine/Plugins/Runtime/MassGameplay/Source/MassEQS/Private/Processors/MassEnvQueryTestProcessor_MassEntityTags.cpp b/Engine/Plugins/Runtime/MassGameplay/Source/MassEQS/Private/Processors/MassEnvQueryTestProcessor_MassEntityTags.cpp
new file mode 100644
--- /dev/null
+++ b/Engine/Plugins/Runtime/MassGameplay/Source/MassEQS/Private/Processors/MassEnvQueryTestProcessor_MassEntityTags.cpp
@@ -0,0 +1,175 @@
+// Copyright Epic Games, Inc. All Rights Reserved.
+
+#include "Processors/MassEnvQueryTestProcessor_MassEntityTags.h"
+
+#include "Tests/MassEnvQueryTest_MassEntityTags.h"
+
+#include "MassEQSSubsystem.h"
+#include "MassEQSTypes.h"
+#include "MassEQSUtils.h"
+#include "MassEQS.h"
+
+#include "MassCommonFragments.h"
+#include "MassExecutionContext.h"
+
+namespace Mass::EQS::Utils
+{
+
+	bool TestChunkForAnyTags(const TArray<FInstancedStruct>& Tags, const FMassExecutionContext& Context);
+	bool TestChunkForAllTags(const TArray<FInstancedStruct>& Tags, const FMassExecutionContext& Context);
+	bool TestChunkForNoTags(const TArray<FInstancedStruct>& Tags, const FMassExecutionContext& Context);
+
+} // namespace Mass::EQS::Utils
+
+
+UMassEnvQueryTestProcessor_MassEntityTags::UMassEnvQueryTestProcessor_MassEntityTags()
+	: EntityQuery(*this)
+{
+	CorrespondingRequestClass = UMassEnvQueryTest_MassEntityTags::StaticClass();
+}
+
+void UMassEnvQueryTestProcessor_MassEntityTags::ConfigureQueries()
+{
+	ProcessorRequirements.AddSubsystemRequirement<UMassEQSSubsystem>(EMassFragmentAccess::ReadWrite);
+}
+
+void UMassEnvQueryTestProcessor_MassEntityTags::Execute(FMassEntityManager& EntityManager, FMassExecutionContext& ExecutionContext)
+{
+	const UWorld* World = GetWorld();
+	check(World);
+
+	UMassEQSSubsystem* MassEQSSubsystem = ExecutionContext.GetMutableSubsystem<UMassEQSSubsystem>();
+	check(MassEQSSubsystem);
+
+	// Check for any requests of this type from MassEQSSubsystem, complete one if found.
+	TUniquePtr<FMassEQSRequestData> TestDataUniquePtr = MassEQSSubsystem->PopRequest(CachedRequestQueryIndex);
+	FMassEQSRequestData_MassEntityTags* TestData = FMassEQSUtils::TryAndEnsureCast<FMassEQSRequestData_MassEntityTags>(TestDataUniquePtr);
+	if (!TestData)
+	{
+		return;
+	}
+	if (TestDataUniquePtr->EntityHandles.IsEmpty())
+	{
+		UE_LOG(LogMassEQS, Error, TEXT("Request: [%s] Acquired by UMassEnvQueryTestProcessor_MassEntityTags, but had no Entities to query."), *TestDataUniquePtr->RequestHandle.ToString());
+		return;
+	}
+
+	const EMassEntityTagsTestMode TagTestMode = TestData->TagTestMode;
+	const TArray<FInstancedStruct> Tags = TestData->Tags;
+
+	TFunction<bool(const TArray<FInstancedStruct>&, const FMassExecutionContext&)> ScoringFunction;
+	TMap<FMassEntityHandle, bool> ScoreMap = {};
+
+	switch (TagTestMode)
+	{
+	case EMassEntityTagsTestMode::Any:
+	{
+		ScoringFunction = Mass::EQS::Utils::TestChunkForAnyTags;
+		break;
+	}
+	case EMassEntityTagsTestMode::All:
+	{
+		ScoringFunction = Mass::EQS::Utils::TestChunkForAllTags;
+		break;
+	}
+	case EMassEntityTagsTestMode::None:
+	{
+		ScoringFunction = Mass::EQS::Utils::TestChunkForNoTags;
+		break;
+	}
+	default:
+	{
+		ScoringFunction = [](const TArray<FInstancedStruct>& Tags, const FMassExecutionContext& Context) { return false; };
+		break;
+	}
+	}
+
+	ensureMsgf(ExecutionContext.GetEntityCollection().IsEmpty(), TEXT("We don't expect any collections to be set at this point. The data is going to be overridden."));
+
+	TArray<FMassArchetypeEntityCollection> EntityCollectionsToTest;
+	UE::Mass::Utils::CreateEntityCollections(EntityManager, TestDataUniquePtr->EntityHandles, FMassArchetypeEntityCollection::NoDuplicates, EntityCollectionsToTest);
+
+	EntityQuery.ForEachEntityChunkInCollections(EntityCollectionsToTest, EntityManager, ExecutionContext
+		, [&ScoreMap, &ScoringFunction, &Tags](FMassExecutionContext& Context)
+		{
+			bool ChunkSuccess = ScoringFunction(Tags, Context);
+
+			const int32 NumEntities = Context.GetNumEntities();
+			for (int32 EntityIndex = 0; EntityIndex < NumEntities; ++EntityIndex)
+			{
+				const FMassEntityHandle EntityHandle = Context.GetEntity(EntityIndex);
+				ScoreMap.Add(EntityHandle, ChunkSuccess);
+			}
+		});
+
+	MassEQSSubsystem->SubmitResults(TestData->RequestHandle, MakeUnique<FMassEnvQueryResultData_MassEntityTags>(MoveTemp(ScoreMap)));
+}
+
+
+bool Mass::EQS::Utils::TestChunkForAnyTags(const TArray<FInstancedStruct>& Tags, const FMassExecutionContext& Context)
+{
+	for (const FInstancedStruct& Tag : Tags)
+	{
+		const UScriptStruct* TagScriptStruct = Tag.GetScriptStruct();
+#if !UE_BUILD_SHIPPING
+		if (!ensureMsgf(TagScriptStruct && TagScriptStruct->IsChildOf(FMassTag::StaticStruct()), TEXT("Non Tag Element Detected in Enviroment Query Mass-Entity-Tags-Test. Likely left empty.")))
+		{
+			continue;
+		}
+#endif
+		if (Context.DoesArchetypeHaveTag(*Tag.GetScriptStruct()))
+		{
+			return true;
+		}
+	}
+
+	return false;
+}
+
+bool Mass::EQS::Utils::TestChunkForAllTags(const TArray<FInstancedStruct>& Tags, const FMassExecutionContext& Context)
+{
+	for (const FInstancedStruct& Tag : Tags)
+	{
+		const UScriptStruct* TagScriptStruct = Tag.GetScriptStruct();
+#if !UE_BUILD_SHIPPING
+		if (!ensureMsgf(TagScriptStruct && TagScriptStruct->IsChildOf(FMassTag::StaticStruct()), TEXT("Non Tag Element Detected in Enviroment Query Mass-Entity-Tags-Test. Likely left empty.")))
+		{
+			continue;
+		}
+#endif
+		if (!Context.DoesArchetypeHaveTag(*Tag.GetScriptStruct()))
+		{
+			return false;
+		}
+	}
+
+	return true;
+}
+
+bool Mass::EQS::Utils::TestChunkForNoTags(const TArray<FInstancedStruct>& Tags, const FMassExecutionContext& Context)
+{
+	for (const FInstancedStruct& Tag : Tags)
+	{
+		const UScriptStruct* TagScriptStruct = Tag.GetScriptStruct();
+#if !UE_BUILD_SHIPPING
+		if (!ensureMsgf(TagScriptStruct && TagScriptStruct->IsChildOf(FMassTag::StaticStruct()), TEXT("Non Tag Element Detected in Enviroment Query Mass-Entity-Tags-Test. Likely left empty.")))
+		{
+			continue;
+		}
+#endif
+		if (Context.DoesArchetypeHaveTag(*Tag.GetScriptStruct()))
+		{
+			return false;
+		}
+	}
+
+	return true;
+}
+
+
+
+
+
+
+
+
diff --git a/Engine/Plugins/Runtime/MassGameplay/Source/MassEQS/Private/Tests/MassEnvQueryTest.cpp b/Engine/Plugins/Runtime/MassGameplay/Source/MassEQS/Private/Tests/MassEnvQueryTest.cpp
new file mode 100644
--- /dev/null
+++ b/Engine/Plugins/Runtime/MassGameplay/Source/MassEQS/Private/Tests/MassEnvQueryTest.cpp
@@ -0,0 +1,14 @@
+// Copyright Epic Games, Inc. All Rights Reserved.
+
+#include "Tests/MassEnvQueryTest.h"
+#include "MassEQSSubsystem.h"
+
+UMassEnvQueryTest::UMassEnvQueryTest(const FObjectInitializer& ObjectInitializer)
+	: Super(ObjectInitializer)
+{
+}
+
+void UMassEnvQueryTest::RunTest(FEnvQueryInstance& QueryInstance) const
+{
+	MassEQSRequestHandler.SendOrRecieveRequest(QueryInstance, *this);
+}
\ No newline at end of file
diff --git a/Engine/Plugins/Runtime/MassGameplay/Source/MassEQS/Private/Tests/MassEnvQueryTest_MassEntityTags.cpp b/Engine/Plugins/Runtime/MassGameplay/Source/MassEQS/Private/Tests/MassEnvQueryTest_MassEntityTags.cpp
new file mode 100644
--- /dev/null
+++ b/Engine/Plugins/Runtime/MassGameplay/Source/MassEQS/Private/Tests/MassEnvQueryTest_MassEntityTags.cpp
@@ -0,0 +1,66 @@
+// Copyright Epic Games, Inc. All Rights Reserved.
+
+#include "Tests/MassEnvQueryTest_MassEntityTags.h"
+
+#include "Items/EnvQueryItemType_MassEntityHandle.h"
+
+#include "MassEQSUtils.h"
+#include "MassEQSTypes.h"
+#include "MassEQSSubsystem.h"
+
+#include "MassCommandBuffer.h"
+#include "MassEntityManager.h"
+#include "MassEntitySubsystem.h"
+
+UMassEnvQueryTest_MassEntityTags::UMassEnvQueryTest_MassEntityTags(const FObjectInitializer& ObjectInitializer)
+	: Super(ObjectInitializer)
+{
+	Cost = EEnvTestCost::Low;
+	TestPurpose = EEnvTestPurpose::Type::Filter;
+	FilterType = EEnvTestFilterType::Type::Match;
+
+	ValidItemType = UEnvQueryItemType_MassEntityHandle::StaticClass();
+
+	SetWorkOnFloatValues(false);
+
+}
+
+TUniquePtr<FMassEQSRequestData> UMassEnvQueryTest_MassEntityTags::GetRequestData(FEnvQueryInstance& QueryInstance) const
+{
+	return MakeUnique<FMassEQSRequestData_MassEntityTags>(TagTestMode, Tags);
+}
+
+bool UMassEnvQueryTest_MassEntityTags::TryAcquireResults(FEnvQueryInstance& QueryInstance) const
+{
+	check(MassEQSRequestHandler.MassEQSSubsystem)
+	TUniquePtr<FMassEQSRequestData> RawRequestData = MassEQSRequestHandler.MassEQSSubsystem->TryAcquireResults(MassEQSRequestHandler.RequestHandle);
+	FMassEnvQueryResultData_MassEntityTags* RequestData = FMassEQSUtils::TryAndEnsureCast<FMassEnvQueryResultData_MassEntityTags>(RawRequestData);
+	if (!RequestData)
+	{
+		return false;
+	}
+
+	FEnvQueryInstance::ItemIterator It(this, QueryInstance);
+	for (It.IgnoreTimeLimit(); It; ++It)
+	{
+		const FMassEnvQueryEntityInfo& EntityInfo = FMassEQSUtils::GetItemAsEntityInfo(QueryInstance, It.GetIndex());
+
+		const bool bSuccess = RequestData->ResultMap.Contains(EntityInfo.EntityHandle) && RequestData->ResultMap[EntityInfo.EntityHandle];
+
+		It.SetScore(TestPurpose, FilterType, bSuccess, true);
+	}
+
+	return true;
+}
+
+FText UMassEnvQueryTest_MassEntityTags::GetDescriptionTitle() const
+{
+	return FText::FromString(FString::Printf(TEXT("Mass Entity Tags Test : Match %s Tags"),
+		TagTestMode == EMassEntityTagsTestMode::Any ? TEXT("Any") :
+		TagTestMode == EMassEntityTagsTestMode::All ? TEXT("All") : TEXT("None")));
+}
+
+FText UMassEnvQueryTest_MassEntityTags::GetDescriptionDetails() const
+{
+	return DescribeFloatTestParams();
+}
\ No newline at end of file
diff --git a/Engine/Plugins/Runtime/MassGameplay/Source/MassEQS/Public/Generators/MassEnvQueryGenerator.h b/Engine/Plugins/Runtime/MassGameplay/Source/MassEQS/Public/Generators/MassEnvQueryGenerator.h
new file mode 100644
--- /dev/null
+++ b/Engine/Plugins/Runtime/MassGameplay/Source/MassEQS/Public/Generators/MassEnvQueryGenerator.h
@@ -0,0 +1,33 @@
+// Copyright Epic Games, Inc. All Rights Reserved.
+
+#pragma once
+
+#include "EnvironmentQuery/EnvQueryGenerator.h"
+#include "MassEQSTypes.h"
+
+#include "MassEnvQueryGenerator.generated.h"
+
+/** Generator that will send its work to MassEQSSubsystem in order to be processed in a Mass Processor */
+UCLASS(EditInlineNew, Abstract, meta = (Category = "Generators"), MinimalAPI)
+class UMassEnvQueryGenerator : public UEnvQueryGenerator, public IMassEQSRequestInterface
+{
+	GENERATED_UCLASS_BODY()
+public:
+	/** 
+	 * This will send this Generator Request to MassEQSSubsystem the first time it is called,
+	 * and Try to complete Generation with the Result from MassEQSSubsystem on subsequent calls.
+	*/
+	virtual void GenerateItems(FEnvQueryInstance& QueryInstance) const final;
+
+	virtual FORCEINLINE bool IsCurrentlyRunningAsync() const final { return MassEQSRequestHandler.IsPendingResults(); }
+
+	// To be implemented by child class:
+	// Begin IMassEQSRequestInterface
+	virtual TUniquePtr<FMassEQSRequestData> GetRequestData(FEnvQueryInstance& QueryInstance) const override PURE_VIRTUAL(UMassEnvQueryGenerator::GetRequestData, return nullptr;);
+	virtual UClass* GetRequestClass() const override PURE_VIRTUAL(UMassEnvQueryGenerator::GetRequestData, return nullptr;);
+	virtual bool TryAcquireResults(FEnvQueryInstance& QueryInstance) const override PURE_VIRTUAL(UMassEnvQueryGenerator::TryAcquireResults, return false;);
+	// ~IMassEQSRequestInterface
+
+protected:
+	mutable FMassEQSRequestHandler MassEQSRequestHandler;
+};
\ No newline at end of file
diff --git a/Engine/Plugins/Runtime/MassGameplay/Source/MassEQS/Public/Generators/MassEnvQueryGenerator_MassEntityHandles.h b/Engine/Plugins/Runtime/MassGameplay/Source/MassEQS/Public/Generators/MassEnvQueryGenerator_MassEntityHandles.h
new file mode 100644
--- /dev/null
+++ b/Engine/Plugins/Runtime/MassGameplay/Source/MassEQS/Public/Generators/MassEnvQueryGenerator_MassEntityHandles.h
@@ -0,0 +1,62 @@
+// Copyright Epic Games, Inc. All Rights Reserved.
+
+#pragma once
+
+#include "DataProviders/AIDataProvider.h"
+#include "Generators/MassEnvQueryGenerator.h"
+#include "MassEQSTypes.h"
+
+#include "MassEnvQueryGenerator_MassEntityHandles.generated.h"
+
+struct FMassEnvQueryEntityInfo;
+
+/**
+ * Generator to be sent to MassEQSSubsystem for processing on Mass.
+ * This will Generate UEnvQueryItemType_MassEntityHandles within SearchRadius of any ContextPositions.
+ * Set SearchRadius to a value <= 0 in order to get all EntityHandles who have an FTransformFragment.
+ */
+UCLASS(meta = (DisplayName = "Mass Entity Handles"), MinimalAPI)
+class UMassEnvQueryGenerator_MassEntityHandles : public UMassEnvQueryGenerator
+{
+	GENERATED_UCLASS_BODY()
+
+public:
+	// Begin IMassEQSRequestInterface
+	virtual TUniquePtr<FMassEQSRequestData> GetRequestData(FEnvQueryInstance& QueryInstance) const override;
+	virtual UClass* GetRequestClass() const override { return StaticClass(); }
+	
+	virtual bool TryAcquireResults(FEnvQueryInstance& QueryInstance) const override;
+	// ~IMassEQSRequestInterface
+
+protected:
+	/** Any Entity who is within SearchRadius of any SearchCenter will be acquired */
+	UPROPERTY(EditDefaultsOnly, Category=Generator)
+	FAIDataProviderFloatValue SearchRadius;
+
+	/** Context of query */
+	UPROPERTY(EditAnywhere, Category=Generator)
+	TSubclassOf<UEnvQueryContext> SearchCenter = nullptr;
+};
+
+/** Data required to be sent to Mass for processing this Generator Request */
+struct FMassEQSRequestData_MassEntityHandles : public FMassEQSRequestData
+{
+	FMassEQSRequestData_MassEntityHandles(const TArray<FVector>& InContextPositions, const float InSearchRadius)
+		: ContextPositions(InContextPositions)
+		, SearchRadius(InSearchRadius)
+	{
+	}
+
+	TArray<FVector> ContextPositions;
+	float SearchRadius;
+};
+
+struct FMassEnvQueryResultData_MassEntityHandles: public FMassEQSRequestData
+{
+	FMassEnvQueryResultData_MassEntityHandles(TArray<FMassEnvQueryEntityInfo>&& InGeneratedEntityInfo)
+		: GeneratedEntityInfo(MoveTemp(InGeneratedEntityInfo))
+	{
+	}
+
+	TArray<FMassEnvQueryEntityInfo> GeneratedEntityInfo;
+};
\ No newline at end of file
diff --git a/Engine/Plugins/Runtime/MassGameplay/Source/MassEQS/Public/Items/EnvQueryItemType_MassEntityHandle.h b/Engine/Plugins/Runtime/MassGameplay/Source/MassEQS/Public/Items/EnvQueryItemType_MassEntityHandle.h
new file mode 100644
--- /dev/null
+++ b/Engine/Plugins/Runtime/MassGameplay/Source/MassEQS/Public/Items/EnvQueryItemType_MassEntityHandle.h
@@ -0,0 +1,28 @@
+// Copyright Epic Games, Inc. All Rights Reserved.
+
+#pragma once
+
+#include "MassEQSTypes.h"
+#include "EnvironmentQuery/Items/EnvQueryItemType_VectorBase.h"
+#include "EnvQueryItemType_MassEntityHandle.generated.h"
+
+/** 
+ * EnvQueryType representing a MassEntity. Additionally stores a Cached Transform of the Entity at the
+ * time of acquisition, in order to use in the implemented UEnvQueryItemType_VectorBase functions.
+ */
+UCLASS(meta = (DisplayName = "Mass Entity Handles Item Type"), MinimalAPI)
+class UEnvQueryItemType_MassEntityHandle : public UEnvQueryItemType_VectorBase
+{
+	GENERATED_BODY()
+
+public:
+	typedef FMassEnvQueryEntityInfo FValueType;
+
+	UEnvQueryItemType_MassEntityHandle();
+
+	static const FMassEnvQueryEntityInfo& GetValue(const uint8* RawData);
+	static void SetValue(uint8* RawData, const FMassEnvQueryEntityInfo& Value);
+
+	virtual FVector GetItemLocation(const uint8* RawData) const override;
+	virtual FRotator GetItemRotation(const uint8* RawData) const override;
+};
diff --git a/Engine/Plugins/Runtime/MassGameplay/Source/MassEQS/Public/MassEQS.h b/Engine/Plugins/Runtime/MassGameplay/Source/MassEQS/Public/MassEQS.h
new file mode 100644
--- /dev/null
+++ b/Engine/Plugins/Runtime/MassGameplay/Source/MassEQS/Public/MassEQS.h
@@ -0,0 +1,11 @@
+// Copyright Epic Games, Inc. All Rights Reserved.
+
+#pragma once
+
+#include "Modules/ModuleManager.h"
+
+DECLARE_LOG_CATEGORY_EXTERN(LogMassEQS, Log, All);
+
+class FMassEQSModule : public IModuleInterface
+{
+};
diff --git a/Engine/Plugins/Runtime/MassGameplay/Source/MassEQS/Public/MassEQSBlueprintLibrary.h b/Engine/Plugins/Runtime/MassGameplay/Source/MassEQS/Public/MassEQSBlueprintLibrary.h
new file mode 100644
--- /dev/null
+++ b/Engine/Plugins/Runtime/MassGameplay/Source/MassEQS/Public/MassEQSBlueprintLibrary.h
@@ -0,0 +1,88 @@
+// Copyright Epic Games, Inc. All Rights Reserved.
+
+#pragma once
+
+#include "Kismet/BlueprintFunctionLibrary.h"
+
+#include "MassEntityTypes.h"
+#include "MassEQSTypes.h"
+#include "MassEQSUtils.h"
+
+#include "MassEQSBlueprintLibrary.generated.h"
+
+class UEnvQueryInstanceBlueprintWrapper;
+struct FEnvQueryResult;
+
+/** Wrapper for Blueprints to be able to parse MassEntityInfo to use functions defined in UMassEQSBlueprintLibrary */
+USTRUCT(Blueprintable, BlueprintType, meta = (DisplayName = "Mass Entity Info"))
+struct FMassEnvQueryEntityInfoBlueprintWrapper
+{
+    GENERATED_BODY()
+
+public:
+	FMassEnvQueryEntityInfoBlueprintWrapper() = default;
+	FMassEnvQueryEntityInfoBlueprintWrapper(FMassEnvQueryEntityInfo InEntityInfo)
+		: EntityInfo(InEntityInfo)
+	{
+	}
+
+	FORCEINLINE FVector GetCachedEntityPosition() const { return EntityInfo.CachedTransform.GetLocation(); }
+	FORCEINLINE FMassEntityHandle GetEntityHandle() const { return EntityInfo.EntityHandle; }
+	FORCEINLINE void SetEntityHandle(FMassEntityHandle Handle) { EntityInfo.EntityHandle = Handle; }
+
+	const FMassEnvQueryEntityInfo& GetEntityInfo() const { return EntityInfo; };
+	inline bool operator==(const FMassEnvQueryEntityInfoBlueprintWrapper& Other) const { return GetEntityInfo() == Other.GetEntityInfo(); }
+
+private:
+	FMassEnvQueryEntityInfo EntityInfo;
+
+};
+
+/** Function library for interfacing with EntityInfo inside blueprints. */
+UCLASS()
+class UMassEQSBlueprintLibrary : public UBlueprintFunctionLibrary
+{
+    GENERATED_UCLASS_BODY()
+
+public:
+	
+	//----------------------------------------------------------------------//
+	// Commands
+	//----------------------------------------------------------------------//
+
+	/** Sends the input Signal to the Entity defined by EntityInfo.EntityHandle using the UMassSignalSubsystem. */
+	UFUNCTION(BlueprintCallable, Category = "MassEnvQuery|Commands")
+	static void SendSignalToEntity(const AActor* Owner, const FMassEnvQueryEntityInfoBlueprintWrapper& EntityInfo, const FName Signal);
+
+	//----------------------------------------------------------------------//
+	// Utils
+	//----------------------------------------------------------------------//
+
+	UFUNCTION(BlueprintCallable, BlueprintPure, Category = "MassEnvQuery|Utils")
+	static FORCEINLINE FString EntityToString(const FMassEnvQueryEntityInfoBlueprintWrapper& EntityInfo){ return EntityInfo.GetEntityHandle().DebugGetDescription(); }
+
+	UFUNCTION(BlueprintCallable, BlueprintPure, Category = "MassEnvQuery|Utils")
+	static FORCEINLINE FVector GetCachedEntityPosition(const FMassEnvQueryEntityInfoBlueprintWrapper& EntityInfo) { return EntityInfo.GetCachedEntityPosition(); }
+
+	UFUNCTION(BlueprintCallable, BlueprintPure, Category = "MassEnvQuery|Utils")
+	static FVector GetCurrentEntityPosition(const AActor* Owner, const FMassEnvQueryEntityInfoBlueprintWrapper& EntityInfo);
+
+	/** Custom comparison function, as the Blueprint Equals did not seem to work. */
+	UFUNCTION(BlueprintCallable, BlueprintPure, Category = "MassEnvQuery|Utils")
+	static FORCEINLINE bool EntityComparison(const FMassEnvQueryEntityInfoBlueprintWrapper& A, const FMassEnvQueryEntityInfoBlueprintWrapper& B) { return A == B; }
+
+	/** Custom array-contains function, as the Blueprint version did not seem to work. */
+	UFUNCTION(BlueprintCallable, BlueprintPure, Category = "MassEnvQuery|Utils")
+	static bool ContainsEntity(const TArray<FMassEnvQueryEntityInfoBlueprintWrapper>& EntityList, const FMassEnvQueryEntityInfoBlueprintWrapper& EntityInfo);
+
+	/** Outputs an array filled with resulting EntityInfos. Note that it makes sense only if ItemType is a EnvQueryItemType_MassEntityHandle-derived type. */
+	UFUNCTION(BlueprintCallable, Category = "AI|EQS")
+	static TArray<FMassEnvQueryEntityInfoBlueprintWrapper> GetEnviromentQueryResultAsEntityInfo(const UEnvQueryInstanceBlueprintWrapper* QueryInstance);
+
+private:
+	/** Get result and immediately convert to EntityInfoBlueprintWrapper to skip an extra copy step */
+	static FMassEnvQueryEntityInfoBlueprintWrapper GetItemAsEntityInfoBPWrapper(const FEnvQueryResult* QueryResult, int32 Index);
+
+	/** Get result array and convert each EntityInfo to EntityInfoBlueprintWrapper along the way to skip extra copy steps */
+	static void GetAllAsEntityInfoBPWrappers(const FEnvQueryResult* QueryResult, TArray<FMassEnvQueryEntityInfoBlueprintWrapper>& OutEntityInfo);
+};
\ No newline at end of file
diff --git a/Engine/Plugins/Runtime/MassGameplay/Source/MassEQS/Public/MassEQSSubsystem.h b/Engine/Plugins/Runtime/MassGameplay/Source/MassEQS/Public/MassEQSSubsystem.h
new file mode 100644
--- /dev/null
+++ b/Engine/Plugins/Runtime/MassGameplay/Source/MassEQS/Public/MassEQSSubsystem.h
@@ -0,0 +1,94 @@
+// Copyright Epic Games, Inc. All Rights Reserved.
+
+#pragma once
+
+#include "Containers/Queue.h"
+
+#include "EnvironmentQuery/EnvQueryNode.h"
+#include "MassEQSTypes.h"
+
+#include "MassSubsystemBase.h"
+
+#include "MassEQSSubsystem.generated.h"
+
+/**
+ * Subsystem for sending UMassEQSGenerator and UMassEQSTest requests to Mass Processors, 
+ * and for returning the results back to EQS. 
+ */
+UCLASS()
+class UMassEQSSubsystem : public UMassSubsystemBase
+{
+	GENERATED_BODY()
+	/** Queue holding UniquePtrs to the data required for each Request */
+	using FRequestQueue = TQueue<TUniquePtr<FMassEQSRequestData>>;
+
+public:
+	/** Push new Request into RequestQueue */
+	FMassEQSRequestHandle PushRequest(const FEnvQueryInstance& QueryInstance, int32 RequestQueueIndex, TUniquePtr<FMassEQSRequestData>&& Request);
+
+	/** Dequeue next Request off the RequestQueue, removing it from the queue, and transferring ownership of UniquePtr */
+	TUniquePtr<FMassEQSRequestData> PopRequest(int32 RequestQueueIndex);
+
+	/** Send results from a finished/processed request to be stored in this subsystem */
+	void SubmitResults(FMassEQSRequestHandle RequestHandle, TUniquePtr<FMassEQSRequestData>&& Result);
+	/** Try to Acquire the results from a request from the QueryInstance, if it is available. */
+	TUniquePtr<FMassEQSRequestData> TryAcquireResults(FMassEQSRequestHandle RequestHandle);
+
+	/** Returns the index into RequestQueues for the input class */
+	int32 GetRequestQueueIndex(TSubclassOf<UEnvQueryNode> RequestClass);
+
+	/**
+	 * If Request in RequestQueue -> Remove from Queue
+	 * If Result in Results -> Remove from Results
+	 * Otherwise, Request is currently being processed -> Log this ID as Cancelled, and ignore result when submitted.
+	 */
+	void CancelRequest(FMassEQSRequestHandle RequestHandle);
+
+	/** Finalization step for Handle Management */
+	 void FinalizeRequest(FMassEQSRequestHandle RequestHandle);
+
+	virtual void Initialize(FSubsystemCollectionBase& Collection) override;
+	virtual void Deinitialize() override;
+protected:
+
+	/**
+	 * Mapping from Class of requester, to index into RequestQueues.
+	 * Each Class that extends UMassEQSGenerator or UMassEQSTest 
+	 * can be a requester, this will store a Queue for each of these types.
+	 * @see PostInitialize()
+	 */
+	TMap<UClass*, int32> RequestQueueLookup = {};
+	/**
+	 * Array of Queues holding each Request. 
+	 * TQueue does not support copy/assignment, which TMap requires.
+	 * This is the reason the TQueues are stored in a separate Array.
+	 */
+	TArray<FRequestQueue> RequestQueues = {};
+	UE_MT_DECLARE_RW_ACCESS_DETECTOR(RequestAccessDetector);
+
+	/** Holds results for finished Query */
+	TMap<FMassEQSRequestHandle, TUniquePtr<FMassEQSRequestData>> Results;
+	UE_MT_DECLARE_RW_ACCESS_DETECTOR(ResultAccessDetector);
+
+	FIndexedHandleManagerBase<FMassEQSRequestHandle, int32, uint32, true> HandleManager;
+
+	/** 
+	 * Holds all RequestHandles of requests who have been canceled, and not yet submitted.
+	 * Once this request is submitted, we will ignore the results.
+	 */
+	TArray<FMassEQSRequestHandle> CancelledRequests;
+
+	/** Used to ignore request cancellation after subsystem is De-Initialized */
+	bool bIsInitialized = false;
+};
+
+template<>
+struct TMassExternalSubsystemTraits<UMassEQSSubsystem> final
+{
+	enum
+	{
+		GameThreadOnly = false,
+		ThreadSafeRead = true,
+		ThreadSafeWrite = true,
+	};
+};
\ No newline at end of file
diff --git a/Engine/Plugins/Runtime/MassGameplay/Source/MassEQS/Public/MassEQSTypes.h b/Engine/Plugins/Runtime/MassGameplay/Source/MassEQS/Public/MassEQSTypes.h
new file mode 100644
--- /dev/null
+++ b/Engine/Plugins/Runtime/MassGameplay/Source/MassEQS/Public/MassEQSTypes.h
@@ -0,0 +1,141 @@
+// Copyright Epic Games, Inc. All Rights Reserved.
+
+#pragma once
+
+#include "IndexedHandle.h"
+#include "MassEntityTypes.h"
+
+#include "MassEQSTypes.generated.h"
+
+struct FEnvQueryInstance;
+struct FMassEQSRequestHandle;
+class UMassEQSSubsystem;
+
+/** Contents of EnvQueryItemType_MassEntityHandle
+ *  @see EnvQueryItemType_MassEntityHandle
+ */
+USTRUCT()
+struct FMassEnvQueryEntityInfo
+{
+	GENERATED_BODY()
+
+	FMassEnvQueryEntityInfo() = default;
+	FMassEnvQueryEntityInfo(int32 Index, int32 SerialNumber, const FTransform& Transform);
+
+	FMassEntityHandle EntityHandle;
+	/** Cached Transform from time of acquisition allows EnvQueryItemType_MassEntityHandle to implement UEnvQueryItemType_VectorBase */
+	FTransform CachedTransform;
+
+	inline bool operator==(const FMassEnvQueryEntityInfo& Other) const { return EntityHandle == Other.EntityHandle; }
+};
+
+/**
+ * Provides a Unique Handle for incoming MassEQS Requests
+ */
+USTRUCT()
+struct FMassEQSRequestHandle : public FIndexedHandleBase
+{
+	GENERATED_BODY()
+	FMassEQSRequestHandle() = default;
+	FMassEQSRequestHandle(int32 InIndex, uint32 InSerialNumber) : FIndexedHandleBase(InIndex, InSerialNumber)
+	{
+	}
+
+	FString ToString() const;
+
+	static FMassEQSRequestHandle Invalid;
+};
+
+/** To be extended to hold any data which needs to be transferred between EQS and Mass for a given Request. */
+struct FMassEQSRequestData
+{
+	FMassEQSRequestData() = default;
+	FMassEQSRequestData(FMassEQSRequestHandle Handle)
+		: RequestHandle(Handle)
+	{
+	}
+
+	/** Unique ID that MassEQSSubsystem's HandleManager has assigned to this Request */
+	FMassEQSRequestHandle RequestHandle = FMassEQSRequestHandle::Invalid;
+
+	/** EntityHandles to run this request on. Used for Tests. */
+	TArray<FMassEntityHandle> EntityHandles;
+};
+
+/** Interface for sending and receiving EQS Requests to and from MassEQSSubsystem */
+UINTERFACE()
+class UMassEQSRequestInterface : public UInterface
+{
+	GENERATED_BODY()
+};
+
+/** Interface for sending and receiving EQS Requests to and from MassEQSSubsystem utilizing a MassEQSRequestHandler */
+class IMassEQSRequestInterface
+{
+	GENERATED_BODY()
+public:
+
+	/**
+	 * Stores relevant information for managing requests sent to MassEQSSubsystem.
+	 * Provides generic functionality for Sending, Receiving, and Canceling these requests.
+	 * Any class who implements IMassEQSRequestInterface should have one of these.
+	 */
+	struct FMassEQSRequestHandler
+	{
+		/**
+		 * Called when an EnvQueryInstance is trying to run a Generator or Test who implements this.
+		 * If this Request has not yet been sent (IE. RequestID == 0u) then we will send this request
+		 * Otherwise, if this Request was already sent (IE. RequestID > 0u) then we will call TryAcquireResults on the input MassEnvRequestInterface.
+		 * @param QueryInstance - The EnvQueryInstance sending/receiving this request
+		 * @param MassEQSRequestInterface - The Interface used to call this,
+		 * @return A unique ID that the QueryInstance can hold on to in order to later retrieve this request result.
+		 */
+		void SendOrRecieveRequest(FEnvQueryInstance& QueryInstance, const IMassEQSRequestInterface& MassEQSRequestInterface);
+
+		/** Cancel request if currently ongoing. */
+		void CancelRequest() const;
+
+		/** Reset RequestID to 0 and RequestStartTime to -1.f */
+		void Reset();
+
+		FORCEINLINE bool IsPendingResults() const { return RequestHandle.IsValid(); }
+
+		// Handle for request in MassEQSSubsystem.
+		FMassEQSRequestHandle RequestHandle = FMassEQSRequestHandle::Invalid;
+		// Cached MassEQSSubsystem pointer for interacting with this subsystem
+		TObjectPtr<UMassEQSSubsystem> MassEQSSubsystem = nullptr;
+		// The point in time that an ongoing request was sent, -1 if none have been sent.
+		float RequestStartTime = -1.f;
+		// The max amount of time allowed for this request. CancelRequest will be called if this is used after this time has elapsed.
+		float MaxRequestTime = DefaultMaxRequestTime;
+
+		int32 CachedRequestQueueIndex = -1;
+
+		FMassEQSRequestHandler() = default;
+		~FMassEQSRequestHandler();
+
+	protected:
+		static constexpr float DefaultMaxRequestTime = 5.f;
+	};
+
+
+	/**
+	 * Called inside RequestProcessing, in order to create a copy of the relevant data for this Request.
+	 * @param QueryInstance - The EnvQueryInstance that owns this Generator
+	 * @return A UniquePtr to the Data-Copy created.
+	 */
+	virtual TUniquePtr<FMassEQSRequestData> GetRequestData(FEnvQueryInstance& QueryInstance) const = 0;
+
+	/**
+	 * Get StaticClass of the ImplementingClass. Implementer of this Interface must provide this so we can use
+	 * the correct StaticClass inside RequestProcessing.
+	 * Implementation should just be StaticClass() called in derived class.
+	 */
+	virtual UClass* GetRequestClass() const = 0;
+
+	/**
+	 * Called when EnvQueryInstance is checking to try to complete its request.
+	 * @param QueryInstance - The EnvQueryInstance that owns this Generator
+	 */
+	virtual bool TryAcquireResults(FEnvQueryInstance& QueryInstance) const = 0;
+};
\ No newline at end of file
diff --git a/Engine/Plugins/Runtime/MassGameplay/Source/MassEQS/Public/MassEQSUtils.h b/Engine/Plugins/Runtime/MassGameplay/Source/MassEQS/Public/MassEQSUtils.h
new file mode 100644
--- /dev/null
+++ b/Engine/Plugins/Runtime/MassGameplay/Source/MassEQS/Public/MassEQSUtils.h
@@ -0,0 +1,48 @@
+// Copyright Epic Games, Inc. All Rights Reserved.
+
+#pragma once
+
+#include "HAL/Platform.h"
+#include "MassEQSTypes.h"
+
+struct FEnvQueryResult;
+struct FEnvQueryInstance;
+struct FMassEnvQueryEntityInfo;
+struct FMassEntityHandle;
+struct FMassEntityHandle;
+struct FMassEQSRequestData;
+
+/** Holds Utility functions for Mass EQS Needs */
+struct FMassEQSUtils
+{
+
+	/** Returns the Item stored in QueryInstance/QueryResult Items[Index] as EntityInfo */
+	static FMassEnvQueryEntityInfo GetItemAsEntityInfo(const FEnvQueryInstance& QueryInstance, int32 Index);
+	static FMassEnvQueryEntityInfo GetItemAsEntityInfo(const FEnvQueryResult& QueryResult, int32 Index);
+
+	/** Returns all Items stored in QueryInstance/QueryResult as EntityInfo */
+	static void GetAllAsEntityInfo(const FEnvQueryInstance& QueryInstance, TArray<FMassEnvQueryEntityInfo>& OutEntityInfo);
+	static void GetAllAsEntityInfo(const FEnvQueryResult& QueryResult, TArray<FMassEnvQueryEntityInfo>& OutEntityInfo);
+
+	/** Extracts all EntityHandles out of an array of EntityInfo */
+	static void GetEntityHandles(const TArray<FMassEnvQueryEntityInfo>& EntityInfo, TArray<FMassEntityHandle>& OutHandles);
+	static void GetAllAsEntityHandles(const FEnvQueryInstance& QueryInstance, TArray<FMassEntityHandle>& OutHandles);
+
+	/**
+	 * Used in MassEnvQueryProcessors, to cast generic FMassEQSRequestData to its corresponding child class.
+	 * If InPtr is not null, then this Cast should never fail.
+	 */
+	template<typename RequestDataType>
+	static FORCEINLINE RequestDataType* TryAndEnsureCast(TUniquePtr<FMassEQSRequestData>& InPtr)
+	{
+		if (!InPtr)
+		{
+			return nullptr;
+		}
+
+		RequestDataType* OutPtr = static_cast<RequestDataType*>(InPtr.Get());
+		ensureMsgf(OutPtr, TEXT("RequestData was pushed to MassEQSSubsystem, but corresponding child RequestData was not found."));
+
+		return OutPtr;
+	}
+};
diff --git a/Engine/Plugins/Runtime/MassGameplay/Source/MassEQS/Public/Processors/MassEnvQueryGeneratorProcessor_MassEntityHandles.h b/Engine/Plugins/Runtime/MassGameplay/Source/MassEQS/Public/Processors/MassEnvQueryGeneratorProcessor_MassEntityHandles.h
new file mode 100644
--- /dev/null
+++ b/Engine/Plugins/Runtime/MassGameplay/Source/MassEQS/Public/Processors/MassEnvQueryGeneratorProcessor_MassEntityHandles.h
@@ -0,0 +1,21 @@
+// Copyright Epic Games, Inc. All Rights Reserved.
+
+#pragma once
+
+#include "Processors/MassEnvQueryProcessorBase.h"
+#include "MassEnvQueryGeneratorProcessor_MassEntityHandles.generated.h"
+
+/** Processor for completing MassEQSSubsystem Requests sent from UMassEnvQueryGenerator_MassEntityHandles */
+UCLASS(meta = (DisplayName = "Mass Entity Handles Generator Processor"))
+class UMassEnvQueryGeneratorProcessor_MassEntityHandles : public UMassEnvQueryProcessorBase
+{
+	GENERATED_BODY()
+public:
+	UMassEnvQueryGeneratorProcessor_MassEntityHandles();
+
+protected:
+	virtual void ConfigureQueries() override;
+	virtual void Execute(FMassEntityManager& EntityManager, FMassExecutionContext& Context) override;
+
+	FMassEntityQuery EntityQuery;
+};
\ No newline at end of file
diff --git a/Engine/Plugins/Runtime/MassGameplay/Source/MassEQS/Public/Processors/MassEnvQueryProcessorBase.h b/Engine/Plugins/Runtime/MassGameplay/Source/MassEQS/Public/Processors/MassEnvQueryProcessorBase.h
new file mode 100644
--- /dev/null
+++ b/Engine/Plugins/Runtime/MassGameplay/Source/MassEQS/Public/Processors/MassEnvQueryProcessorBase.h
@@ -0,0 +1,21 @@
+// Copyright Epic Games, Inc. All Rights Reserved.
+
+#pragma once
+
+#include "MassProcessor.h"
+#include "MassEnvQueryProcessorBase.generated.h"
+
+class UEnvQueryNode;
+
+/** Processor for completing MassEQSSubsystem Requests sent from UMassEnvQueryTest_MassEntityTags */
+UCLASS(Abstract, meta = (DisplayName = "Mass EQS Processor Base"))
+class UMassEnvQueryProcessorBase : public UMassProcessor
+{
+	GENERATED_BODY()
+
+protected:
+	virtual void Initialize(UObject& Owner);
+	
+	TSubclassOf<UEnvQueryNode> CorrespondingRequestClass = nullptr;
+	int32 CachedRequestQueryIndex = -1;
+};
\ No newline at end of file
diff --git a/Engine/Plugins/Runtime/MassGameplay/Source/MassEQS/Public/Processors/MassEnvQueryTestProcessor_MassEntityTags.h b/Engine/Plugins/Runtime/MassGameplay/Source/MassEQS/Public/Processors/MassEnvQueryTestProcessor_MassEntityTags.h
new file mode 100644
--- /dev/null
+++ b/Engine/Plugins/Runtime/MassGameplay/Source/MassEQS/Public/Processors/MassEnvQueryTestProcessor_MassEntityTags.h
@@ -0,0 +1,22 @@
+// Copyright Epic Games, Inc. All Rights Reserved.
+
+#pragma once
+
+#include "Processors/MassEnvQueryProcessorBase.h"
+#include "MassEnvQueryTestProcessor_MassEntityTags.generated.h"
+
+/** Processor for completing MassEQSSubsystem Requests sent from UMassEnvQueryTest_MassEntityTags */
+UCLASS(meta = (DisplayName = "Mass Entity Tags Test Processor"))
+class UMassEnvQueryTestProcessor_MassEntityTags : public UMassEnvQueryProcessorBase
+{
+	GENERATED_BODY()
+public:
+	UMassEnvQueryTestProcessor_MassEntityTags();
+
+protected:
+	virtual void ConfigureQueries() override;
+	virtual void Execute(FMassEntityManager& EntityManager, FMassExecutionContext& Context) override;
+	virtual bool ShouldAllowQueryBasedPruning(const bool bRuntimeMode = true) const override { return false; }
+
+	FMassEntityQuery EntityQuery;
+};
\ No newline at end of file
diff --git a/Engine/Plugins/Runtime/MassGameplay/Source/MassEQS/Public/Tests/MassEnvQueryTest.h b/Engine/Plugins/Runtime/MassGameplay/Source/MassEQS/Public/Tests/MassEnvQueryTest.h
new file mode 100644
--- /dev/null
+++ b/Engine/Plugins/Runtime/MassGameplay/Source/MassEQS/Public/Tests/MassEnvQueryTest.h
@@ -0,0 +1,32 @@
+// Copyright Epic Games, Inc. All Rights Reserved.
+
+#pragma once
+
+#include "EnvironmentQuery/EnvQueryTest.h"
+#include "MassEQSTypes.h"
+
+#include "MassEnvQueryTest.generated.h"
+
+/** Test that will send its work to MassEQSSubsystem in order to be processed in a Mass Processor */
+UCLASS(EditInlineNew, Abstract, meta = (Category = "Tests"), MinimalAPI)
+class UMassEnvQueryTest : public UEnvQueryTest, public IMassEQSRequestInterface
+{
+	GENERATED_UCLASS_BODY()
+public:
+	/**
+	 * This will send this Test Request to MassEQSSubsystem the first time it is called,
+	 * and Try to complete Testing with the Result from MassEQSSubsystem on subsequent calls.
+	 */
+	virtual void RunTest(FEnvQueryInstance& QueryInstance) const final;
+	virtual FORCEINLINE bool IsCurrentlyRunningAsync() const final { return MassEQSRequestHandler.IsPendingResults(); }
+
+	// Begin IMassEQSRequestInterface
+	virtual TUniquePtr<FMassEQSRequestData> GetRequestData(FEnvQueryInstance& QueryInstance) const override PURE_VIRTUAL(UMassEnvQueryTest::GetRequestData, return nullptr;)
+	virtual UClass* GetRequestClass() const override PURE_VIRTUAL(UMassEnvQueryTest::GetRequestClass, return nullptr;);
+
+	virtual bool TryAcquireResults(FEnvQueryInstance& QueryInstance) const override PURE_VIRTUAL(UMassEnvQueryTest::GetRequestData, return false;)
+	// ~IMassEQSRequestInterface
+
+protected:
+	mutable FMassEQSRequestHandler MassEQSRequestHandler;
+};
\ No newline at end of file
diff --git a/Engine/Plugins/Runtime/MassGameplay/Source/MassEQS/Public/Tests/MassEnvQueryTest_MassEntityTags.h b/Engine/Plugins/Runtime/MassGameplay/Source/MassEQS/Public/Tests/MassEnvQueryTest_MassEntityTags.h
new file mode 100644
--- /dev/null
+++ b/Engine/Plugins/Runtime/MassGameplay/Source/MassEQS/Public/Tests/MassEnvQueryTest_MassEntityTags.h
@@ -0,0 +1,68 @@
+// Copyright Epic Games, Inc. All Rights Reserved.
+
+#pragma once
+
+#include "Tests/MassEnvQueryTest.h"
+#include "MassEnvQueryTest_MassEntityTags.generated.h"
+
+struct FInstancedStruct;
+struct FMassEntityHandle;
+
+/** Different modes that this Test can be run in */
+UENUM(BlueprintType)
+enum class EMassEntityTagsTestMode : uint8
+{
+	Any = 0	UMETA(DisplayName = "Any Tags", Tooltip = "Filter will require just one of the tags to be present on the Entity."),
+	All		UMETA(DisplayName = "All Tags", Tooltip = "Filter will require All of the tags to be present on the Entity."),
+	None	UMETA(DisplayName = "None of the Tags", Tooltip = "Filter will require that none of the tags are present on the Entity.")
+};
+
+/**
+ * Test to be sent to MassEQSSubsystem for processing on Mass.
+ * This will test the Entities in the QueryInstance based on the MassTags they have in comparison to the input Tags, and the TagTestMode selected.
+ */
+UCLASS(MinimalAPI)
+class UMassEnvQueryTest_MassEntityTags : public UMassEnvQueryTest
+{
+	GENERATED_UCLASS_BODY()
+public:
+	// Begin IMassEQSRequestInterface
+	virtual TUniquePtr<FMassEQSRequestData> GetRequestData(FEnvQueryInstance& QueryInstance) const override;
+	virtual UClass* GetRequestClass() const override { return StaticClass(); }
+
+	virtual bool TryAcquireResults(FEnvQueryInstance& QueryInstance) const override;
+	// ~IMassEQSRequestInterface
+
+protected:
+	virtual FText GetDescriptionTitle() const override;
+	virtual FText GetDescriptionDetails() const override;
+
+	UPROPERTY(EditAnywhere, Category = "MassEntityTagsTest")
+	EMassEntityTagsTestMode TagTestMode = EMassEntityTagsTestMode::All;
+
+	UPROPERTY(EditAnywhere, Category = "MassEntityTagsTest", meta = (BaseStruct = "/Script/MassEntity.MassTag", ExcludeBaseStruct))
+	TArray<FInstancedStruct> Tags;
+};
+
+/** Data required to be sent to Mass for processing this Test Request */
+struct FMassEQSRequestData_MassEntityTags : public FMassEQSRequestData
+{
+	FMassEQSRequestData_MassEntityTags(EMassEntityTagsTestMode InTagTestMode, TArray<FInstancedStruct> InTags)
+		: TagTestMode(InTagTestMode)
+		, Tags(InTags)
+	{
+	}
+	EMassEntityTagsTestMode TagTestMode;
+	TArray<FInstancedStruct> Tags;
+};
+
+/** Data required to be sent to Mass for processing this Test Request */
+struct FMassEnvQueryResultData_MassEntityTags : public FMassEQSRequestData
+{
+	FMassEnvQueryResultData_MassEntityTags(TMap<FMassEntityHandle, bool>&& InResultMap)
+		: ResultMap(MoveTemp(InResultMap))
+	{
+	}
+
+	TMap<FMassEntityHandle, bool> ResultMap;
+};
\ No newline at end of file
diff --git a/Engine/Source/Runtime/AIModule/Classes/EnvironmentQuery/EnvQueryGenerator.h b/Engine/Source/Runtime/AIModule/Classes/EnvironmentQuery/EnvQueryGenerator.h
--- a/Engine/Source/Runtime/AIModule/Classes/EnvironmentQuery/EnvQueryGenerator.h
+++ b/Engine/Source/Runtime/AIModule/Classes/EnvironmentQuery/EnvQueryGenerator.h
@@ -35,7 +35,13 @@ class UEnvQueryGenerator : public UEnvQueryNode
 
 	virtual void GenerateItems(FEnvQueryInstance& QueryInstance) const { checkNoEntry(); }
 	virtual bool IsValidGenerator() const { return ItemType != nullptr; }
+	bool CanRunAsync() const { return bCanRunAsync; }
 
 	AIMODULE_API virtual void PostLoad() override;
 	AIMODULE_API void UpdateNodeVersion() override;
+
+protected:
+	/** To be overwritten by MassEnvQueryGenerators to indicate that they will run asynchronously. */
+	UPROPERTY(EditDefaultsOnly, Category = Option, AdvancedDisplay)
+	uint32 bCanRunAsync : 1;
 };
diff --git a/Engine/Source/Runtime/AIModule/Classes/EnvironmentQuery/EnvQueryInstanceBlueprintWrapper.h b/Engine/Source/Runtime/AIModule/Classes/EnvironmentQuery/EnvQueryInstanceBlueprintWrapper.h
--- a/Engine/Source/Runtime/AIModule/Classes/EnvironmentQuery/EnvQueryInstanceBlueprintWrapper.h
+++ b/Engine/Source/Runtime/AIModule/Classes/EnvironmentQuery/EnvQueryInstanceBlueprintWrapper.h
@@ -83,6 +83,9 @@ public:
 
 	AIMODULE_API void SetInstigator(const UObject* Object);
 
+	/** Used in UMassEQSBlueprintLibrary::GetEnviromentQueryResultAsEntityInfo */
+	AIMODULE_API FORCEINLINE EEnvQueryRunMode::Type GetRunMode() const { return RunMode; };
+
 	FEQSQueryDoneSignature& GetOnQueryFinishedEvent() { return OnQueryFinishedEvent; }
 
 protected:
diff --git a/Engine/Source/Runtime/AIModule/Classes/EnvironmentQuery/EnvQueryNode.h b/Engine/Source/Runtime/AIModule/Classes/EnvironmentQuery/EnvQueryNode.h
--- a/Engine/Source/Runtime/AIModule/Classes/EnvironmentQuery/EnvQueryNode.h
+++ b/Engine/Source/Runtime/AIModule/Classes/EnvironmentQuery/EnvQueryNode.h
@@ -24,6 +24,13 @@ class UEnvQueryNode : public UObject
 	AIMODULE_API virtual FText GetDescriptionTitle() const;
 	AIMODULE_API virtual FText GetDescriptionDetails() const;
 
+	/**
+	 * To be extended by any Node who offloads its work to another thread.
+	 * Returns false by default, unless overridden.
+	 * If overridden, will return whether or not this Node is currently being processed asynchronously. 
+	 */
+	AIMODULE_API virtual FORCEINLINE bool IsCurrentlyRunningAsync() const { return false; }
+
 #if WITH_EDITOR
 	AIMODULE_API virtual void PostEditChangeProperty(FPropertyChangedEvent& PropertyChangedEvent) override;
 #endif //WITH_EDITOR
diff --git a/Engine/Source/Runtime/AIModule/Classes/EnvironmentQuery/EnvQueryTypes.h b/Engine/Source/Runtime/AIModule/Classes/EnvironmentQuery/EnvQueryTypes.h
--- a/Engine/Source/Runtime/AIModule/Classes/EnvironmentQuery/EnvQueryTypes.h
+++ b/Engine/Source/Runtime/AIModule/Classes/EnvironmentQuery/EnvQueryTypes.h
@@ -811,6 +811,9 @@ struct FEnvQueryInstance : public FEnvQueryResult
 	/** true if this query has logged a warning that it overran the time limit */
 	uint8 bHasLoggedTimeLimitWarning : 1;
 
+	/** true if current Generator or Test is running asynchronously */
+	uint8 bIsCurrentlyRunningAsync : 1;
+
 	/** timestamp of creating query instance */
 	double StartTime;
 
@@ -855,6 +858,8 @@ struct FEnvQueryInstance : public FEnvQueryResult
 	/** check if current test can batch its calculations */
 	bool CanBatchTest() const { return !IsInSingleItemFinalSearch(); }
 
+	bool IsCurrentlyRunningAsync() const { return bIsCurrentlyRunningAsync; }
+
 	/** raw data operations */
 	AIMODULE_API void ReserveItemData(int32 NumAdditionalItems);
 
diff --git a/Engine/Source/Runtime/AIModule/Private/EnvironmentQuery/EnvQueryGenerator.cpp b/Engine/Source/Runtime/AIModule/Private/EnvironmentQuery/EnvQueryGenerator.cpp
--- a/Engine/Source/Runtime/AIModule/Private/EnvironmentQuery/EnvQueryGenerator.cpp
+++ b/Engine/Source/Runtime/AIModule/Private/EnvironmentQuery/EnvQueryGenerator.cpp
@@ -7,6 +7,7 @@
 UEnvQueryGenerator::UEnvQueryGenerator(const FObjectInitializer& ObjectInitializer) : Super(ObjectInitializer)
 {
 	bAutoSortTests = true;
+	bCanRunAsync = false;
 }
 
 void UEnvQueryGenerator::UpdateNodeVersion()
diff --git a/Engine/Source/Runtime/AIModule/Private/EnvironmentQuery/EnvQueryInstance.cpp b/Engine/Source/Runtime/AIModule/Private/EnvironmentQuery/EnvQueryInstance.cpp
--- a/Engine/Source/Runtime/AIModule/Private/EnvironmentQuery/EnvQueryInstance.cpp
+++ b/Engine/Source/Runtime/AIModule/Private/EnvironmentQuery/EnvQueryInstance.cpp
@@ -310,77 +310,83 @@ void FEnvQueryInstance::ExecuteOneStep(double TimeLimit)
 
 	const bool bDoingLastTest = (CurrentTest >= OptionItem.Tests.Num() - 1);
 	bool bStepDone = true;
+	bIsCurrentlyRunningAsync = false;
 	CurrentStepTimeLimit = TimeLimit;
 
 	if (CurrentTest < 0)
 	{
-		SCOPE_CYCLE_COUNTER(STAT_AI_EQS_GeneratorTime);
-		DEC_DWORD_STAT_BY(STAT_AI_EQS_NumItems, Items.Num());
-
-		RawData.Reset();
-		Items.Reset();
-		ItemType = OptionItem.ItemType;
-		bPassOnSingleResult = false;
-		ValueSize = (ItemType->GetDefaultObject<UEnvQueryItemType>())->GetValueSize();
-
 		bool bRunGenerator = true;
 #if USE_EQS_DEBUGGER
 		int32 LastValidItems = 0;
-		if (bStoreDebugInfo)
+#endif
+		if (!OptionItem.Generator->IsCurrentlyRunningAsync())
 		{
-			UEnvQueryGenerator_Composite* CompositeGen = Cast<UEnvQueryGenerator_Composite>(OptionItem.Generator);
-			TArray<UEnvQueryGenerator*> GeneratorList;
+			SCOPE_CYCLE_COUNTER(STAT_AI_EQS_GeneratorTime);
+			DEC_DWORD_STAT_BY(STAT_AI_EQS_NumItems, Items.Num());
 
-			if (CompositeGen)
+			RawData.Reset();
+			Items.Reset();
+			ItemType = OptionItem.ItemType;
+			bPassOnSingleResult = false;
+			ValueSize = (ItemType->GetDefaultObject<UEnvQueryItemType>())->GetValueSize();
+
+#if USE_EQS_DEBUGGER
+			if (bStoreDebugInfo)
 			{
-				// resolve nested composites while on it
-				GeneratorList.Append(CompositeGen->Generators);
-				for (int32 InnerIdx = 0; InnerIdx < GeneratorList.Num(); InnerIdx++)
+				UEnvQueryGenerator_Composite* CompositeGen = Cast<UEnvQueryGenerator_Composite>(OptionItem.Generator);
+				TArray<UEnvQueryGenerator*> GeneratorList;
+
+				if (CompositeGen)
 				{
-					UEnvQueryGenerator_Composite* InnerCompositeGen = Cast<UEnvQueryGenerator_Composite>(GeneratorList[InnerIdx]);
-					if (InnerCompositeGen)
+					// resolve nested composites while on it
+					GeneratorList.Append(CompositeGen->Generators);
+					for (int32 InnerIdx = 0; InnerIdx < GeneratorList.Num(); InnerIdx++)
 					{
-						GeneratorList.Append(InnerCompositeGen->Generators);
-						GeneratorList.RemoveAt(InnerIdx, EAllowShrinking::No);
-						InnerIdx--;
+						UEnvQueryGenerator_Composite* InnerCompositeGen = Cast<UEnvQueryGenerator_Composite>(GeneratorList[InnerIdx]);
+						if (InnerCompositeGen)
+						{
+							GeneratorList.Append(InnerCompositeGen->Generators);
+							GeneratorList.RemoveAt(InnerIdx, EAllowShrinking::No);
+							InnerIdx--;
+						}
 					}
 				}
-			}
 
-			DebugData.PrepareOption(*this, GeneratorList, OptionItem.Tests.Num());
+				DebugData.PrepareOption(*this, GeneratorList, OptionItem.Tests.Num());
 
-			// special case for composite generator: run each inner generator separately and record times
-			if (GeneratorList.Num())
-			{
-				bRunGenerator = false;
-				DebugData.CurrentOptionGeneratorIdx = 0;
-				for (int32 GeneratorIdx = 0; GeneratorIdx < GeneratorList.Num() - 1; GeneratorIdx++)
+				// special case for composite generator: run each inner generator separately and record times
+				if (GeneratorList.Num())
 				{
+					bRunGenerator = false;
+					DebugData.CurrentOptionGeneratorIdx = 0;
+					for (int32 GeneratorIdx = 0; GeneratorIdx < GeneratorList.Num() - 1; GeneratorIdx++)
 					{
-						FScopeCycleCounterUObject GeneratorScope(GeneratorList[GeneratorIdx]);
-						GeneratorList[GeneratorIdx]->GenerateItems(*this);
+						{
+							FScopeCycleCounterUObject GeneratorScope(GeneratorList[GeneratorIdx]);
+							GeneratorList[GeneratorIdx]->GenerateItems(*this);
+						}
+
+						const double GenTime = FPlatformTime::Seconds();
+						const double StepExecutionTime = GenTime - StepStartTime;
+						StepStartTime += StepExecutionTime;
+						TotalExecutionTime += StepExecutionTime;
+						StartTime = GenTime;
+						NumProcessedItems = Items.Num() - LastValidItems;
+						LastValidItems = Items.Num();
+						DebugData.CurrentOptionGeneratorIdx++;
+						DebugData.Store(*this, FloatCastChecked<float>(StepExecutionTime, /* Precision */ 1. / 512.), false);
+						NumProcessedItems = 0;
 					}
 
-					const double GenTime = FPlatformTime::Seconds();
-					const double StepExecutionTime = GenTime - StepStartTime;
-					StepStartTime += StepExecutionTime;
-					TotalExecutionTime += StepExecutionTime;
-					StartTime = GenTime;
-					NumProcessedItems = Items.Num() - LastValidItems;
-					LastValidItems = Items.Num();
-					DebugData.CurrentOptionGeneratorIdx++;
-					DebugData.Store(*this, FloatCastChecked<float>(StepExecutionTime, /* Precision */ 1./512.), false);
-					NumProcessedItems = 0;
-				}
-
-				{
-					FScopeCycleCounterUObject GeneratorScope(GeneratorList.Last());
-					GeneratorList.Last()->GenerateItems(*this);
+					{
+						FScopeCycleCounterUObject GeneratorScope(GeneratorList.Last());
+						GeneratorList.Last()->GenerateItems(*this);
+					}
+					DebugData.CurrentOptionGeneratorIdx = INDEX_NONE;
 				}
-				DebugData.CurrentOptionGeneratorIdx = INDEX_NONE;
 			}
-		}
 #endif // USE_EQS_DEBUGGER
+		}
 
 		if (bRunGenerator)
 		{
@@ -400,7 +406,13 @@ void FEnvQueryInstance::ExecuteOneStep(double TimeLimit)
 #endif // UE_BUILD_SHIPPING
 		}
 
-		FinalizeGeneration();
+		bIsCurrentlyRunningAsync |= OptionItem.Generator->IsCurrentlyRunningAsync();
+		bStepDone = !bIsCurrentlyRunningAsync;
+
+		if (bStepDone)
+		{
+			FinalizeGeneration();
+		}
 
 #if USE_EQS_DEBUGGER
 		NumProcessedItems = Items.Num() - LastValidItems;
@@ -412,26 +424,29 @@ void FEnvQueryInstance::ExecuteOneStep(double TimeLimit)
 
 		UEnvQueryTest* TestObject = OptionItem.Tests[CurrentTest];
 
-		// item generator uses this flag to alter the scoring behavior
-		bPassOnSingleResult = (bDoingLastTest && Mode == EEnvQueryRunMode::SingleResult && TestObject->CanRunAsFinalCondition());
-
-		if (bPassOnSingleResult && (CurrentTestStartingItem == 0))
+		if (!TestObject->IsCurrentlyRunningAsync())
 		{
-			// Since we know we're the last test that is a final condition, if we were scoring previously we should sort the tests now before we test them
-			bool bSortTests = false;
-			for (int32 TestIndex = 0; TestIndex < OptionItem.Tests.Num() - 1; ++TestIndex)
+			// item generator uses this flag to alter the scoring behavior
+			bPassOnSingleResult = (bDoingLastTest && Mode == EEnvQueryRunMode::SingleResult && TestObject->CanRunAsFinalCondition());
+
+			if (bPassOnSingleResult && (CurrentTestStartingItem == 0))
 			{
-				if (OptionItem.Tests[TestIndex]->TestPurpose != EEnvTestPurpose::Filter)
+				// Since we know we're the last test that is a final condition, if we were scoring previously we should sort the tests now before we test them
+				bool bSortTests = false;
+				for (int32 TestIndex = 0; TestIndex < OptionItem.Tests.Num() - 1; ++TestIndex)
 				{
-					// Found one.  We should sort.
-					bSortTests = true;
-					break;
+					if (OptionItem.Tests[TestIndex]->TestPurpose != EEnvTestPurpose::Filter)
+					{
+						// Found one.  We should sort.
+						bSortTests = true;
+						break;
+					}
 				}
-			}
 
-			if (bSortTests)
-			{
-				SortScores();
+				if (bSortTests)
+				{
+					SortScores();
+				}
 			}
 		}
 
@@ -442,9 +457,10 @@ void FEnvQueryInstance::ExecuteOneStep(double TimeLimit)
 			TestObject->RunTest(*this);
 		}
 
-		bStepDone = CurrentTestStartingItem >= Items.Num() || bFoundSingleResult
-			// or no items processed ==> this means error
-			|| (ItemsAlreadyProcessed == CurrentTestStartingItem);
+		bIsCurrentlyRunningAsync |= TestObject->IsCurrentlyRunningAsync();
+		bStepDone = !bIsCurrentlyRunningAsync && (CurrentTestStartingItem >= Items.Num() || bFoundSingleResult
+			// or no items processed ==> this means error, unless this test is running async
+			|| (ItemsAlreadyProcessed == CurrentTestStartingItem));
 
 		if (bStepDone)
 		{
@@ -461,7 +477,7 @@ void FEnvQueryInstance::ExecuteOneStep(double TimeLimit)
 	TotalExecutionTime += StepExecutionTime;
 
 #if USE_EQS_DEBUGGER
-	if (bStoreDebugInfo)
+	if (bStoreDebugInfo && !bIsCurrentlyRunningAsync)
 	{
 		DebugData.Store(*this, FloatCastChecked<float>(StepExecutionTime, /* Precision */ 1./512.), bStepDone);
 	}
@@ -477,7 +493,7 @@ void FEnvQueryInstance::ExecuteOneStep(double TimeLimit)
 	}
 
 	// sort results or switch to next option when all tests are performed
-	if (IsFinished() == false && (OptionItem.Tests.Num() == CurrentTest || NumValidItems <= 0))
+	if (!bIsCurrentlyRunningAsync && IsFinished() == false && (OptionItem.Tests.Num() == CurrentTest || NumValidItems <= 0))
 	{
 		if (NumValidItems > 0)
 		{
diff --git a/Engine/Source/Runtime/AIModule/Private/EnvironmentQuery/EnvQueryManager.cpp b/Engine/Source/Runtime/AIModule/Private/EnvironmentQuery/EnvQueryManager.cpp
--- a/Engine/Source/Runtime/AIModule/Private/EnvironmentQuery/EnvQueryManager.cpp
+++ b/Engine/Source/Runtime/AIModule/Private/EnvironmentQuery/EnvQueryManager.cpp
@@ -448,13 +448,14 @@ void UEnvQueryManager::Tick(float DeltaTime)
 		SCOPE_CYCLE_COUNTER(STAT_AI_EQS_TickWork);
 		
 		const int32 NumRunningQueries = RunningQueries.Num();
+		int32 NumQueriesRunningAsync = 0;
 		int32 Index = 0;
 
 		while ((TimeLeft > 0.0) 
 			&& (Index < NumRunningQueries) 
 			// make sure we account for queries that have finished (been aborted)
 			// before UEnvQueryManager::Tick has been called
-			&& (QueriesFinishedDuringUpdate + NumRunningQueriesAbortedSinceLastUpdate < NumRunningQueries))
+			&& (QueriesFinishedDuringUpdate + NumRunningQueriesAbortedSinceLastUpdate + NumQueriesRunningAsync < NumRunningQueries))
 		{
 			const double StepStartTime = FPlatformTime::Seconds();
 			double ResultHandlingDuration = 0.;
@@ -535,6 +536,13 @@ void UEnvQueryManager::Tick(float DeltaTime)
 					++QueriesFinishedDuringUpdate;
 					++Index;
 				}
+
+				// If current Query is waiting on results to be processed asynchronously, check for results next frame.
+				else if (QueryInstancePtr->IsCurrentlyRunningAsync())
+				{
+					++NumQueriesRunningAsync;
+					++Index;
+				}
 				// If we're testing queries using breadth, move on to the next query.
 				// If we're testing queries using depth, we only move on to the next query when we finish the current one.
 				else if (bTestQueriesUsingBreadth)
diff --git a/Engine/Source/Runtime/AIModule/Private/EnvironmentQuery/EnvQueryTypes.cpp b/Engine/Source/Runtime/AIModule/Private/EnvironmentQuery/EnvQueryTypes.cpp
--- a/Engine/Source/Runtime/AIModule/Private/EnvironmentQuery/EnvQueryTypes.cpp
+++ b/Engine/Source/Runtime/AIModule/Private/EnvironmentQuery/EnvQueryTypes.cpp
@@ -324,6 +324,7 @@ FEnvQueryInstance::FEnvQueryInstance()
 	, bFoundSingleResult(false)
 	, bPassOnSingleResult(false)
 	, bHasLoggedTimeLimitWarning(false)
+	, bIsCurrentlyRunningAsync(false)
 	, StartTime(0)
 	, TotalExecutionTime(0)
 	, CurrentStepTimeLimit(0)
diff --git a/Engine/Source/Runtime/AIModule/Private/EnvironmentQuery/Generators/EnvQueryGenerator_Composite.cpp b/Engine/Source/Runtime/AIModule/Private/EnvironmentQuery/Generators/EnvQueryGenerator_Composite.cpp
--- a/Engine/Source/Runtime/AIModule/Private/EnvironmentQuery/Generators/EnvQueryGenerator_Composite.cpp
+++ b/Engine/Source/Runtime/AIModule/Private/EnvironmentQuery/Generators/EnvQueryGenerator_Composite.cpp
@@ -23,6 +23,7 @@ void UEnvQueryGenerator_Composite::GenerateItems(FEnvQueryInstance& QueryInstanc
 			{
 				FScopeCycleCounterUObject GeneratorScope(Generators[Idx]);
 				Generators[Idx]->GenerateItems(QueryInstance);
+				ensureMsgf(!Generators[Idx]->IsCurrentlyRunningAsync(), TEXT("Composite Generators do not yet support asynchronous processing."));
 			}
 		}
 	}
@@ -82,6 +83,8 @@ void UEnvQueryGenerator_Composite::VerifyItemTypes()
 		{
 			if (Generators[Idx])
 			{
+				checkf(!Generators[Idx]->CanRunAsync(), TEXT("Async Generators such as [%s] are not supported members of Composite Generators."), *Generators[Idx]->GetName())
+
 				if (CommonItemType)
 				{
 					if (CommonItemType != Generators[Idx]->ItemType)
