Fix crash in PathFollowingComponent when using Move to Location or Actor in a level without navigation data.

diff --git a/Engine/Source/Runtime/AIModule/Private/Navigation/PathFollowingComponent.cpp b/Engine/Source/Runtime/AIModule/Private/Navigation/PathFollowingComponent.cpp
--- a/Engine/Source/Runtime/AIModule/Private/Navigation/PathFollowingComponent.cpp
+++ b/Engine/Source/Runtime/AIModule/Private/Navigation/PathFollowingComponent.cpp
@@ -55,20 +55,25 @@ namespace UE::Navigation::Private
 			return Path->GetFilter();
 		}
 
-		if (RequestData.GetNavigationFilter() != nullptr)
+		if (MyNavData)
 		{
-			return UNavigationQueryFilter::GetQueryFilter(*MyNavData, Owner, RequestData.GetNavigationFilter());
-		}
+			if (RequestData.GetNavigationFilter() != nullptr)
+			{
+				return UNavigationQueryFilter::GetQueryFilter(*MyNavData, Owner, RequestData.GetNavigationFilter());
+			}
 
-		if (const AAIController* AIOwner = Cast<AAIController>(Owner))
-		{
-			if (AIOwner->GetDefaultNavigationFilterClass() != nullptr)
+			if (const AAIController* AIOwner = Cast<AAIController>(Owner))
 			{
-				return UNavigationQueryFilter::GetQueryFilter(*MyNavData, AIOwner->GetDefaultNavigationFilterClass());
+				if (AIOwner->GetDefaultNavigationFilterClass() != nullptr)
+				{
+					return UNavigationQueryFilter::GetQueryFilter(*MyNavData, AIOwner->GetDefaultNavigationFilterClass());
+				}
 			}
+
+			return MyNavData->GetDefaultQueryFilter();
 		}
 
-		return MyNavData->GetDefaultQueryFilter();
+		return FSharedConstNavQueryFilter();
 	}
 }
 
