Fixes crash like 'FGenericTeamId::GetAttitude' where 'call'.

diff --git a/Engine/Source/Runtime/AIModule/AIModule.Build.cs b/Engine/Source/Runtime/AIModule/AIModule.Build.cs
--- a/Engine/Source/Runtime/AIModule/AIModule.Build.cs
+++ b/Engine/Source/Runtime/AIModule/AIModule.Build.cs
@@ -50,6 +50,7 @@ namespace UnrealBuildTool.Rules
             }
 
             SetupGameplayDebuggerSupport(Target);
+            bAllowAutoRTFMInstrumentation = true;
         }
     }
 }
