[BehaviorTree] make sure to remove any scheduled updates when the tree is stopped

diff --git a/Engine/Source/Runtime/AIModule/Private/BehaviorTree/BehaviorTreeComponent.cpp b/Engine/Source/Runtime/AIModule/Private/BehaviorTree/BehaviorTreeComponent.cpp
--- a/Engine/Source/Runtime/AIModule/Private/BehaviorTree/BehaviorTreeComponent.cpp
+++ b/Engine/Source/Runtime/AIModule/Private/BehaviorTree/BehaviorTreeComponent.cpp
@@ -430,6 +430,9 @@ void UBehaviorTreeComponent::StopTree(EBTStopMode::Type StopMode)
 	bRequestedStop = false;
 	bIsRunning = false;
 	bWaitingForLatentAborts = false;
+
+	// make sure to not process scheduled ticks
+	ScheduleNextTick(FLT_MAX);
 }
 
 void UBehaviorTreeComponent::RestartTree(EBTRestartMode RestartMode /*= EBTRestartMode::SkipReAddedNodes*/)
@@ -1678,6 +1681,8 @@ void UBehaviorTreeComponent::ApplyDiscardedSearch()
 
 void UBehaviorTreeComponent::TickComponent(float DeltaTime, const ELevelTick TickType, FActorComponentTickFunction *ThisTickFunction)
 {
+	ensureMsgf(bIsRunning, TEXT("Tree should never be ticked if it is no longer running"));
+
 	// Tick can be optimized by the tick function to not be called every frame so we need
 	// to set the current frame delta time based on that information for other tick scenarios (e.g. manual ticking in unit tests)
 	const UWorld* World = GetWorld();
