[BehaviorTree] updated ensure condition to only detect tick after stopping the tree (a first tick is possible on non started tree)

diff --git a/Engine/Source/Runtime/AIModule/Private/BehaviorTree/BehaviorTreeComponent.cpp b/Engine/Source/Runtime/AIModule/Private/BehaviorTree/BehaviorTreeComponent.cpp
--- a/Engine/Source/Runtime/AIModule/Private/BehaviorTree/BehaviorTreeComponent.cpp
+++ b/Engine/Source/Runtime/AIModule/Private/BehaviorTree/BehaviorTreeComponent.cpp
@@ -1681,7 +1681,7 @@ void UBehaviorTreeComponent::ApplyDiscardedSearch()
 
 void UBehaviorTreeComponent::TickComponent(float DeltaTime, const ELevelTick TickType, FActorComponentTickFunction *ThisTickFunction)
 {
-	ensureMsgf(bIsRunning, TEXT("Tree should never be ticked if it is no longer running"));
+	ensureMsgf(bIsRunning || LastRequestedDeltaTimeGameTime == 0, TEXT("Tree should never be ticked if it is no longer running"));
 
 	// Tick can be optimized by the tick function to not be called every frame so we need
 	// to set the current frame delta time based on that information for other tick scenarios (e.g. manual ticking in unit tests)
