Add new functionalities to SimpleCellGrid:
- init with a number of cells instead of a bounding box
- change the vertical positioning of the grid
- queries to get the 2D box of a cell
- queries to get a box from a rectangle of cells, and a rectangle of cells from a box

diff --git a/Engine/Source/Runtime/AIModule/Public/SimpleCellGrid.h b/Engine/Source/Runtime/AIModule/Public/SimpleCellGrid.h
--- a/Engine/Source/Runtime/AIModule/Public/SimpleCellGrid.h
+++ b/Engine/Source/Runtime/AIModule/Public/SimpleCellGrid.h
@@ -15,6 +15,11 @@ struct FGridSize2D
 		: Width(InWidth), Height(InHeight)
 	{
 	}
+
+	bool operator==(const FGridSize2D& Other) const
+	{
+		return Width == Other.Width && Height == Other.Height;
+	}
 };
 
 /**	No virtuals on purpose */
@@ -41,7 +46,12 @@ public:
 	{
 	}
 
-	bool Init(float InCellSize, const FBox& Bounds)
+	/** Initialize the grid from a bounding box
+	 * @param InCellSize Size of a cell
+	 * @param Bounds Bounding box the grid needs to encapsulate
+	 * @return True if the grid was initialized properly
+	 */
+	bool Init(const float InCellSize, const FBox& Bounds)
 	{
 		if (InCellSize <= 0.0f || !Bounds.IsValid)
 		{
@@ -64,6 +74,45 @@ public:
 		return true;
 	}
 
+	/** Initialize the grid
+	 * @param InCellSize Size of a cell
+	 * @param InGridSize Amount of cells needed in each direction
+	 * @param InOrigin World location of the grid
+	 * @param VerticalBoundSize Size of the grid above and under the Origin
+	 * @return True if the grid was initialized properly
+	 */
+	bool Init(const float InCellSize, const FGridSize2D& InGridSize, const FVector& InOrigin, const float VerticalBoundSize)
+	{
+		if (InCellSize < 0 || InGridSize.Height == 0 || InGridSize.Width == 0 || VerticalBoundSize < 0)
+		{
+			return false;
+		}
+
+		GridCellSize = InCellSize;
+		GridSize = InGridSize;
+		BoundsSize = FVector(GridSize.Width * InCellSize, GridSize.Height * InCellSize, VerticalBoundSize);
+		Origin = InOrigin;
+		UpdateWorldBounds();
+
+		const uint64 TempCellCount = GridSize.Width * GridSize.Height;
+		const typename TArray<FCellType>::SizeType CellCount = FMath::Min(TempCellCount, (uint64)TNumericLimits<typename TArray<FCellType>::SizeType>::Max());
+		ensureMsgf(CellCount == TempCellCount, TEXT("Grid width and height are too big."));
+		Cells.AddDefaulted(CellCount);
+
+		return true;
+	}
+
+	/** Change the vertical position of the grid by providing an interval */
+	void SetVerticalInterval(const FFloatInterval& VerticalInterval)
+	{
+		if (VerticalInterval.IsValid())
+		{
+			BoundsSize.Z = VerticalInterval.Size();
+			Origin.Z = VerticalInterval.Interpolate(0.5);
+			UpdateWorldBounds();
+		}
+	}
+
 	void UpdateWorldBounds()
 	{
 		WorldBounds = FBox(Origin - FVector(0, 0, BoundsSize.Z / 2), Origin + FVector(BoundsSize.X, BoundsSize.Y, BoundsSize.Z / 2));
@@ -159,12 +208,14 @@ public:
 		const FIntVector CellCoords = GetCellCoordsUnsafe(WorldLocation);
 		return GetCellIndex(CellCoords.X, CellCoords.Y);
 	}
-	
+
+	/** Return the bounding box of a cell. */
 	FORCEINLINE FBox GetWorldCellBox(int32 CellIndex) const
 	{
 		return GetWorldCellBox(GetCellCoordX(CellIndex), GetCellCoordY(CellIndex));
 	}
 
+	/** Return the bounding box of a cell. */
 	FORCEINLINE FBox GetWorldCellBox(int32 LocationX, int32 LocationY) const
 	{
 		return FBox(
@@ -173,6 +224,49 @@ public:
 			);
 	}
 
+	/** Return the 2D bounding box of a cell. */
+	FORCEINLINE FBox2D GetWorldCellBox2D(int32 CellIndex) const
+	{
+		return GetWorldCellBox2D(GetCellCoordX(CellIndex), GetCellCoordY(CellIndex));
+	}
+
+	/** Return the 2D bounding box of a cell. */
+	FORCEINLINE FBox2D GetWorldCellBox2D(int32 LocationX, int32 LocationY) const
+	{
+		return FBox2D(
+			FVector2D(Origin) + FVector2D(LocationX * GridCellSize, LocationY * GridCellSize), 
+			FVector2D(Origin) + FVector2D((LocationX + 1) * GridCellSize, (LocationY + 1) * GridCellSize)
+			);
+	}
+
+	/** Return the world bounding box of all cells included in the given rectangle. */
+	FORCEINLINE FBox GetWorldCellRectangleBox(const FIntRect& CellRect) const
+	{
+		return FBox(
+			Origin + FVector(CellRect.Min.X * GridCellSize, CellRect.Min.Y * GridCellSize, -BoundsSize.Z * 0.5f), 
+			Origin + FVector((CellRect.Max.X+1) * GridCellSize, (CellRect.Max.Y+1) * GridCellSize, BoundsSize.Z * 0.5f)
+			);
+	}
+
+	/** Compute a rectangle of cells overlapping the given WorldBox. */
+	FORCEINLINE FIntRect GetCellRectangleFromBox(const FBox& WorldBox) const
+	{
+		if (!WorldBox.IsValid)
+		{
+			return FIntRect();
+		}
+
+		const FIntVector CellMin = GetCellCoordsUnsafe(WorldBox.Min);
+		const FIntVector CellMax = GetCellCoordsUnsafe(WorldBox.Max);
+		return FIntRect(CellMin.X, CellMin.Y, CellMax.X, CellMax.Y);
+	}
+
+	/** Return an IntRect that includes all the cells of the grid. Max is inclusive. */
+	FORCEINLINE FIntRect GetGridRectangle() const
+	{
+		return FIntRect(0, 0, GridSize.Width-1, GridSize.Height-1);
+	}
+
 	FORCEINLINE FVector GetWorldCellCenter(int32 CellIndex) const
 	{
 		return GetWorldCellCenter(GetCellCoordX(CellIndex), GetCellCoordY(CellIndex));
