[BehaviorTreeComponent] updated ensure to fix false positive when starting the tree.

diff --git a/Engine/Source/Runtime/AIModule/Private/BehaviorTree/BehaviorTreeComponent.cpp b/Engine/Source/Runtime/AIModule/Private/BehaviorTree/BehaviorTreeComponent.cpp
--- a/Engine/Source/Runtime/AIModule/Private/BehaviorTree/BehaviorTreeComponent.cpp
+++ b/Engine/Source/Runtime/AIModule/Private/BehaviorTree/BehaviorTreeComponent.cpp
@@ -33,8 +33,11 @@ double UBehaviorTreeComponent::FrameSearchTime = 0.;
 int32 UBehaviorTreeComponent::NumSearchTimeCalls = 0;
 #endif
 
-namespace UE::BehaviorTreeCVars
+namespace UE::BehaviorTree
 {
+	// Dedicated value used to stop ticking the tree
+	constexpr float DisableTick = FLT_MAX;
+
 	// Note this is defaulted to off for now as it caused a further bug, there is a BT unit test that will fire if this
 	// code is re-enabled, once that is addressed this can be re-enabled by default.
 	static bool bApplyAuxNodesFromFailedSearches = false;
@@ -432,7 +435,7 @@ void UBehaviorTreeComponent::StopTree(EBTStopMode::Type StopMode)
 	bWaitingForLatentAborts = false;
 
 	// make sure to not process scheduled ticks
-	ScheduleNextTick(FLT_MAX);
+	ScheduleNextTick(UE::BehaviorTree::DisableTick);
 }
 
 void UBehaviorTreeComponent::RestartTree(EBTRestartMode RestartMode /*= EBTRestartMode::SkipReAddedNodes*/)
@@ -1655,7 +1658,7 @@ void UBehaviorTreeComponent::ApplySearchData(UBTNode* NewActiveNode)
 
 void UBehaviorTreeComponent::ApplyDiscardedSearch()
 {
-	if (UE::BehaviorTreeCVars::bApplyAuxNodesFromFailedSearches)
+	if (UE::BehaviorTree::bApplyAuxNodesFromFailedSearches)
 	{
 		// Apply aux nodes from last search with the currently active node as the 'NewNode'
 		int32 NewNodeExecutionIdx = 0;
@@ -1681,7 +1684,7 @@ void UBehaviorTreeComponent::ApplyDiscardedSearch()
 
 void UBehaviorTreeComponent::TickComponent(float DeltaTime, const ELevelTick TickType, FActorComponentTickFunction *ThisTickFunction)
 {
-	ensureMsgf(bIsRunning || LastRequestedDeltaTimeGameTime == 0, TEXT("Tree should never be ticked if it is no longer running"));
+	ensureMsgf(bIsRunning || NextTickDeltaTime != UE::BehaviorTree::DisableTick, TEXT("Tree should never be ticked if it is no longer running"));
 
 	// Tick can be optimized by the tick function to not be called every frame so we need
 	// to set the current frame delta time based on that information for other tick scenarios (e.g. manual ticking in unit tests)
@@ -1736,7 +1739,7 @@ void UBehaviorTreeComponent::TickComponent(float DeltaTime, const ELevelTick Tic
 #endif
 
 	check(IsValid(this));
-	float NextNeededDeltaTime = FLT_MAX;
+	float NextNeededDeltaTime = UE::BehaviorTree::DisableTick;
 
 	checkf(PendingBranchActionRequests.Num() == 0, TEXT("Pending branches action requests should always be flushed immediately with the new system"))
 
@@ -1792,7 +1795,7 @@ void UBehaviorTreeComponent::TickComponent(float DeltaTime, const ELevelTick Tic
 
         // Since hierarchy might changed in the ProcessExecutionRequest, we need to go through all the active auxiliary nodes again to fetch new next DeltaTime
 		bActiveAuxiliaryNodeDTDirty = true;
-		NextNeededDeltaTime = FLT_MAX;
+		NextNeededDeltaTime = UE::BehaviorTree::DisableTick;
 	}
 
 	if (InstanceStack.Num() > 0 && bIsRunning && !bIsPaused)
@@ -1913,7 +1916,7 @@ void UBehaviorTreeComponent::ScheduleNextTick(const float NextNeededDeltaTime)
 	}
 
 	UE_VLOG(GetOwner(), LogBehaviorTree, VeryVerbose, TEXT("BT(%i) schedule next tick %f, asked %f."), GFrameCounter, NextTickDeltaTime, NextNeededDeltaTime);
-	if (NextTickDeltaTime == FLT_MAX)
+	if (NextTickDeltaTime == UE::BehaviorTree::DisableTick)
 	{
 		if (IsComponentTickEnabled())
 		{
