Exposing trace EQS helpers to AIModule_API

diff --git a/Engine/Source/Runtime/AIModule/Private/EnvironmentQuery/EnvQueryTraceHelpers.h b/Engine/Source/Runtime/AIModule/Private/EnvironmentQuery/EnvQueryTraceHelpers.h
--- a/Engine/Source/Runtime/AIModule/Private/EnvironmentQuery/EnvQueryTraceHelpers.h
+++ b/Engine/Source/Runtime/AIModule/Private/EnvironmentQuery/EnvQueryTraceHelpers.h
@@ -82,15 +82,15 @@ namespace FEQSHelpers
 		}
 	};
 
-	void RunNavRaycasts(const ANavigationData& NavData, const UObject& Querier, const FEnvTraceData& TraceData, const FVector& SourcePt, TArray<FNavLocation>& Points, const ETraceMode TraceMode = ETraceMode::Keep);
+	AIMODULE_API void RunNavRaycasts(const ANavigationData& NavData, const UObject& Querier, const FEnvTraceData& TraceData, const FVector& SourcePt, TArray<FNavLocation>& Points, const ETraceMode TraceMode = ETraceMode::Keep);
 	AIMODULE_API void RunNavProjection(const ANavigationData& NavData, const UObject& Querier, const FEnvTraceData& TraceData, TArray<FNavLocation>& Points, const ETraceMode TraceMode = ETraceMode::Discard);
-	void RunPhysRaycasts(UWorld* World, const FEnvTraceData& TraceData, const FVector& SourcePt, TArray<FNavLocation>& Points, const TArray<AActor*>& IgnoredActors, const ETraceMode TraceMode = ETraceMode::Keep);
-	void RunPhysProjection(UWorld* World, const FEnvTraceData& TraceData, TArray<FNavLocation>& Points, const ETraceMode TraceMode = ETraceMode::Discard);
-	void RunPhysProjection(UWorld* World, const FEnvTraceData& TraceData, TArray<FNavLocation>& Points, TArray<uint8>& TraceHits);
+	AIMODULE_API void RunPhysRaycasts(UWorld* World, const FEnvTraceData& TraceData, const FVector& SourcePt, TArray<FNavLocation>& Points, const TArray<AActor*>& IgnoredActors, const ETraceMode TraceMode = ETraceMode::Keep);
+	AIMODULE_API void RunPhysProjection(UWorld* World, const FEnvTraceData& TraceData, TArray<FNavLocation>& Points, const ETraceMode TraceMode = ETraceMode::Discard);
+	AIMODULE_API void RunPhysProjection(UWorld* World, const FEnvTraceData& TraceData, TArray<FNavLocation>& Points, TArray<uint8>& TraceHits);
 	/** Does initial raycast on navmesh, just like, RunNavRaycasts but
 	 *	once it hits navmesh edge it does a geometry trace to determine whether
 	 *	it hit a wall or a ledge (empty space). */
-	void RunRaycastsOnNavHitOnlyWalls(const ANavigationData& NavData, const UObject& Querier, const FEnvTraceData& TraceData, const FVector& SourcePt, TArray<FNavLocation>& Points, const TArray<AActor*>& IgnoredActors, const ETraceMode TraceMode = ETraceMode::Keep);
+	AIMODULE_API void RunRaycastsOnNavHitOnlyWalls(const ANavigationData& NavData, const UObject& Querier, const FEnvTraceData& TraceData, const FVector& SourcePt, TArray<FNavLocation>& Points, const TArray<AActor*>& IgnoredActors, const ETraceMode TraceMode = ETraceMode::Keep);
 
 	// deprecated
 
