[Blackboard] use DuplicateObject instead of NewObject when instantiating a key so all default values are properly copied.

diff --git a/Engine/Source/Runtime/AIModule/Private/BehaviorTree/Blackboard/BlackboardKeyType.cpp b/Engine/Source/Runtime/AIModule/Private/BehaviorTree/Blackboard/BlackboardKeyType.cpp
--- a/Engine/Source/Runtime/AIModule/Private/BehaviorTree/Blackboard/BlackboardKeyType.cpp
+++ b/Engine/Source/Runtime/AIModule/Private/BehaviorTree/Blackboard/BlackboardKeyType.cpp
@@ -25,7 +25,7 @@ void UBlackboardKeyType::InitializeKey(UBlackboardComponent& OwnerComp, FBlackbo
 	if (bCreateKeyInstance)
 	{
 		FBlackboardInstancedKeyMemory* MyMemory = (FBlackboardInstancedKeyMemory*)RawData;
-		UBlackboardKeyType* KeyInstance = NewObject<UBlackboardKeyType>(&OwnerComp, GetClass());
+		UBlackboardKeyType* KeyInstance = DuplicateObject<UBlackboardKeyType>(this, &OwnerComp);
 		KeyInstance->bIsInstanced = true;
 		MyMemory->KeyIdx = KeyID;
 		OwnerComp.KeyInstances[KeyID] = KeyInstance;
