CIS fix
#rnx

diff --git a/Engine/Source/Runtime/AIModule/Classes/BehaviorTree/Tasks/BTTask_MoveTo.h b/Engine/Source/Runtime/AIModule/Classes/BehaviorTree/Tasks/BTTask_MoveTo.h
--- a/Engine/Source/Runtime/AIModule/Classes/BehaviorTree/Tasks/BTTask_MoveTo.h
+++ b/Engine/Source/Runtime/AIModule/Classes/BehaviorTree/Tasks/BTTask_MoveTo.h
@@ -12,6 +12,7 @@
 #include "AITypes.h"
 #include "BehaviorTree/Tasks/BTTask_BlackboardBase.h"
 #include "BehaviorTree/ValueOrBBKey.h"
+#include "NavFilters/NavigationQueryFilter.h"
 #include "BTTask_MoveTo.generated.h"
 
 class UAITask_MoveTo;
