Add missing include
#rnx

diff --git a/Engine/Source/Runtime/AIModule/Classes/BehaviorTree/Decorators/BTDecorator_DoesPathExist.h b/Engine/Source/Runtime/AIModule/Classes/BehaviorTree/Decorators/BTDecorator_DoesPathExist.h
--- a/Engine/Source/Runtime/AIModule/Classes/BehaviorTree/Decorators/BTDecorator_DoesPathExist.h
+++ b/Engine/Source/Runtime/AIModule/Classes/BehaviorTree/Decorators/BTDecorator_DoesPathExist.h
@@ -11,6 +11,7 @@
 #include "BehaviorTree/BehaviorTreeTypes.h"
 #include "BehaviorTree/BTDecorator.h"
 #include "BehaviorTree/ValueOrBBKey.h"
+#include "NavFilters/NavigationQueryFilter.h"
 #include "BTDecorator_DoesPathExist.generated.h"
 
 class UBehaviorTree;
