Fix missing Type in SerializeFromMismatchedTag  for ValueOrBBKey_Enum
#rnx

diff --git a/Engine/Source/Runtime/AIModule/Private/BehaviorTree/ValueOrBBKey.cpp b/Engine/Source/Runtime/AIModule/Private/BehaviorTree/ValueOrBBKey.cpp
--- a/Engine/Source/Runtime/AIModule/Private/BehaviorTree/ValueOrBBKey.cpp
+++ b/Engine/Source/Runtime/AIModule/Private/BehaviorTree/ValueOrBBKey.cpp
@@ -341,7 +341,7 @@ bool FValueOrBBKey_Class::SerializeFromMismatchedTag(const FPropertyTag& Tag, FS
 
 bool FValueOrBBKey_Enum::SerializeFromMismatchedTag(const FPropertyTag& Tag, FStructuredArchive::FSlot Slot)
 {
-	if (Tag.Type == NAME_EnumProperty)
+	if (Tag.Type == NAME_EnumProperty || Tag.Type == NAME_ByteProperty || Tag.Type == NAME_Int8Property)
 	{
 		Slot << DefaultValue;
 		return true;
