Add missing includes to fix android compilation
#rnx

diff --git a/Engine/Source/Runtime/AIModule/Classes/BehaviorTree/Tasks/BTTask_PlayAnimation.h b/Engine/Source/Runtime/AIModule/Classes/BehaviorTree/Tasks/BTTask_PlayAnimation.h
--- a/Engine/Source/Runtime/AIModule/Classes/BehaviorTree/Tasks/BTTask_PlayAnimation.h
+++ b/Engine/Source/Runtime/AIModule/Classes/BehaviorTree/Tasks/BTTask_PlayAnimation.h
@@ -7,6 +7,7 @@
 #include "Engine/EngineTypes.h"
 #include "TimerManager.h"
 #include "Components/SkeletalMeshComponent.h"
+#include "Animation/AnimationAsset.h"
 #include "BehaviorTree/BTTaskNode.h"
 #include "BehaviorTree/ValueOrBBKey.h"
 #include "BTTask_PlayAnimation.generated.h"
diff --git a/Engine/Source/Runtime/AIModule/Classes/BehaviorTree/Tasks/BTTask_PlaySound.h b/Engine/Source/Runtime/AIModule/Classes/BehaviorTree/Tasks/BTTask_PlaySound.h
--- a/Engine/Source/Runtime/AIModule/Classes/BehaviorTree/Tasks/BTTask_PlaySound.h
+++ b/Engine/Source/Runtime/AIModule/Classes/BehaviorTree/Tasks/BTTask_PlaySound.h
@@ -6,6 +6,7 @@
 #include "UObject/ObjectMacros.h"
 #include "BehaviorTree/BTTaskNode.h"
 #include "BehaviorTree/ValueOrBBKey.h"
+#include "Sound/SoundCue.h"
 #include "BTTask_PlaySound.generated.h"
 
 class USoundCue;
