Properly serializefrommismatchedtag for enum and byte property for ValueorBBKeyType
#rnx

diff --git a/Engine/Source/Runtime/AIModule/Private/BehaviorTree/ValueOrBBKey.cpp b/Engine/Source/Runtime/AIModule/Private/BehaviorTree/ValueOrBBKey.cpp
--- a/Engine/Source/Runtime/AIModule/Private/BehaviorTree/ValueOrBBKey.cpp
+++ b/Engine/Source/Runtime/AIModule/Private/BehaviorTree/ValueOrBBKey.cpp
@@ -341,10 +341,15 @@ bool FValueOrBBKey_Class::SerializeFromMismatchedTag(const FPropertyTag& Tag, FS
 
 bool FValueOrBBKey_Enum::SerializeFromMismatchedTag(const FPropertyTag& Tag, FStructuredArchive::FSlot Slot)
 {
-	if (Tag.Type == NAME_EnumProperty)
+	if ((Tag.Type == NAME_EnumProperty || Tag.Type == NAME_ByteProperty) && Tag.GetType().GetParameterCount() > 0)
 	{
-		Slot << DefaultValue;
-		return true;
+		if (const UEnum* Enum = FindFirstObject<UEnum>(*Tag.GetType().GetParameterName(0).ToString(), EFindFirstObjectOptions::NativeFirst))
+		{
+			FName EnumValue;
+			Slot << EnumValue;
+			DefaultValue = static_cast<uint8>(Enum->GetValueByName(EnumValue));
+			return true;
+		}
 	}
 	return false;
 }
