Creating a new NavMovement interface to allow other movement components to implement and handle pathfollowing driven nav movement. Also contains necessary changes to convert existing code to use this new interface.

diff --git a/Engine/Source/Runtime/AIModule/Classes/Navigation/PathFollowingComponent.h b/Engine/Source/Runtime/AIModule/Classes/Navigation/PathFollowingComponent.h
--- a/Engine/Source/Runtime/AIModule/Classes/Navigation/PathFollowingComponent.h
+++ b/Engine/Source/Runtime/AIModule/Classes/Navigation/PathFollowingComponent.h
@@ -274,8 +274,16 @@ class UPathFollowingComponent : public UActorComponent, public IAIResourceInterf
 	/** notify about changing current path: new pointer or update from path event */
 	AIMODULE_API virtual void OnPathUpdated();
 
-	/** set associated movement component */
-	AIMODULE_API virtual void SetMovementComponent(UNavMovementComponent* MoveComp);
+	/**
+	 * set associated movement component.
+	 * Note: This function is deprecated and was marked final to purposely alert users that this function should no longer be overriden
+	 * Users should instead override SetNavMovementInterface(INavMovementInterface* NavMovementInterface)!
+	 */
+	UE_DEPRECATED(5.5, "SetMovementComponent(UNavMovementComponent* MoveComp) is deprecated, please use SetNavMoveInterface(INavMoveInterface* NavMoveInterface) instead.")
+	AIMODULE_API virtual void SetMovementComponent(UNavMovementComponent* MoveComp) final;
+
+	/** set associated nav movement interface */
+	AIMODULE_API virtual void SetNavMovementInterface(INavMovementInterface* NavMoveInterface);
 
 	/** get current focal point of movement */
 	AIMODULE_API virtual FVector GetMoveFocus(bool bAllowStrafe) const;
@@ -351,7 +359,7 @@ class UPathFollowingComponent : public UActorComponent, public IAIResourceInterf
 	FVector GetCurrentMoveInput() const { return CurrentMoveInput; }
 
 	/** check if path following has authority over movement (e.g. not falling) and can update own state */
-	FORCEINLINE bool HasMovementAuthority() const { return (MovementComp == nullptr) || MovementComp->CanStopPathFollowing(); }
+	FORCEINLINE bool HasMovementAuthority() const { return (NavMovementInterface == nullptr) || NavMovementInterface->CanStopPathFollowing(); }
 
 	FORCEINLINE const FNavPathSharedPtr GetPath() const { return Path; }
 	FORCEINLINE bool HasValidPath() const { return Path.IsValid() && Path->IsValid(); }
@@ -423,9 +431,13 @@ class UPathFollowingComponent : public UActorComponent, public IAIResourceInterf
 
 protected:
 
-	/** associated movement component */
-	UPROPERTY(transient)
+	/** deprecated associated movement component */
+	UE_DEPRECATED(5.5, "MovementComp is deprecated, please use NavMovementComp and the INavMoveInterface instead.")
+	UPROPERTY(transient, meta = (DeprecatedProperty, DeprecationMessage = "MovementComp is deprecated, please use NavMovementInterface and the INavMoveInterface instead."))
 	TObjectPtr<UNavMovementComponent> MovementComp;
+	
+	/** associated movement interface */
+	TWeakInterfacePtr<INavMovementInterface> NavMovementInterface;
 
 	/** currently traversed custom nav link */
 	FWeakObjectPtr CurrentCustomLinkOb;
diff --git a/Engine/Source/Runtime/AIModule/Private/Navigation/CrowdFollowingComponent.cpp b/Engine/Source/Runtime/AIModule/Private/Navigation/CrowdFollowingComponent.cpp
--- a/Engine/Source/Runtime/AIModule/Private/Navigation/CrowdFollowingComponent.cpp
+++ b/Engine/Source/Runtime/AIModule/Private/Navigation/CrowdFollowingComponent.cpp
@@ -48,27 +48,30 @@ UCrowdFollowingComponent::UCrowdFollowingComponent(const FObjectInitializer& Obj
 
 FVector UCrowdFollowingComponent::GetCrowdAgentLocation() const
 {
-	return MovementComp ? MovementComp->GetActorFeetLocation() : FVector::ZeroVector;
+	return NavMovementInterface.IsValid() ? NavMovementInterface->GetFeetLocation() : FVector::ZeroVector;
 }
 
 FVector UCrowdFollowingComponent::GetCrowdAgentVelocity() const
 {
-	FVector Velocity(MovementComp ? MovementComp->Velocity : FVector::ZeroVector);
+	FVector Velocity(NavMovementInterface.IsValid() ? NavMovementInterface->GetVelocityForNavMovement() : FVector::ZeroVector);
 	Velocity *= (Status == EPathFollowingStatus::Moving) ? 1.0f : 0.25f;
 	return Velocity;
 }
 
 void UCrowdFollowingComponent::GetCrowdAgentCollisions(float& CylinderRadius, float& CylinderHalfHeight) const
 {
-	if (MovementComp && MovementComp->UpdatedComponent)
+	if (NavMovementInterface.IsValid() && NavMovementInterface->GetUpdatedObject())
 	{
-		MovementComp->UpdatedComponent->CalcBoundingCylinder(CylinderRadius, CylinderHalfHeight);
+		if (const USceneComponent* UpdatedComponent = Cast<USceneComponent>(NavMovementInterface->GetUpdatedObject()))
+		{
+			UpdatedComponent->CalcBoundingCylinder(CylinderRadius, CylinderHalfHeight);
+		}
 	}
 }
 
 float UCrowdFollowingComponent::GetCrowdAgentMaxSpeed() const
 {
-	return MovementComp ? MovementComp->GetMaxSpeed() : 0.0f;
+	return NavMovementInterface.IsValid() ? NavMovementInterface->GetMaxSpeedForNavMovement() : 0.0f;
 }
 
 int32 UCrowdFollowingComponent::GetCrowdAgentAvoidanceGroup() const
@@ -336,7 +339,7 @@ void UCrowdFollowingComponent::UpdateCachedDirections(const FVector& NewVelocity
 	}
 
 	// CrowdAgentMoveDirection either direction on path or aligned with current velocity
-	const bool bIsNotFalling = (MovementComp == nullptr || !MovementComp->IsFalling());
+	const bool bIsNotFalling = (NavMovementInterface == nullptr || !NavMovementInterface->IsFalling());
 	if (bIsNotFalling)
 	{
 		if (bUpdateDirectMoveVelocity)
@@ -380,14 +383,14 @@ void UCrowdFollowingComponent::UpdateDestinationForMovingGoal(const FVector& New
 void UCrowdFollowingComponent::ApplyCrowdAgentVelocity(const FVector& NewVelocity, const FVector& DestPathCorner, bool bTraversingLink, bool bIsNearEndOfPath)
 {
 	bCanCheckMovingTooFar = !bTraversingLink && bIsNearEndOfPath;
-	if (IsCrowdSimulationEnabled() && Status == EPathFollowingStatus::Moving && MovementComp)
+	if (IsCrowdSimulationEnabled() && Status == EPathFollowingStatus::Moving && NavMovementInterface.IsValid())
 	{
-		const bool bIsNotFalling = (MovementComp == nullptr || !MovementComp->IsFalling());
+		const bool bIsNotFalling = (NavMovementInterface == nullptr || !NavMovementInterface->IsFalling());
 		if (bAffectFallingVelocity || bIsNotFalling)
 		{
 			UpdateCachedDirections(NewVelocity, DestPathCorner, bTraversingLink);
 
-			const bool bAccelerationBased = MovementComp->UseAccelerationForPathFollowing();
+			const bool bAccelerationBased = NavMovementInterface->UseAccelerationForPathFollowing();
 			if (bAccelerationBased)
 			{
 				const FVector::FReal MaxSpeed = GetCrowdAgentMaxSpeed();
@@ -395,11 +398,11 @@ void UCrowdFollowingComponent::ApplyCrowdAgentVelocity(const FVector& NewVelocit
 				const FVector::FReal SpeedPct = FMath::Clamp(NewSpeed / MaxSpeed, 0., 1.);
 				const FVector MoveInput = FMath::IsNearlyZero(NewSpeed) ? FVector::ZeroVector : ((NewVelocity / NewSpeed) * SpeedPct);
 
-				MovementComp->RequestPathMove(MoveInput);
+				NavMovementInterface->RequestPathMove(MoveInput);
 			}
 			else
 			{
-				MovementComp->RequestDirectMove(NewVelocity, false);
+				NavMovementInterface->RequestDirectMove(NewVelocity, false);
 			}
 		}
 	}
@@ -571,7 +574,7 @@ bool UCrowdFollowingComponent::UpdateMovementComponent(bool bForce)
 {
 	bool bRet = Super::UpdateMovementComponent(bForce);
 
-	AvoidanceInterface = TWeakInterfacePtr<IRVOAvoidanceInterface>(MovementComp);
+	AvoidanceInterface = TWeakInterfacePtr<IRVOAvoidanceInterface>(NavMovementInterface.GetObject());
 
 	return bRet;
 }
@@ -863,7 +866,7 @@ void UCrowdFollowingComponent::SetMoveSegment(int32 SegmentStartIndex)
 		CurrentDestination.Set(Path->GetBaseActor(), CurrentTargetPt);
 
 		LogPathPartHelper(GetOwner(), NavMeshPath, PathStartIndex, PathPartEndIdx);
-		UE_VLOG_SEGMENT(GetOwner(), LogCrowdFollowing, Log, MovementComp->GetActorFeetLocation(), CurrentTargetPt, FColor::Red, TEXT("path part"));
+		UE_VLOG_SEGMENT(GetOwner(), LogCrowdFollowing, Log, NavMovementInterface->GetFeetLocation(), CurrentTargetPt, FColor::Red, TEXT("path part"));
 		UE_VLOG(GetOwner(), LogCrowdFollowing, Log, TEXT("SetMoveSegment, from:%d segments:%d%s"),
 			PathStartIndex, (PathPartEndIdx - PathStartIndex)+1, bFinalPathPart ? TEXT(" (final)") : TEXT(""));
 
@@ -875,7 +878,7 @@ void UCrowdFollowingComponent::SetMoveSegment(int32 SegmentStartIndex)
 		// direct paths are not using any steering or avoidance
 		// pathfinding is replaced with simple velocity request 
 
-		const FVector AgentLoc = MovementComp->GetActorFeetLocation();
+		const FVector AgentLoc = NavMovementInterface->GetFeetLocation();
 
 		bFinalPathPart = true;
 		bCheckMovementAngle = true;
@@ -907,7 +910,7 @@ void UCrowdFollowingComponent::UpdatePathSegment()
 		return;
 	}
 
-	if (!Path.IsValid() || MovementComp == NULL)
+	if (!Path.IsValid() || NavMovementInterface == NULL)
 	{
 		OnPathFinished(FPathFollowingResult(EPathFollowingResult::Aborted, FPathFollowingResultFlags::InvalidPath));
 		return;
@@ -930,8 +933,8 @@ void UCrowdFollowingComponent::UpdatePathSegment()
 	}
 
 	// if agent has control over its movement, check finish conditions
-	const FVector CurrentLocation = MovementComp->GetActorFeetLocation();
-	const bool bCanReachTarget = MovementComp->CanStopPathFollowing();
+	const FVector CurrentLocation = NavMovementInterface->GetFeetLocation();
+	const bool bCanReachTarget = NavMovementInterface->CanStopPathFollowing();
 	if (bCanReachTarget && Status == EPathFollowingStatus::Moving)
 	{
 		const FVector GoalLocation = GetCurrentTargetLocation();
@@ -944,20 +947,20 @@ void UCrowdFollowingComponent::UpdatePathSegment()
 		}
 		else if (bFinalPathPart)
 		{
-			const FVector ToTarget = (GoalLocation - MovementComp->GetActorFeetLocation());
+			const FVector ToTarget = (GoalLocation - NavMovementInterface->GetFeetLocation());
 			const bool bMovedTooFar = (bCheckMovementAngle || bCanCheckMovingTooFar) && !CrowdAgentMoveDirection.IsNearlyZero() && FVector::DotProduct(ToTarget, CrowdAgentMoveDirection) < 0.0;
 
 #if ENABLE_VISUAL_LOG
 			if (bMovedTooFar)
 			{
-				const FVector AgentLoc = MovementComp->GetActorFeetLocation();
+				const FVector AgentLoc = NavMovementInterface->GetFeetLocation();
 				UE_VLOG_SEGMENT(GetOwner(), LogCrowdFollowing, Log, AgentLoc, AgentLoc + CrowdAgentMoveDirection * 100.0f, FColor::Cyan, TEXT("moveDir"));
 				UE_VLOG_SEGMENT(GetOwner(), LogCrowdFollowing, Log, AgentLoc, AgentLoc + ToTarget.GetSafeNormal() * 100.0f, FColor::Cyan, TEXT("toTarget"));
 				UE_VLOG(GetOwner(), LogCrowdFollowing, Log, TEXT("Moved too far, dotValue: %.2f (normalized dot: %.2f) velocity:%s (speed:%.0f)"),
 					FVector::DotProduct(ToTarget, CrowdAgentMoveDirection),
 					FVector::DotProduct(ToTarget.GetSafeNormal(), CrowdAgentMoveDirection),
-					*MovementComp->Velocity.ToString(),
-					MovementComp->Velocity.Size()
+					*NavMovementInterface->GetVelocityForNavMovement().ToString(),
+					NavMovementInterface->GetVelocityForNavMovement().Size()
 					);
 			}
 #endif
@@ -1036,14 +1039,14 @@ FVector UCrowdFollowingComponent::GetMoveFocus(bool bAllowStrafe) const
 	// can't really use CurrentDestination here, as it's pointing at end of path part
 	// fallback to looking at point in front of agent
 
-	if (!bAllowStrafe && MovementComp && IsCrowdSimulationEnabled())
+	if (!bAllowStrafe && NavMovementInterface.IsValid() && IsCrowdSimulationEnabled())
 	{
-		const FVector AgentLoc = MovementComp->GetActorLocation();
+		const FVector AgentLoc = NavMovementInterface->GetLocation();
 
 		// if we're not moving, falling, or don't have a crowd agent move direction, set our focus to ahead of the rotation of our owner to keep the same rotation,
 		// otherwise use the Crowd Agent Move Direction to move in the direction we're supposed to be going
-		const FVector ForwardDir = MovementComp->GetOwner() && ((Status != EPathFollowingStatus::Moving) || (MovementComp->IsFalling()) || CrowdAgentMoveDirection.IsNearlyZero()) ?
-			MovementComp->GetOwner()->GetActorForwardVector() :
+		const FVector ForwardDir = NavMovementInterface->GetOwnerAsObject() && ((Status != EPathFollowingStatus::Moving) || (NavMovementInterface->IsFalling()) || CrowdAgentMoveDirection.IsNearlyZero()) ?
+			NavMovementInterface->GetForwardVector() :
 			CrowdAgentMoveDirection.GetSafeNormal2D();
 
 		return AgentLoc + (ForwardDir * FAIConfig::Navigation::FocalPointDistance);
@@ -1118,7 +1121,7 @@ void UCrowdFollowingComponent::OnNavigationInitDone()
 	}
 
 	// set movement component, it will cache MyNavData from its NavAgentProperties
-	SetMovementComponent(MovementComp);
+	SetNavMovementInterface(NavMovementInterface.Get());
 }
 
 void UCrowdFollowingComponent::GetDebugStringTokens(TArray<FString>& Tokens, TArray<EPathFollowingDebugTokens::Type>& Flags) const
@@ -1167,8 +1170,7 @@ void UCrowdFollowingComponent::GetDebugStringTokens(TArray<FString>& Tokens, TAr
 	// get cylinder of moving agent
 	float AgentRadius = 0.0f;
 	float AgentHalfHeight = 0.0f;
-	AActor* MovingAgent = MovementComp->GetOwner();
-	MovingAgent->GetSimpleCollisionCylinder(AgentRadius, AgentHalfHeight);
+	NavMovementInterface->GetSimpleCollisionCylinder(AgentRadius, AgentHalfHeight);
 
 	if (bFinalPathPart)
 	{
@@ -1188,7 +1190,7 @@ void UCrowdFollowingComponent::GetDebugStringTokens(TArray<FString>& Tokens, TAr
 	}
 	else
 	{
-		const FVector CurrentLocation = MovementComp->GetActorFeetLocation();
+		const FVector CurrentLocation = NavMovementInterface->GetFeetLocation();
 
 		// make sure we're not too close to end of path part (poly count can always fail when AI goes off path)
 		const FVector::FReal DistSq = (GetCurrentTargetLocation() - CurrentLocation).SizeSquared();
diff --git a/Engine/Source/Runtime/AIModule/Private/Navigation/GridPathFollowingComponent.cpp b/Engine/Source/Runtime/AIModule/Private/Navigation/GridPathFollowingComponent.cpp
--- a/Engine/Source/Runtime/AIModule/Private/Navigation/GridPathFollowingComponent.cpp
+++ b/Engine/Source/Runtime/AIModule/Private/Navigation/GridPathFollowingComponent.cpp
@@ -47,7 +47,7 @@ void UGridPathFollowingComponent::OnPathUpdated()
 
 void UGridPathFollowingComponent::UpdatePathSegment()
 {
-	const FVector CurrentLocation = MovementComp ? MovementComp->GetActorFeetLocation() : FVector::ZeroVector;
+	const FVector CurrentLocation = NavMovementInterface.IsValid() ? NavMovementInterface->GetFeetLocation() : FVector::ZeroVector;
 	UpdateActiveGrid(CurrentLocation);
 
 	if (HasActiveGrid() && bHasGridPath)
@@ -173,7 +173,7 @@ void UGridPathFollowingComponent::ResumeMove(FAIRequestID RequestID)
 {
 	if (RequestID.IsEquivalent(GetCurrentRequestId()) && RequestID.IsValid())
 	{
-		const FVector CurrentLocation = MovementComp ? MovementComp->GetActorFeetLocation() : FVector::ZeroVector;
+		const FVector CurrentLocation = NavMovementInterface.IsValid() ? NavMovementInterface->GetFeetLocation() : FVector::ZeroVector;
 		UpdateActiveGrid(CurrentLocation);
 
 		if (HasActiveGrid())
diff --git a/Engine/Source/Runtime/AIModule/Private/Navigation/PathFollowingComponent.cpp b/Engine/Source/Runtime/AIModule/Private/Navigation/PathFollowingComponent.cpp
--- a/Engine/Source/Runtime/AIModule/Private/Navigation/PathFollowingComponent.cpp
+++ b/Engine/Source/Runtime/AIModule/Private/Navigation/PathFollowingComponent.cpp
@@ -231,20 +231,19 @@ void UPathFollowingComponent::LogPathHelper(const AActor* LogOwner, FNavPathShar
 #endif // ENABLE_VISUAL_LOG
 }
 
-void LogBlockHelper(AActor* LogOwner, UNavMovementComponent* MoveComp, float RadiusPct, float HeightPct, const FVector& SegmentStart, const FVector& SegmentEnd)
+void LogBlockHelper(AActor* LogOwner, INavMovementInterface* NavMoveInterface, float RadiusPct, float HeightPct, const FVector& SegmentStart, const FVector& SegmentEnd)
 {
 #if ENABLE_VISUAL_LOG
-	if (MoveComp && LogOwner)
+	if (NavMoveInterface && LogOwner)
 	{
-		const FVector AgentLocation = MoveComp->GetActorFeetLocation();
+		const FVector AgentLocation = NavMoveInterface->GetFeetLocation();
 		const FVector ToTarget = (SegmentEnd - AgentLocation);
 		const FVector::FReal SegmentDot = FVector::DotProduct(ToTarget.GetSafeNormal(), (SegmentEnd - SegmentStart).GetSafeNormal());
 		UE_VLOG(LogOwner, LogPathFollowing, Verbose, TEXT("[agent to segment end] dot [segment dir]: %f"), SegmentDot);
 		
 		float AgentRadius = 0.0f;
 		float AgentHalfHeight = 0.0f;
-		AActor* MovingAgent = MoveComp->GetOwner();
-		MovingAgent->GetSimpleCollisionCylinder(AgentRadius, AgentHalfHeight);
+		NavMoveInterface->GetSimpleCollisionCylinder(AgentRadius, AgentHalfHeight);
 
 		const FVector::FReal Dist2D = ToTarget.Size2D();
 		UE_VLOG(LogOwner, LogPathFollowing, Verbose, TEXT("dist 2d: %f (agent radius: %f [%f])"), Dist2D, AgentRadius, AgentRadius * (1 + RadiusPct));
@@ -404,9 +403,12 @@ FAIRequestID UPathFollowingComponent::RequestMove(const FAIMoveRequest& RequestD
 		// store new data
 		Path = InPath;
 		Path->AddObserver(FNavigationPath::FPathObserverDelegate::FDelegate::CreateUObject(this, &UPathFollowingComponent::OnPathEvent));
-		if (MovementComp && MovementComp->GetOwner())
+		if (NavMovementInterface.IsValid() && NavMovementInterface->GetOwnerAsObject())
 		{
-			Path->SetSourceActor(*(MovementComp->GetOwner()));
+			if (AActor* OwnerAsActor = Cast<AActor>(NavMovementInterface->GetOwnerAsObject()))
+			{
+				Path->SetSourceActor(*(OwnerAsActor));
+			}
 		}
 
 		OriginalMoveRequestGoalLocation = RequestData.GetGoalActor() ? RequestData.GetGoalActor()->GetActorLocation() : RequestData.GetGoalLocation();
@@ -417,7 +419,7 @@ FAIRequestID UPathFollowingComponent::RequestMove(const FAIMoveRequest& RequestD
 		{
 			bIsUsingMetaPath = true;
 
-			const FVector CurrentLocation = MovementComp ? MovementComp->GetActorFeetLocation() : FAISystem::InvalidLocation;
+			const FVector CurrentLocation = NavMovementInterface.IsValid() ? NavMovementInterface->GetFeetLocation() : FAISystem::InvalidLocation;
 			MetaNavPath->Initialize(CurrentLocation);
 		}
 
@@ -439,7 +441,7 @@ FAIRequestID UPathFollowingComponent::RequestMove(const FAIMoveRequest& RequestD
 			UpdateDecelerationData();
 
 #if ENABLE_VISUAL_LOG
-			const FVector CurrentLocation = MovementComp ? MovementComp->GetActorFeetLocation() : FVector::ZeroVector;
+			const FVector CurrentLocation = NavMovementInterface.IsValid() ? NavMovementInterface->GetFeetLocation() : FVector::ZeroVector;
 			const FVector DestLocation = InPath->GetDestinationLocation();
 			const FVector ToDest = DestLocation - CurrentLocation;
 			UE_VLOG(GetOwner(), LogPathFollowing, Log, TEXT("RequestMove: accepted, ID(%u) dist2D(%.0f) distZ(%.0f)"), MoveId.GetID(), ToDest.Size2D(), FMath::Abs(ToDest.Z));
@@ -501,12 +503,12 @@ void UPathFollowingComponent::PauseMove(FAIRequestID RequestID, EPathFollowingVe
 
 	if (RequestID.IsEquivalent(GetCurrentRequestId()))
 	{
-		if ((VelocityMode == EPathFollowingVelocityMode::Reset) && MovementComp && HasMovementAuthority())
+		if ((VelocityMode == EPathFollowingVelocityMode::Reset) && NavMovementInterface.IsValid() && HasMovementAuthority())
 		{
-			MovementComp->StopMovementKeepPathing();
+			NavMovementInterface->StopMovementKeepPathing();
 		}
 
-		LocationWhenPaused = MovementComp ? MovementComp->GetActorFeetLocation() : FVector::ZeroVector;
+		LocationWhenPaused = NavMovementInterface.IsValid() ? NavMovementInterface->GetFeetLocation() : FVector::ZeroVector;
 		PathTimeWhenPaused = Path.IsValid() ? Path->GetTimeStamp() : 0.;
 		Status = EPathFollowingStatus::Paused;
 
@@ -557,12 +559,12 @@ void UPathFollowingComponent::ResumeMove(FAIRequestID RequestID)
 bool UPathFollowingComponent::ShouldCheckPathOnResume() const
 {
 	bool bCheckPath = true;
-	if (MovementComp != NULL)
+	if (NavMovementInterface != NULL)
 	{
 		float AgentRadius = 0.0f, AgentHalfHeight = 0.0f;
-		MovementComp->GetOwner()->GetSimpleCollisionCylinder(AgentRadius, AgentHalfHeight);
+		NavMovementInterface->GetSimpleCollisionCylinder(AgentRadius, AgentHalfHeight);
 
-		const FVector CurrentLocation = MovementComp->GetActorFeetLocation();
+		const FVector CurrentLocation = NavMovementInterface->GetFeetLocation();
 		const FVector::FReal DeltaMove2DSq = (CurrentLocation - LocationWhenPaused).SizeSquared2D();
 		const FVector::FReal DeltaZ = FMath::Abs(CurrentLocation.Z - LocationWhenPaused.Z);
 		if (DeltaMove2DSq < FMath::Square(AgentRadius) && DeltaZ < (AgentHalfHeight * 0.5))
@@ -586,10 +588,10 @@ void UPathFollowingComponent::OnPathFinished(const FPathFollowingResult& Result)
 	}
 	
 	// update meta path if needed
-	if (bIsUsingMetaPath && Result.IsSuccess() && MovementComp)
+	if (bIsUsingMetaPath && Result.IsSuccess() && NavMovementInterface.IsValid())
 	{
 		FMetaNavMeshPath* MetaNavPath = Path->CastPath<FMetaNavMeshPath>();
-		const bool bNeedPathUpdate = MetaNavPath && MetaNavPath->ConditionalMoveToNextSection(MovementComp->GetActorFeetLocation(), EMetaPathUpdateReason::PathFinished);
+		const bool bNeedPathUpdate = MetaNavPath && MetaNavPath->ConditionalMoveToNextSection(NavMovementInterface->GetFeetLocation(), EMetaPathUpdateReason::PathFinished);
 		if (bNeedPathUpdate)
 		{
 			// keep move request active
@@ -608,7 +610,7 @@ void UPathFollowingComponent::OnPathFinished(const FPathFollowingResult& Result)
 
 	if (ShouldStopMovementOnPathFinished())
 	{
-		MovementComp->StopMovementKeepPathing();
+		NavMovementInterface->StopMovementKeepPathing();
 	}
 
 	// notify observers after state was reset (they can request another move)
@@ -621,7 +623,7 @@ void UPathFollowingComponent::OnPathFinished(const FPathFollowingResult& Result)
 
 bool UPathFollowingComponent::ShouldStopMovementOnPathFinished() const
 {
-	return bStopMovementOnFinish && MovementComp && HasMovementAuthority() && !MovementComp->UseAccelerationForPathFollowing();
+	return bStopMovementOnFinish && NavMovementInterface.IsValid() && HasMovementAuthority() && !NavMovementInterface->UseAccelerationForPathFollowing();
 }
 
 void UPathFollowingComponent::OnSegmentFinished()
@@ -633,6 +635,11 @@ void UPathFollowingComponent::OnPathUpdated()
 {
 }
 
+void UPathFollowingComponent::SetMovementComponent(UNavMovementComponent* MoveComp)
+{
+	SetNavMovementInterface(Cast<INavMovementInterface>(MoveComp));
+}
+
 void UPathFollowingComponent::Initialize()
 {
 	UpdateCachedComponents();
@@ -641,7 +648,7 @@ void UPathFollowingComponent::Initialize()
 
 void UPathFollowingComponent::Cleanup()
 {
-	SetMovementComponent(nullptr);
+	SetNavMovementInterface(nullptr);
 	Reset();
 }
 
@@ -683,27 +690,27 @@ void UPathFollowingComponent::TickComponent(float DeltaTime, enum ELevelTick Tic
 	}
 };
 
-void UPathFollowingComponent::SetMovementComponent(UNavMovementComponent* MoveComp)
+void UPathFollowingComponent::SetNavMovementInterface(INavMovementInterface* NavMoveInterface)
 {
-	if (!MoveComp && MovementComp && MovementComp->GetPathFollowingAgent() == this)
+	if (!NavMoveInterface && NavMovementInterface.IsValid() && NavMovementInterface->GetPathFollowingAgent() == this)
 	{
-		MovementComp->SetPathFollowingAgent(nullptr);
+		NavMovementInterface->SetPathFollowingAgent(nullptr);
 	}
-
-	MovementComp = MoveComp;
+	
+	NavMovementInterface = NavMoveInterface;
 	MyNavData = nullptr;
 
-	if (MoveComp)
+	if (NavMovementInterface.IsValid())
 	{
-		const FNavAgentProperties& NavAgentProps = MoveComp->GetNavAgentPropertiesRef();
+		const FNavAgentProperties& NavAgentProps = NavMovementInterface->GetNavAgentPropertiesRef();
 		MyDefaultAcceptanceRadius = NavAgentProps.AgentRadius * 0.1f;
-		MoveComp->SetPathFollowingAgent(Cast<IPathFollowingAgentInterface>(this));
+		NavMovementInterface->SetPathFollowingAgent(Cast<IPathFollowingAgentInterface>(this));
 
 		UNavigationSystemV1* NavSys = FNavigationSystem::GetCurrent<UNavigationSystemV1>(GetWorld());
 		
 		if (NavSys)
 		{
-			MyNavData = NavSys->GetNavDataForProps(NavAgentProps, MoveComp->GetActorLocation());
+			MyNavData = NavSys->GetNavDataForProps(NavAgentProps, NavMovementInterface->GetLocation());
 			if (MyNavData == nullptr)
 			{
 				if (NavSys->IsInitialized() == false)
@@ -722,21 +729,21 @@ void UPathFollowingComponent::SetMovementComponent(UNavMovementComponent* MoveCo
 void UPathFollowingComponent::OnNavigationInitDone()
 {
 	UWorld* MyWorld = GetWorld();
-	if (MovementComp && MyWorld)
+	if (NavMovementInterface.IsValid() && MyWorld)
 	{
 		UNavigationSystemV1* NavSys = FNavigationSystem::GetCurrent<UNavigationSystemV1>(MyWorld);
 		check(NavSys);
-		const FNavAgentProperties& NavAgentProps = MovementComp->GetNavAgentPropertiesRef();
-		MyNavData = NavSys->GetNavDataForProps(NavAgentProps, MovementComp->GetActorLocation());
+		const FNavAgentProperties& NavAgentProps = NavMovementInterface->GetNavAgentPropertiesRef();
+		MyNavData = NavSys->GetNavDataForProps(NavAgentProps, NavMovementInterface->GetLocation());
 		NavSys->OnNavigationInitDone.RemoveAll(this);
 	}
 }
 
 void UPathFollowingComponent::OnNavDataRegistered(ANavigationData* NavData)
 {
-	if (NavData && MovementComp)
+	if (NavData && NavMovementInterface.IsValid())
 	{
-		const FNavAgentProperties& NavAgentProps = MovementComp->GetNavAgentPropertiesRef();
+		const FNavAgentProperties& NavAgentProps = NavMovementInterface->GetNavAgentPropertiesRef();
 		if (NavData->DoesSupportAgent(NavAgentProps))
 		{
 			MyNavData = NavData;
@@ -823,12 +830,12 @@ int32 UPathFollowingComponent::DetermineStartingPathPoint(const FNavigationPath*
 			}
 		}
 
-		if (MovementComp && PickedPathPoint == INDEX_NONE)
+		if (NavMovementInterface.IsValid() && PickedPathPoint == INDEX_NONE)
 		{
 			if (ConsideredPath->GetPathPoints().Num() > 2)
 			{
 				// check if is closer to first or second path point (don't assume AI's standing)
-				const FVector CurrentLocation = MovementComp->GetActorFeetLocation();
+				const FVector CurrentLocation = NavMovementInterface->GetFeetLocation();
 				const FVector PathPt0 = *ConsideredPath->GetPathPointLocation(0);
 				const FVector PathPt1 = *ConsideredPath->GetPathPointLocation(1);
 				// making this test in 2d to avoid situation where agent's Z location not being in "navmesh plane"
@@ -946,14 +953,14 @@ void UPathFollowingComponent::UpdatePathSegment()
 	DEBUG_bMovingDirectlyToGoal = false;
 #endif // !UE_BUILD_SHIPPING
 
-	if ((Path.IsValid() == false) || (MovementComp == nullptr))
+	if ((Path.IsValid() == false) || (NavMovementInterface == nullptr))
 	{
 		UE_CVLOG(Path.IsValid() == false, this, LogPathFollowing, Log, TEXT("Aborting move due to not having a valid path object"));
 		OnPathFinished(EPathFollowingResult::Aborted, FPathFollowingResultFlags::InvalidPath);
 		return;
 	}
 
-	const FVector CurrentLocation = MovementComp->GetActorFeetLocation();
+	const FVector CurrentLocation = NavMovementInterface->GetFeetLocation();
 	const bool bCanUpdateState = HasMovementAuthority();
 
 	if (!Path->IsValid())
@@ -1062,13 +1069,13 @@ void UPathFollowingComponent::UpdatePathSegment()
 		{
 			if (Path->GetPathPoints().IsValidIndex(MoveSegmentEndIndex) && Path->GetPathPoints().IsValidIndex(MoveSegmentStartIndex))
 			{
-				LogBlockHelper(GetOwner(), MovementComp, MinAgentRadiusPct, MinAgentHalfHeightPct,
+				LogBlockHelper(GetOwner(), NavMovementInterface.Get(), MinAgentRadiusPct, MinAgentHalfHeightPct,
 					*Path->GetPathPointLocation(MoveSegmentStartIndex),
 					*Path->GetPathPointLocation(MoveSegmentEndIndex));
 			}
 			else
 			{
-				if ((GetOwner() != NULL) && (MovementComp != NULL))
+				if ((GetOwner() != NULL) && (NavMovementInterface != NULL))
 				{
 					UE_VLOG(GetOwner(), LogPathFollowing, Verbose, TEXT("Path blocked, but move segment indices are not valid: start %d, end %d of %d"), MoveSegmentStartIndex, MoveSegmentEndIndex, Path->GetPathPoints().Num());
 				}
@@ -1080,18 +1087,18 @@ void UPathFollowingComponent::UpdatePathSegment()
 
 void UPathFollowingComponent::FollowPathSegment(float DeltaTime)
 {
-	if (!Path.IsValid() || MovementComp == nullptr)
+	if (!Path.IsValid() || NavMovementInterface == nullptr)
 	{
 		return;
 	}
 
-	const FVector CurrentLocation = MovementComp->GetActorFeetLocation();
+	const FVector CurrentLocation = NavMovementInterface->GetFeetLocation();
 	const FVector CurrentTarget = GetCurrentTargetLocation();
 	
 	// set to false by default, we will set set this back to true if appropriate
 	bIsDecelerating = false;
 
-	const bool bAccelerationBased = MovementComp->UseAccelerationForPathFollowing();
+	const bool bAccelerationBased = NavMovementInterface->UseAccelerationForPathFollowing();
 	if (bAccelerationBased)
 	{
 		CurrentMoveInput = (CurrentTarget - CurrentLocation).GetSafeNormal();
@@ -1111,7 +1118,7 @@ void UPathFollowingComponent::FollowPathSegment(float DeltaTime)
 		}
 
 		PostProcessMove.ExecuteIfBound(this, CurrentMoveInput);
-		MovementComp->RequestPathMove(CurrentMoveInput);
+		NavMovementInterface->RequestPathMove(CurrentMoveInput);
 	}
 	else
 	{
@@ -1121,14 +1128,14 @@ void UPathFollowingComponent::FollowPathSegment(float DeltaTime)
 		const bool bNotFollowingLastSegment = (MoveSegmentStartIndex < LastSegmentStartIndex);
 
 		PostProcessMove.ExecuteIfBound(this, MoveVelocity);
-		MovementComp->RequestDirectMove(MoveVelocity, bNotFollowingLastSegment);
+		NavMovementInterface->RequestDirectMove(MoveVelocity, bNotFollowingLastSegment);
 	}
 }
 
 bool UPathFollowingComponent::HasReached(const FVector& TestPoint, EPathFollowingReachMode ReachMode, float InAcceptanceRadius) const
 {
 	// simple test for stationary agent, used as early finish condition
-	const FVector CurrentLocation = MovementComp ? MovementComp->GetActorFeetLocation() : FVector::ZeroVector;
+	const FVector CurrentLocation = NavMovementInterface.IsValid() ? NavMovementInterface->GetFeetLocation() : FVector::ZeroVector;
 	const float GoalRadius = 0.0f;
 	const float GoalHalfHeight = 0.0f;
 	if (InAcceptanceRadius == UPathFollowingComponent::DefaultAcceptanceRadius)
@@ -1175,7 +1182,7 @@ bool UPathFollowingComponent::HasReached(const AActor& TestGoal, EPathFollowingR
 		}
 	}
 
-	const FVector CurrentLocation = MovementComp ? MovementComp->GetActorFeetLocation() : FVector::ZeroVector;
+	const FVector CurrentLocation = NavMovementInterface.IsValid() ? NavMovementInterface->GetFeetLocation() : FVector::ZeroVector;
 	const float AgentRadiusMod = (ReachMode == EPathFollowingReachMode::ExactLocation) || (ReachMode == EPathFollowingReachMode::OverlapGoal) ? 0.0f : MinAgentRadiusPct;
 	return HasReachedInternal(TestPoint, GoalRadius, GoalHalfHeight, CurrentLocation, InAcceptanceRadius, AgentRadiusMod);
 }
@@ -1229,7 +1236,7 @@ bool UPathFollowingComponent::HasReachedDestination(const FVector& CurrentLocati
 
 bool UPathFollowingComponent::HasReachedCurrentTarget(const FVector& CurrentLocation) const
 {
-	if (MovementComp == NULL)
+	if (NavMovementInterface == NULL)
 	{
 		return false;
 	}
@@ -1259,7 +1266,7 @@ bool UPathFollowingComponent::HasReachedCurrentTarget(const FVector& CurrentLoca
 	const FVector CurrentDirection = GetCurrentDirection();
 
 	// check if moved too far
-	const FVector ToTarget = (CurrentTarget - MovementComp->GetActorFeetLocation());
+	const FVector ToTarget = (CurrentTarget - NavMovementInterface->GetFeetLocation());
 	const FVector::FReal SegmentDot = FVector::DotProduct(ToTarget, CurrentDirection);
 	if (SegmentDot < 0.)
 	{
@@ -1276,7 +1283,7 @@ bool UPathFollowingComponent::HasReachedCurrentTarget(const FVector& CurrentLoca
 
 bool UPathFollowingComponent::HasReachedInternal(const FVector& GoalLocation, float GoalRadius, float GoalHalfHeight, const FVector& AgentLocation, float RadiusThreshold, float AgentRadiusMultiplier) const
 {
-	if (MovementComp == NULL)
+	if (NavMovementInterface == NULL)
 	{
 		return false;
 	}
@@ -1284,8 +1291,7 @@ bool UPathFollowingComponent::HasReachedInternal(const FVector& GoalLocation, fl
 	// get cylinder of moving agent
 	float AgentRadius = 0.0f;
 	float AgentHalfHeight = 0.0f;
-	AActor* MovingAgent = MovementComp->GetOwner();
-	MovingAgent->GetSimpleCollisionCylinder(AgentRadius, AgentHalfHeight);
+	NavMovementInterface->GetSimpleCollisionCylinder(AgentRadius, AgentHalfHeight);
 
 	// check if they overlap (with added AcceptanceRadius)
 	const FVector ToGoal = GoalLocation - AgentLocation;
@@ -1358,7 +1364,7 @@ float UPathFollowingComponent::GetFinalAcceptanceRadius(const FNavigationPath& P
 
 void UPathFollowingComponent::DebugReachTest(float& CurrentDot, float& CurrentDistance, float& CurrentHeight, uint8& bDotFailed, uint8& bDistanceFailed, uint8& bHeightFailed) const
 {
-	if (!Path.IsValid() || MovementComp == nullptr || MovementComp->UpdatedComponent == nullptr)
+	if (!Path.IsValid() || NavMovementInterface == nullptr || NavMovementInterface->GetUpdatedObject() == nullptr)
 	{
 		return;
 	}
@@ -1371,7 +1377,7 @@ void UPathFollowingComponent::DebugReachTest(float& CurrentDot, float& CurrentDi
 	float RadiusThreshold = 0.0f;
 	float AgentRadiusPct = 0.05f;
 
-	FVector AgentLocation = MovementComp->GetActorFeetLocation();
+	FVector AgentLocation = NavMovementInterface->GetFeetLocation();
 	FVector GoalLocation = GetCurrentTargetLocation();
 	RadiusThreshold = CurrentAcceptanceRadius;
 
@@ -1406,8 +1412,7 @@ void UPathFollowingComponent::DebugReachTest(float& CurrentDot, float& CurrentDi
 	// get cylinder of moving agent
 	float AgentRadius = 0.0f;
 	float AgentHalfHeight = 0.0f;
-	AActor* MovingAgent = MovementComp->GetOwner();
-	MovingAgent->GetSimpleCollisionCylinder(AgentRadius, AgentHalfHeight);
+	NavMovementInterface->GetSimpleCollisionCylinder(AgentRadius, AgentHalfHeight);
 
 	CurrentDistance = FloatCastChecked<float>(ToGoal.Size2D(), UE::LWC::DefaultFloatPrecision);
 	const float UseRadius = FMath::Max(RadiusThreshold, GoalRadius + (AgentRadius * AgentRadiusPct));
@@ -1462,7 +1467,7 @@ void UPathFollowingComponent::FinishUsingCustomLink(INavLinkCustomInterface* Cus
 
 bool UPathFollowingComponent::UpdateMovementComponent(bool bForce)
 {
-	if (MovementComp == NULL || bForce == true)
+	if (NavMovementInterface == NULL || bForce == true)
 	{
 		APawn* MyPawn = Cast<APawn>(GetOwner());
 		if (MyPawn == NULL)
@@ -1478,11 +1483,11 @@ bool UPathFollowingComponent::UpdateMovementComponent(bool bForce)
 		{
 			MyPawn->OnActorHit.AddUniqueDynamic(this, &UPathFollowingComponent::OnActorBump);
 
-			SetMovementComponent(MyPawn->FindComponentByClass<UNavMovementComponent>());
+			SetNavMovementInterface(Cast<INavMovementInterface>(MyPawn->FindComponentByInterface<UNavMovementInterface>()));
 		}
 	}
 
-	return (MovementComp != NULL);
+	return (NavMovementInterface != nullptr);
 }
 
 void UPathFollowingComponent::OnActorBump(AActor* SelfActor, AActor* OtherActor, FVector NormalImpulse, const FHitResult& Hit)
@@ -1555,7 +1560,7 @@ bool UPathFollowingComponent::UpdateBlockDetection()
 {
 	const double GameTime = GetWorld()->GetTimeSeconds();
 	if (bUseBlockDetection &&
-		MovementComp && 
+		NavMovementInterface.IsValid() && 
 		GameTime > (LastSampleTime + BlockDetectionInterval) &&
 		BlockDetectionSampleCount > 0)
 	{
@@ -1566,7 +1571,7 @@ bool UPathFollowingComponent::UpdateBlockDetection()
 			LocationSamples.AddZeroed(1);
 		}
 
-		LocationSamples[NextSampleIdx] = MovementComp->GetActorFeetLocationBased();
+		LocationSamples[NextSampleIdx] = NavMovementInterface->GetFeetLocationBased();
 		NextSampleIdx = (NextSampleIdx + 1) % BlockDetectionSampleCount;
 		return true;
 	}
@@ -1603,7 +1608,7 @@ void UPathFollowingComponent::ForceBlockDetectionUpdate()
 void UPathFollowingComponent::UpdateDecelerationData()
 {
 	// no point in updating if we don't have a MovementComp
-	if (MovementComp == nullptr)
+	if (NavMovementInterface == nullptr)
 	{
 		return;
 	}
@@ -1613,12 +1618,12 @@ void UPathFollowingComponent::UpdateDecelerationData()
 		return;
 	}
 
-	const float CurrentMaxSpeed = MovementComp->GetMaxSpeed();
+	const float CurrentMaxSpeed = NavMovementInterface->GetMaxSpeedForNavMovement();
 	bool bUpdatePathSegment = (DecelerationSegmentIndex == INDEX_NONE);
 	if (CurrentMaxSpeed != CachedBrakingMaxSpeed)
 	{
 		const float PrevBrakingDistance = CachedBrakingDistance;
-		CachedBrakingDistance = MovementComp->GetPathFollowingBrakingDistance(CurrentMaxSpeed);
+		CachedBrakingDistance = NavMovementInterface->GetPathFollowingBrakingDistance(CurrentMaxSpeed);
 		CachedBrakingMaxSpeed = CurrentMaxSpeed;
 		
 		bUpdatePathSegment = bUpdatePathSegment || (PrevBrakingDistance != CachedBrakingDistance);
@@ -1686,7 +1691,7 @@ FVector::FReal UPathFollowingComponent::GetRemainingPathCost() const
 FNavLocation UPathFollowingComponent::GetCurrentNavLocation() const
 {
 	// get navigation location of moved actor
-	if (MovementComp == nullptr)
+	if (NavMovementInterface == nullptr)
 	{
 		return FNavLocation();
 	}
@@ -1697,14 +1702,13 @@ FNavLocation UPathFollowingComponent::GetCurrentNavLocation() const
 		return FNavLocation();
 	}
 
-	const FVector OwnerLoc = MovementComp->GetActorNavLocation();
+	const FVector OwnerLoc = NavMovementInterface->GetNavLocation();
 	const float Tolerance = 1.0f;
 	
 	// Using !Equals rather than != because exact equivalence isn't required.  (Equals allows a small tolerance.)
 	if (!OwnerLoc.Equals(CurrentNavLocation.Location, Tolerance))
 	{
-		const AActor* OwnerActor = MovementComp->GetOwner();
-		const FVector OwnerExtent = OwnerActor ? OwnerActor->GetSimpleCollisionCylinderExtent() : FVector::ZeroVector;
+		const FVector OwnerExtent = NavMovementInterface.IsValid() ? NavMovementInterface->GetSimpleCollisionCylinderExtent() : FVector::ZeroVector;
 
 		NavSys->ProjectPointToNavigation(OwnerLoc, CurrentNavLocation, OwnerExtent, MyNavData);
 	}
@@ -1883,7 +1887,7 @@ FString UPathFollowingComponent::GetDebugString() const
 
 bool UPathFollowingComponent::IsPathFollowingAllowed() const
 {
-	return MovementComp && MovementComp->CanStartPathFollowing();
+	return NavMovementInterface.IsValid() && NavMovementInterface->CanStartPathFollowing();
 }
 
 void UPathFollowingComponent::OnUnableToMove(const UObject& Instigator)
diff --git a/Engine/Source/Runtime/Engine/Classes/AI/Navigation/NavigationTypes.h b/Engine/Source/Runtime/Engine/Classes/AI/Navigation/NavigationTypes.h
--- a/Engine/Source/Runtime/Engine/Classes/AI/Navigation/NavigationTypes.h
+++ b/Engine/Source/Runtime/Engine/Classes/AI/Navigation/NavigationTypes.h
@@ -444,6 +444,35 @@ struct FMovementProperties
 	}
 };
 
+/**
+ * Struct to hold properties a user might set for navigation movement
+ */
+USTRUCT(BlueprintType)
+struct FNavMovementProperties
+{
+	GENERATED_BODY()
+
+	/** Braking distance override used with acceleration driven path following (bUseAccelerationForPaths) */
+	UPROPERTY(EditAnywhere, Category = NavMovement, meta = (EditCondition = "bUseFixedBrakingDistanceForPaths"))
+	float FixedPathBrakingDistance = 0;
+
+	/** If set to true, NavAgentProperties' radius and height will be updated with Owner's collision capsule size */
+	UPROPERTY(EditAnywhere, Category = NavMovement)
+	bool bUpdateNavAgentWithOwnersCollision = true;
+
+	/** If set, pathfollowing will control character movement via acceleration values. If false, it will set velocities directly. */
+	UPROPERTY(EditAnywhere, Category = NavMovement, GlobalConfig)
+	bool bUseAccelerationForPaths = false;
+
+	/** If set, FixedPathBrakingDistance will be used for path following deceleration */
+	UPROPERTY(EditAnywhere, Category = NavMovement, meta = (EditCondition = "bUseAccelerationForPaths"))
+	bool bUseFixedBrakingDistanceForPaths = false;
+
+	/** If set, StopActiveMovement call will abort current path following request */
+	UPROPERTY(EditAnywhere, Category = NavMovement)
+	bool bStopMovementAbortPaths = true;
+};
+
 /** Properties of representation of an 'agent' (or Pawn) used by AI navigation/pathfinding. */
 USTRUCT(BlueprintType)
 struct FNavAgentProperties : public FMovementProperties
diff --git a/Engine/Source/Runtime/Engine/Classes/GameFramework/NavMovementComponent.h b/Engine/Source/Runtime/Engine/Classes/GameFramework/NavMovementComponent.h
--- a/Engine/Source/Runtime/Engine/Classes/GameFramework/NavMovementComponent.h
+++ b/Engine/Source/Runtime/Engine/Classes/GameFramework/NavMovementComponent.h
@@ -7,6 +7,7 @@
 #pragma once
 
 #include "CoreMinimal.h"
+#include "NavMovementInterface.h"
 #include "UObject/ObjectMacros.h"
 #include "Engine/EngineTypes.h"
 #include "GameFramework/Actor.h"
@@ -22,35 +23,43 @@ class UCapsuleComponent;
  * NavMovementComponent defines base functionality for MovementComponents that move any 'agent' that may be involved in AI pathfinding.
  */
 UCLASS(abstract, config=Engine, MinimalAPI)
-class UNavMovementComponent : public UMovementComponent
+class UNavMovementComponent : public UMovementComponent, public INavMovementInterface
 {
 	GENERATED_UCLASS_BODY()
 
-	/** Properties that define how the component can move. */
-	UPROPERTY(EditAnywhere, BlueprintReadWrite, Category = NavMovement, meta = (DisplayName = "Movement Capabilities", Keywords = "Nav Agent"))
-	FNavAgentProperties NavAgentProps;
-
 protected:
 	/** Braking distance override used with acceleration driven path following (bUseAccelerationForPaths) */
-	UPROPERTY(EditAnywhere, Category = NavMovement, meta = (EditCondition = "bUseFixedBrakingDistanceForPaths"))
+	UE_DEPRECATED(5.5, "FixedPathBrakingDistance is deprecated, please use NavMovementProperties.FixedPathBrakingDistance instead.")
+	UPROPERTY(meta = (DeprecatedProperty, DeprecationMessage = "FixedPathBrakingDistance is deprecated, please use NavMovementProperties.FixedPathBrakingDistance instead."))
 	float FixedPathBrakingDistance;
 
 	/** If set to true NavAgentProps' radius and height will be updated with Owner's collision capsule size */
-	UPROPERTY(EditAnywhere, Category = NavMovement)
+	UE_DEPRECATED(5.5, "bUpdateNavAgentWithOwnersCollision is deprecated, please use NavMovementProperties.bUpdateNavAgentWithOwnersCollision instead.")
+	UPROPERTY(meta = (DeprecatedProperty, DeprecationMessage = "bUpdateNavAgentWithOwnersCollision is deprecated, please use NavMovementProperties.bUpdateNavAgentWithOwnersCollision instead."))
 	uint8 bUpdateNavAgentWithOwnersCollision:1;
 
 	/** If set, pathfollowing will control character movement via acceleration values. If false, it will set velocities directly. */
-	UPROPERTY(EditAnywhere, Category = NavMovement, GlobalConfig)
+	UE_DEPRECATED(5.5, "bUseAccelerationForPaths is deprecated, please use NavMovementProperties.bUseAccelerationForPaths instead.")
+	UPROPERTY(GlobalConfig, meta = (DeprecatedProperty, DeprecationMessage = "bUseAccelerationForPaths is deprecated, please use NavMovementProperties.bUseAccelerationForPaths instead."))
 	uint8 bUseAccelerationForPaths : 1;
 
 	/** If set, FixedPathBrakingDistance will be used for path following deceleration */
-	UPROPERTY(EditAnywhere, Category = NavMovement, meta = (EditCondition = "bUseAccelerationForPaths"))
+	UE_DEPRECATED(5.5, "bUseFixedBrakingDistanceForPaths is deprecated, please use NavMovementProperties.bUseFixedBrakingDistanceForPaths instead.")
+	UPROPERTY(meta = (DeprecatedProperty, DeprecationMessage = "bUseFixedBrakingDistanceForPaths is deprecated, please use NavMovementProperties.bUseFixedBrakingDistanceForPaths instead."))
 	uint8 bUseFixedBrakingDistanceForPaths : 1;
 
 	/** If set, StopActiveMovement call will abort current path following request */
+	UE_DEPRECATED(5.5, "bStopMovementAbortPaths is deprecated, please use NavMovementProperties.bStopMovementAbortPaths instead.")
 	uint8 bStopMovementAbortPaths:1;
+	
+	UPROPERTY(EditAnywhere, Category = NavMovement)
+	FNavMovementProperties NavMovementProperties;
 
 public:
+	/** Properties that define how the component can move. */
+	UPROPERTY(EditAnywhere, BlueprintReadWrite, Category = NavMovement, meta = (DisplayName = "Movement Capabilities", Keywords = "Nav Agent"))
+	FNavAgentProperties NavAgentProps;
+
 	/** Expresses runtime state of character's movement. Put all temporal changes to movement properties here */
 	UPROPERTY()
 	FMovementProperties MovementState;
@@ -62,22 +71,32 @@ private:
 	TObjectPtr<UObject> PathFollowingComp;
 
 public:
-	/** Stops applying further movement (usually zeros acceleration). */
-	UFUNCTION(BlueprintCallable, Category="Pawn|Components|PawnMovement")
-	ENGINE_API virtual void StopActiveMovement();
+	/** Get the owner of the object consuming nav movement */
+	ENGINE_API virtual UObject* GetOwnerAsObject() const override { return GetOwner(); }
+	
+	/** Get the component this movement component is updating */
+	ENGINE_API virtual TObjectPtr<UObject> GetUpdatedObject() const override { return UpdatedComponent; }
+
+	/** Get axis-aligned cylinder around this actor, used for simple collision checks in nav movement */
+	ENGINE_API virtual void GetSimpleCollisionCylinder(float& CollisionRadius, float& CollisionHalfHeight) const override;
+
+	/** Returns collision extents vector for this object, based on GetSimpleCollisionCylinder. */
+	ENGINE_API virtual FVector GetSimpleCollisionCylinderExtent() const override;
+	
+	/** Get forward vector of the object being driven by nav movement */
+	ENGINE_API virtual FVector GetForwardVector() const override;
+	
+	/** Get the current velocity of the movement component */
+	UFUNCTION(BlueprintCallable, Category="AI|Components|NavMovement")
+	ENGINE_API virtual FVector GetVelocityForNavMovement() const override { return Velocity; }
 
-	/** Stops movement immediately (reset velocity) but keeps following current path */
-	UFUNCTION(BlueprintCallable, Category="Components|Movement")
-	ENGINE_API void StopMovementKeepPathing();
+	ENGINE_API virtual float GetMaxSpeedForNavMovement() const override { return GetMaxSpeed(); }
 
 	// Overridden to also call StopActiveMovement().
 	ENGINE_API virtual void StopMovementImmediately() override;
 
 	ENGINE_API void SetUpdateNavAgentWithOwnersCollisions(bool bUpdateWithOwner);
-	FORCEINLINE bool ShouldUpdateNavAgentWithOwnersCollision() const { return bUpdateNavAgentWithOwnersCollision != 0; }
-	
-	ENGINE_API void UpdateNavAgent(const AActor& InOwner);
-	ENGINE_API void UpdateNavAgent(const UCapsuleComponent& CapsuleComponent);
+	FORCEINLINE bool ShouldUpdateNavAgentWithOwnersCollision() const { return NavMovementProperties.bUpdateNavAgentWithOwnersCollision != 0; }
 
 	/** Returns location of controlled actor - meaning center of collision bounding box */
 	FORCEINLINE FVector GetActorLocation() const { return UpdatedComponent ? UpdatedComponent->GetComponentLocation() : FVector(FLT_MAX); }
@@ -90,56 +109,50 @@ public:
 	/** Returns the full world-coordinates transform of the associated scene component (the UpdatedComponent) */
 	FORCEINLINE FTransform GetActorTransform() const { return UpdatedComponent ? UpdatedComponent->GetComponentTransform() : FTransform(); }
 
+	/** Returns location of controlled actor - meaning center of collision bounding box */
+	FORCEINLINE virtual FVector GetLocation() const override { return GetActorLocation(); }
+	/** Returns location of controlled actor's "feet" meaning center of bottom of collision bounding box */
+	virtual FVector GetFeetLocation() const override { return GetActorFeetLocation(); }
+	/** Returns based location of controlled actor */
+	ENGINE_API virtual FBasedPosition GetFeetLocationBased() const override { return GetActorFeetLocationBased(); };
+	
+	ENGINE_API virtual void PostLoad() override;
+
+	/** Set nav agent properties from an object */
+	ENGINE_API virtual void UpdateNavAgent(const UObject& ObjectToUpdateFrom) override;
+	
 	/** path following: request new velocity */
-	ENGINE_API virtual void RequestDirectMove(const FVector& MoveVelocity, bool bForceMaxSpeed);
+	ENGINE_API virtual void RequestDirectMove(const FVector& MoveVelocity, bool bForceMaxSpeed) override;
 
 	/** path following: request new move input (normal vector = full strength) */
-	ENGINE_API virtual void RequestPathMove(const FVector& MoveInput);
+	ENGINE_API virtual void RequestPathMove(const FVector& MoveInput) override;
 
 	/** check if current move target can be reached right now if positions are matching
 	 *  (e.g. performing scripted move and can't stop) */
-	ENGINE_API virtual bool CanStopPathFollowing() const;
-
-	/** Returns braking distance for acceleration driven path following */
-	ENGINE_API virtual float GetPathFollowingBrakingDistance(float MaxSpeed) const;
-
-	void SetPathFollowingAgent(IPathFollowingAgentInterface* InPathFollowingAgent) { PathFollowingComp = Cast<UObject>(InPathFollowingAgent); }
-	IPathFollowingAgentInterface* GetPathFollowingAgent() { return Cast<IPathFollowingAgentInterface>(PathFollowingComp); }
-	const IPathFollowingAgentInterface* GetPathFollowingAgent() const { return Cast<const IPathFollowingAgentInterface>(PathFollowingComp); }
+	ENGINE_API virtual bool CanStopPathFollowing() const override;
 
-	/** set fixed braking distance */
-	ENGINE_API void SetFixedBrakingDistance(float DistanceToEndOfPath);
+	/** Get Nav movement props struct this component uses */
+	ENGINE_API virtual FNavMovementProperties* GetNavMovementProperties() override { return &NavMovementProperties; }
+	/** Returns the NavMovementProps(const) */
+	ENGINE_API virtual const FNavMovementProperties& GetNavMovementProperties() const override{ return NavMovementProperties; } 
 
 	/** clears fixed braking distance */
 	ENGINE_API void ClearFixedBrakingDistance();
 
-	FORCEINLINE bool UseAccelerationForPathFollowing() const { return bUseAccelerationForPaths; }
+	virtual void SetPathFollowingAgent(IPathFollowingAgentInterface* InPathFollowingAgent) override { PathFollowingComp = Cast<UObject>(InPathFollowingAgent); }
+	virtual IPathFollowingAgentInterface* GetPathFollowingAgent() override { return Cast<IPathFollowingAgentInterface>(PathFollowingComp); }
+	virtual const IPathFollowingAgentInterface* GetPathFollowingAgent() const override { return Cast<const IPathFollowingAgentInterface>(PathFollowingComp); }
 
 	/** Returns the NavAgentProps(const) */
-	FORCEINLINE const FNavAgentProperties& GetNavAgentPropertiesRef() const { return NavAgentProps; }
+	FORCEINLINE const FNavAgentProperties& GetNavAgentPropertiesRef() const  override { return NavAgentProps; }
 	/** Returns the NavAgentProps */
-	FORCEINLINE FNavAgentProperties& GetNavAgentPropertiesRef() { return NavAgentProps; }
+	FORCEINLINE FNavAgentProperties& GetNavAgentPropertiesRef()  override { return NavAgentProps; }
 
 	/** Resets runtime movement state to character's movement capabilities */
-	void ResetMoveState() { MovementState = NavAgentProps; }
+	virtual void ResetMoveState() override { MovementState = NavAgentProps; }
 
 	/** Returns true if path following can start */
-	virtual bool CanStartPathFollowing() const { return true; }
-
-	/** Returns true if component can crouch */
-	FORCEINLINE bool CanEverCrouch() const { return NavAgentProps.bCanCrouch; }
-
-	/** Returns true if component can jump */
-	FORCEINLINE bool CanEverJump() const { return NavAgentProps.bCanJump; }
-
-	/** Returns true if component can move along the ground (walk, drive, etc) */
-	FORCEINLINE bool CanEverMoveOnGround() const { return NavAgentProps.bCanWalk; }
-
-	/** Returns true if component can swim */
-	FORCEINLINE bool CanEverSwim() const { return NavAgentProps.bCanSwim; }
-
-	/** Returns true if component can fly */
-	FORCEINLINE bool CanEverFly() const { return NavAgentProps.bCanFly; }
+	virtual bool CanStartPathFollowing() const override { return true; }
 
 	/** Returns true if component is allowed to jump */
 	FORCEINLINE bool IsJumpAllowed() const { return CanEverJump() && MovementState.bCanJump; }
@@ -149,23 +162,23 @@ public:
 
 	/** Returns true if currently crouching */ 
 	UFUNCTION(BlueprintCallable, Category="AI|Components|NavMovement")
-	ENGINE_API virtual bool IsCrouching() const;
+	ENGINE_API virtual bool IsCrouching() const override;
 
 	/** Returns true if currently falling (not flying, in a non-fluid volume, and not on the ground) */ 
 	UFUNCTION(BlueprintCallable, Category="AI|Components|NavMovement")
-	ENGINE_API virtual bool IsFalling() const;
+	ENGINE_API virtual bool IsFalling() const override;
 
 	/** Returns true if currently moving on the ground (e.g. walking or driving) */ 
 	UFUNCTION(BlueprintCallable, Category="AI|Components|NavMovement")
-	ENGINE_API virtual bool IsMovingOnGround() const;
+	ENGINE_API virtual bool IsMovingOnGround() const override;
 	
 	/** Returns true if currently swimming (moving through a fluid volume) */
 	UFUNCTION(BlueprintCallable, Category="AI|Components|NavMovement")
-	ENGINE_API virtual bool IsSwimming() const;
+	ENGINE_API virtual bool IsSwimming() const override;
 
 	/** Returns true if currently flying (moving through a non-fluid volume without resting on the ground) */
 	UFUNCTION(BlueprintCallable, Category="AI|Components|NavMovement")
-	ENGINE_API virtual bool IsFlying() const;
+	ENGINE_API virtual bool IsFlying() const override;
 };
 
 
@@ -194,13 +207,6 @@ inline bool UNavMovementComponent::IsFlying() const
 	return false;
 }
 
-inline void UNavMovementComponent::StopMovementKeepPathing()
-{
-	bStopMovementAbortPaths = false;
-	StopMovementImmediately();
-	bStopMovementAbortPaths = true;
-}
-
 inline void UNavMovementComponent::StopMovementImmediately()
 {
 	Super::StopMovementImmediately();
diff --git a/Engine/Source/Runtime/Engine/Classes/GameFramework/NavMovementInterface.h b/Engine/Source/Runtime/Engine/Classes/GameFramework/NavMovementInterface.h
new file mode 100644
--- /dev/null
+++ b/Engine/Source/Runtime/Engine/Classes/GameFramework/NavMovementInterface.h
@@ -0,0 +1,161 @@
+﻿// Copyright Epic Games, Inc. All Rights Reserved.
+
+#pragma once
+
+#include "UObject/ObjectMacros.h"
+#include "UObject/Interface.h"
+#include "Engine/EngineTypes.h"
+#include "AI/Navigation/NavigationTypes.h"
+
+#include "NavMovementInterface.generated.h"
+
+class IPathFollowingAgentInterface;
+
+/**
+ * Interface for navigation movement - should be implemented on movement objects that control an object directly
+ */
+UINTERFACE(MinimalAPI, NotBlueprintable)
+class UNavMovementInterface : public UInterface
+{
+	GENERATED_BODY()
+};
+
+class ENGINE_API INavMovementInterface
+{
+	GENERATED_BODY()
+
+public:
+	// Begin: Functions to be implemented by objects inheriting this interface
+	/** Returns true if path following can start */
+	virtual bool CanStartPathFollowing() const = 0;
+
+	/** check if current move target can be reached right now if positions are matching
+	 *  (e.g. performing scripted move and can't stop) */
+	virtual bool CanStopPathFollowing() const = 0;
+
+	/** Returns the Nav movement properties struct used by NavMovementInterface and PathFollowing */
+	virtual FNavMovementProperties* GetNavMovementProperties() = 0;
+	/** Returns the Nav movement properties struct (const) used by NavMovementInterface and PathFollowing */
+	virtual const FNavMovementProperties& GetNavMovementProperties() const = 0;
+
+	/** Returns the NavAgentProperties used by NavMovementInterface and PathFollowing */
+	virtual FNavAgentProperties& GetNavAgentPropertiesRef() = 0;
+	/** Returns the NavAgentProperties(const) used by NavMovementInterface and PathFollowing */
+	virtual const FNavAgentProperties& GetNavAgentPropertiesRef() const = 0;
+	
+	/** Set the path following agent this interface uses. Must implement IPathFollowingAgentInterface. */
+	virtual void SetPathFollowingAgent(IPathFollowingAgentInterface* InPathFollowingAgent) = 0;
+	/** Get path following agent this interface uses */
+	virtual IPathFollowingAgentInterface* GetPathFollowingAgent() = 0;
+	/** Get path following agent this interface uses(const) */
+	virtual const IPathFollowingAgentInterface* GetPathFollowingAgent() const = 0;
+	
+	/** path following: request movement through a velocity directly */
+	UFUNCTION(BlueprintCallable, Category="AI|NavMovement")
+	virtual void RequestDirectMove(const FVector& MoveVelocity, bool bForceMaxSpeed) = 0;
+
+	/** path following: request movement through a new move input (normal vector = full strength) */
+	UFUNCTION(BlueprintCallable, Category="AI|NavMovement")
+	virtual void RequestPathMove(const FVector& MoveInput) = 0;
+
+	/** Stops movement by setting velocity to zero - Note: depending on the movement system this may take effect next tick */
+	virtual void StopMovementImmediately() = 0;
+
+	/** Resets runtime movement state to default movement capabilities */
+	virtual void ResetMoveState() = 0;
+	
+	// Movement specific code but not specific to nav movement - could probably be pulled out into a separate interface if desired
+	/** Returns location of controlled agent - meaning center of collision shape */
+	virtual FVector GetLocation() const = 0;
+	
+	/** Returns location of controlled agent's "feet" meaning center of bottom of collision shape */
+	virtual FVector GetFeetLocation() const = 0;
+
+	/** Returns based location of controlled agent */
+	virtual FBasedPosition GetFeetLocationBased() const = 0;
+
+	/** Get the owner of the object consuming nav movement */
+	virtual UObject* GetOwnerAsObject() const = 0;
+	
+	/** Get the Object this movement interface is updating */
+	virtual TObjectPtr<UObject> GetUpdatedObject() const = 0;
+
+	/** Get axis-aligned cylinder around this agent, used for simple collision checks in nav movement */
+	virtual void GetSimpleCollisionCylinder(float& CollisionRadius, float& CollisionHalfHeight) const = 0;
+
+	/** Returns collision extents vector for this object */
+	virtual FVector GetSimpleCollisionCylinderExtent() const = 0;
+	
+	/** Get forward vector of the object being driven by nav movement */
+	virtual FVector GetForwardVector() const = 0;
+	
+	/** Set nav agent properties from an object */
+	virtual void UpdateNavAgent(const UObject& ObjectToUpdateFrom) = 0;
+	
+	/** Get the current velocity of the agent for nav movement */
+	UFUNCTION(BlueprintCallable, Category="AI|NavMovement")
+	virtual FVector GetVelocityForNavMovement() const = 0;
+
+	/** Get maximum movement speed of the agent */
+	UFUNCTION(BlueprintCallable, Category="AI|NavMovement")
+	virtual float GetMaxSpeedForNavMovement() const = 0;
+
+	/** Returns true if currently crouching */ 
+	UFUNCTION(BlueprintCallable, Category="Movement")
+	virtual bool IsCrouching() const = 0;
+
+	/** Returns true if currently falling (not flying, in a non-fluid volume, and not on the ground) */ 
+	UFUNCTION(BlueprintCallable, Category="Movement")
+	virtual bool IsFalling() const = 0;
+
+	/** Returns true if currently moving on the ground (e.g. walking or driving) */ 
+	UFUNCTION(BlueprintCallable, Category="Movement")
+	virtual bool IsMovingOnGround() const = 0;
+	
+	/** Returns true if currently swimming (moving through a fluid volume) */
+	UFUNCTION(BlueprintCallable, Category="Movement")
+	virtual bool IsSwimming() const = 0;
+
+	/** Returns true if currently flying (moving through a non-fluid volume without resting on the ground) */
+	UFUNCTION(BlueprintCallable, Category="Movement")
+	virtual bool IsFlying() const = 0;
+	// End: Functions to be implemented by objects inheriting this interface
+	
+	
+	// Begin: Functions to be implemented here for shared functionality
+	/** Stops applying further movement (usually zeros acceleration). */
+	UFUNCTION(BlueprintCallable, Category="AI|Movement")
+	virtual void StopActiveMovement();
+
+	/** Stops movement immediately (reset velocity) but keeps following current path */
+	UFUNCTION(BlueprintCallable, Category="AI|Movement")
+	virtual void StopMovementKeepPathing();
+	
+	/** Returns navigation location of controlled agent */
+	FVector GetNavLocation() const;
+
+	/** Returns braking distance for acceleration driven path following */
+	virtual float GetPathFollowingBrakingDistance(float MaxSpeed) const;
+
+	/** Set fixed braking distance */
+	void SetFixedBrakingDistance(float DistanceToEndOfPath);
+
+	/** Returns true if acceleration should be used for path following */
+	bool UseAccelerationForPathFollowing() const;
+
+	/** Returns true if agent can crouch */
+	FORCEINLINE virtual bool CanEverCrouch() const { return GetNavAgentPropertiesRef().bCanCrouch; }
+
+	/** Returns true if agent can jump */
+	FORCEINLINE virtual bool CanEverJump() const { return GetNavAgentPropertiesRef().bCanJump; }
+
+	/** Returns true if agent can move along the ground (walk, drive, etc) */
+	FORCEINLINE virtual bool CanEverMoveOnGround() const { return GetNavAgentPropertiesRef().bCanWalk; }
+
+	/** Returns true if agent can swim */
+	FORCEINLINE virtual bool CanEverSwim() const { return GetNavAgentPropertiesRef().bCanSwim; }
+
+	/** Returns true if agent can fly */
+	FORCEINLINE virtual bool CanEverFly() const { return GetNavAgentPropertiesRef().bCanFly; }
+	// End: Functions to be implemented here for shared functionality
+};
diff --git a/Engine/Source/Runtime/Engine/Private/Components/CharacterMovementComponent.cpp b/Engine/Source/Runtime/Engine/Private/Components/CharacterMovementComponent.cpp
--- a/Engine/Source/Runtime/Engine/Private/Components/CharacterMovementComponent.cpp
+++ b/Engine/Source/Runtime/Engine/Private/Components/CharacterMovementComponent.cpp
@@ -899,7 +899,7 @@ void UCharacterMovementComponent::BeginDestroy()
 
 void UCharacterMovementComponent::Deactivate()
 {
-	bStopMovementAbortPaths = false; // Mirrors StopMovementKeepPathing(), because Super calls StopMovement() and we want that handled differently.
+	NavMovementProperties.bStopMovementAbortPaths = false; // Mirrors StopMovementKeepPathing(), because Super calls StopMovement() and we want that handled differently.
 	Super::Deactivate();
 	if (!IsActive())
 	{
@@ -909,7 +909,7 @@ void UCharacterMovementComponent::Deactivate()
 			CharacterOwner->ResetJumpState();
 		}
 	}
-	bStopMovementAbortPaths = true;
+	NavMovementProperties.bStopMovementAbortPaths = true;
 }
 
 
@@ -3900,9 +3900,9 @@ bool UCharacterMovementComponent::CanStopPathFollowing() const
 
 float UCharacterMovementComponent::GetPathFollowingBrakingDistance(float MaxSpeed) const
 {
-	if (bUseFixedBrakingDistanceForPaths)
+	if (NavMovementProperties.bUseFixedBrakingDistanceForPaths)
 	{
-		return FixedPathBrakingDistance;
+		return NavMovementProperties.FixedPathBrakingDistance;
 	}
 
 	const float BrakingDeceleration = FMath::Abs(GetMaxBrakingDeceleration());
diff --git a/Engine/Source/Runtime/Engine/Private/Components/NavMovementComponent.cpp b/Engine/Source/Runtime/Engine/Private/Components/NavMovementComponent.cpp
--- a/Engine/Source/Runtime/Engine/Private/Components/NavMovementComponent.cpp
+++ b/Engine/Source/Runtime/Engine/Private/Components/NavMovementComponent.cpp
@@ -12,10 +12,6 @@
 //----------------------------------------------------------------------//
 UNavMovementComponent::UNavMovementComponent(const FObjectInitializer& ObjectInitializer)
 	: Super(ObjectInitializer)
-	, bUpdateNavAgentWithOwnersCollision(true)
-	, bUseAccelerationForPaths(false)
-	, bUseFixedBrakingDistanceForPaths(false)
-	, bStopMovementAbortPaths(true)
 {
 	bComponentShouldUpdatePhysicsVolume = true;
 }
@@ -25,89 +21,110 @@ FBasedPosition UNavMovementComponent::GetActorFeetLocationBased() const
 	return FBasedPosition(NULL, GetActorFeetLocation());
 }
 
-void UNavMovementComponent::RequestDirectMove(const FVector& MoveVelocity, bool bForceMaxSpeed)
+void UNavMovementComponent::PostLoad()
 {
-	Velocity = MoveVelocity;
+	Super::PostLoad();
+
+#if WITH_EDITORONLY_DATA
+	PRAGMA_DISABLE_DEPRECATION_WARNINGS
+	
+	// copying over any default values to the new struct that houses the properties
+	if (FixedPathBrakingDistance != 0) 
+	{
+		NavMovementProperties.FixedPathBrakingDistance = FixedPathBrakingDistance;
+		FixedPathBrakingDistance = 0;
+	}
+	
+	if (!bUpdateNavAgentWithOwnersCollision)
+	{
+		NavMovementProperties.bUpdateNavAgentWithOwnersCollision = bUpdateNavAgentWithOwnersCollision;
+		bUpdateNavAgentWithOwnersCollision = true;
+	}
+	
+	if (bUseAccelerationForPaths)
+    {
+    	NavMovementProperties.bUseAccelerationForPaths = bUseAccelerationForPaths;
+    	bUseAccelerationForPaths = false;
+    }
+	
+	if (bUseFixedBrakingDistanceForPaths)
+    {
+    	NavMovementProperties.bUseFixedBrakingDistanceForPaths = bUseFixedBrakingDistanceForPaths;
+    	bUseFixedBrakingDistanceForPaths = false;
+    }
+    
+    if (!bStopMovementAbortPaths)
+    {
+    	NavMovementProperties.bStopMovementAbortPaths = bStopMovementAbortPaths;
+    	bStopMovementAbortPaths = true;
+    }
+	
+	PRAGMA_ENABLE_DEPRECATION_WARNINGS
+#endif // WITH_EDITORONLY_DATA	
 }
 
-void UNavMovementComponent::RequestPathMove(const FVector& MoveInput)
+void UNavMovementComponent::UpdateNavAgent(const UObject& ObjectToUpdateFrom)
 {
-	// empty in base class, requires at least PawnMovementComponent for input related operations
+	if (!NavMovementProperties.bUpdateNavAgentWithOwnersCollision)
+	{
+		return;
+	}
+
+	// initialize properties from navigation system
+	NavAgentProps.NavWalkingSearchHeightScale = FNavigationSystem::GetDefaultSupportedAgent().NavWalkingSearchHeightScale;
+	
+	if (const UCapsuleComponent* CapsuleComponent = Cast<UCapsuleComponent>(&ObjectToUpdateFrom))
+	{
+		NavAgentProps.AgentRadius = CapsuleComponent->GetScaledCapsuleRadius();
+		NavAgentProps.AgentHeight = CapsuleComponent->GetScaledCapsuleHalfHeight() * 2.f;;
+	}
+	else if (const AActor* ObjectAsActor = Cast<AActor>(&ObjectToUpdateFrom))
+	{
+		ensureMsgf(&ObjectToUpdateFrom == GetOwner(), TEXT("Object passed to UpdateNavAgent should be the owner actor of the Nav Movement Component"));
+		// Can't call GetSimpleCollisionCylinder(), because no components will be registered.
+		float BoundRadius, BoundHalfHeight;	
+		ObjectAsActor->GetSimpleCollisionCylinder(BoundRadius, BoundHalfHeight);
+		NavAgentProps.AgentRadius = BoundRadius;
+		NavAgentProps.AgentHeight = BoundHalfHeight * 2.f;
+	}
 }
 
-bool UNavMovementComponent::CanStopPathFollowing() const
+void UNavMovementComponent::RequestDirectMove(const FVector& MoveVelocity, bool bForceMaxSpeed)
 {
-	return true;
+	Velocity = MoveVelocity;
 }
 
-float UNavMovementComponent::GetPathFollowingBrakingDistance(float MaxSpeed) const
+void UNavMovementComponent::RequestPathMove(const FVector& MoveInput)
 {
-	return bUseFixedBrakingDistanceForPaths ? FixedPathBrakingDistance : MaxSpeed;
+	// empty in base class, requires at least PawnMovementComponent for input related operations
 }
 
-void UNavMovementComponent::SetFixedBrakingDistance(float DistanceToEndOfPath)
+bool UNavMovementComponent::CanStopPathFollowing() const
 {
-	if (DistanceToEndOfPath > UE_KINDA_SMALL_NUMBER)
-	{
-		bUseFixedBrakingDistanceForPaths = true;
-		FixedPathBrakingDistance = DistanceToEndOfPath;
-	}
+	return true;
 }
 
 void UNavMovementComponent::ClearFixedBrakingDistance()
 {
-	bUseFixedBrakingDistanceForPaths = false;
+	NavMovementProperties.bUseFixedBrakingDistanceForPaths = false;
 }
 
-void UNavMovementComponent::StopActiveMovement()
+void UNavMovementComponent::GetSimpleCollisionCylinder(float& CollisionRadius, float& CollisionHalfHeight) const
 {
-	if (!bStopMovementAbortPaths)
-	{
-		return;
-	}
-
-	IPathFollowingAgentInterface* PFAgent = GetPathFollowingAgent();
-	if (PFAgent)
-	{
-		PFAgent->OnUnableToMove(*this);
-	}
+	GetOwner()->GetSimpleCollisionCylinder(CollisionRadius, CollisionHalfHeight);
 }
 
-void UNavMovementComponent::UpdateNavAgent(const AActor& Owner)
+FVector UNavMovementComponent::GetSimpleCollisionCylinderExtent() const
 {
-	ensure(&Owner == GetOwner());
-	if (ShouldUpdateNavAgentWithOwnersCollision() == false)
-	{
-		return;
-	}
-
-	// initialize properties from navigation system
-	NavAgentProps.NavWalkingSearchHeightScale = FNavigationSystem::GetDefaultSupportedAgent().NavWalkingSearchHeightScale;
-	
-	// Can't call GetSimpleCollisionCylinder(), because no components will be registered.
-	float BoundRadius, BoundHalfHeight;	
-	Owner.GetSimpleCollisionCylinder(BoundRadius, BoundHalfHeight);
-	NavAgentProps.AgentRadius = BoundRadius;
-	NavAgentProps.AgentHeight = BoundHalfHeight * 2.f;
+	return GetOwner()->GetSimpleCollisionCylinderExtent();
 }
 
-void UNavMovementComponent::UpdateNavAgent(const UCapsuleComponent& CapsuleComponent)
+FVector UNavMovementComponent::GetForwardVector() const
 {
-	if (ShouldUpdateNavAgentWithOwnersCollision() == false)
-	{
-		return;
-	}
-
-	// initialize properties from navigation system
-	NavAgentProps.NavWalkingSearchHeightScale = FNavigationSystem::GetDefaultSupportedAgent().NavWalkingSearchHeightScale;
-
-	NavAgentProps.AgentRadius = CapsuleComponent.GetScaledCapsuleRadius();
-	NavAgentProps.AgentHeight = CapsuleComponent.GetScaledCapsuleHalfHeight() * 2.f;
+	return GetOwner()->GetActorForwardVector();
 }
 
 void UNavMovementComponent::SetUpdateNavAgentWithOwnersCollisions(bool bUpdateWithOwner)
 {
-	bUpdateNavAgentWithOwnersCollision = bUpdateWithOwner;
+	NavMovementProperties.bUpdateNavAgentWithOwnersCollision = bUpdateWithOwner;
 }
-
-
diff --git a/Engine/Source/Runtime/Engine/Private/FloatingPawnMovement.cpp b/Engine/Source/Runtime/Engine/Private/FloatingPawnMovement.cpp
--- a/Engine/Source/Runtime/Engine/Private/FloatingPawnMovement.cpp
+++ b/Engine/Source/Runtime/Engine/Private/FloatingPawnMovement.cpp
@@ -38,7 +38,7 @@ void UFloatingPawnMovement::TickComponent(float DeltaTime, enum ELevelTick TickT
 	if (Controller && Controller->IsLocalController())
 	{
 		// apply input for local players but also for AI that's not following a navigation path at the moment
-		if (Controller->IsLocalPlayerController() == true || Controller->IsFollowingAPath() == false || bUseAccelerationForPaths)
+		if (Controller->IsLocalPlayerController() == true || Controller->IsFollowingAPath() == false || NavMovementProperties.bUseAccelerationForPaths)
 		{
 			ApplyControlInputToVelocity(DeltaTime);
 		}
diff --git a/Engine/Source/Runtime/Engine/Private/NavMovementInterface.cpp b/Engine/Source/Runtime/Engine/Private/NavMovementInterface.cpp
new file mode 100644
--- /dev/null
+++ b/Engine/Source/Runtime/Engine/Private/NavMovementInterface.cpp
@@ -0,0 +1,54 @@
+﻿// Copyright Epic Games, Inc. All Rights Reserved.
+
+#include "GameFramework/NavMovementInterface.h"
+#include "AI/Navigation/NavAgentInterface.h"
+#include "AI/Navigation/PathFollowingAgentInterface.h"
+
+void INavMovementInterface::StopActiveMovement()
+{
+	if (!GetNavMovementProperties()->bStopMovementAbortPaths)
+	{
+		return;
+	}
+
+	IPathFollowingAgentInterface* PFAgent = GetPathFollowingAgent();
+	if (PFAgent)
+	{
+		PFAgent->OnUnableToMove(*Cast<UObject>(this));
+	}
+}
+
+void INavMovementInterface::StopMovementKeepPathing()
+{
+	FNavMovementProperties* NavMovementProperties = GetNavMovementProperties();
+	NavMovementProperties->bStopMovementAbortPaths = false;
+	StopMovementImmediately();
+	NavMovementProperties->bStopMovementAbortPaths = true;
+}
+
+FVector INavMovementInterface::GetNavLocation() const
+{
+	const INavAgentInterface* MyOwner = Cast<INavAgentInterface>(GetOwnerAsObject());
+	return MyOwner ? MyOwner->GetNavAgentLocation() : FNavigationSystem::InvalidLocation; 
+}
+
+float INavMovementInterface::GetPathFollowingBrakingDistance(float MaxSpeed) const
+{
+	const FNavMovementProperties NavMovementProperties = GetNavMovementProperties();
+	return NavMovementProperties.bUseFixedBrakingDistanceForPaths ? NavMovementProperties.FixedPathBrakingDistance : MaxSpeed;
+}
+
+void INavMovementInterface::SetFixedBrakingDistance(float DistanceToEndOfPath)
+{
+	if (DistanceToEndOfPath > UE_KINDA_SMALL_NUMBER)
+	{
+		FNavMovementProperties* NavMovementProperties = GetNavMovementProperties();
+		NavMovementProperties->bUseFixedBrakingDistanceForPaths = true;
+		NavMovementProperties->FixedPathBrakingDistance = DistanceToEndOfPath;
+	}
+}
+
+bool INavMovementInterface::UseAccelerationForPathFollowing() const
+{
+	return GetNavMovementProperties().bUseAccelerationForPaths;
+}
