Add utility method in TSimpleGrid

diff --git a/Engine/Source/Runtime/AIModule/Public/SimpleCellGrid.h b/Engine/Source/Runtime/AIModule/Public/SimpleCellGrid.h
--- a/Engine/Source/Runtime/AIModule/Public/SimpleCellGrid.h
+++ b/Engine/Source/Runtime/AIModule/Public/SimpleCellGrid.h
@@ -295,7 +295,10 @@ public:
 
 	FORCEINLINE FCellType& GetCellAtIndexUnsafe(int32 CellIndex) { return Cells.GetData()[CellIndex]; }
 	FORCEINLINE const FCellType& GetCellAtIndexUnsafe(int32 CellIndex) const { return Cells.GetData()[CellIndex]; }
-	
+
+	FCellType& GetCellAtCoordsUnsafe(int32 LocationX, int32 LocationY) { return GetCellAtIndexUnsafe(GetCellIndexUnsafe(LocationX, LocationY)); }
+	const FCellType& GetCellAtCoordsUnsafe(int32 LocationX, int32 LocationY) const { return GetCellAtIndexUnsafe(GetCellIndexUnsafe(LocationX, LocationY)); }
+
 	FORCEINLINE int32 GetCellsCount() const
 	{
 		return Cells.Num();
