AIPerceptionComponent: Lower warning severity and change message, this component can be used on an actor that is not an AI controller
#rnx

diff --git a/Engine/Source/Runtime/AIModule/Private/Perception/AIPerceptionComponent.cpp b/Engine/Source/Runtime/AIModule/Private/Perception/AIPerceptionComponent.cpp
--- a/Engine/Source/Runtime/AIModule/Private/Perception/AIPerceptionComponent.cpp
+++ b/Engine/Source/Runtime/AIModule/Private/Perception/AIPerceptionComponent.cpp
@@ -186,13 +186,9 @@ void UAIPerceptionComponent::OnRegister()
 	bCleanedUp = false;
 
 	AActor* Owner = GetOwner();
-	if (Owner != nullptr)
+	if (Owner)
 	{
 		Owner->OnEndPlay.AddUniqueDynamic(this, &UAIPerceptionComponent::OnOwnerEndPlay);
-		AIOwner = Cast<AAIController>(Owner);
-
-		// Whilst it should be possible with some code changes, to make perception components work when being added to other AActors than AIControllers, it's not something Epic support.
-		UE_CVLOG_UELOG(!AIOwner && Owner->GetWorld() && (Owner->GetWorld()->WorldType != EWorldType::Editor), Owner, LogAIPerception, Warning, TEXT("%s: Perception Component is being registered with %s, they are designed to work with AAIControllers!"), ANSI_TO_TCHAR(__FUNCTION__), *Owner->GetName());
 	}
 
 	UAIPerceptionSystem* AIPerceptionSys = UAIPerceptionSystem::GetCurrent(GetWorld());
@@ -214,15 +210,24 @@ void UAIPerceptionComponent::OnRegister()
 			AIPerceptionSys->UpdateListener(*this);
 		}
 	}
-
-	// this should not be needed but apparently AAIController::PostRegisterAllComponents
-	// gets called component's OnRegister
-	AIOwner = Cast<AAIController>(GetOwner());
-	ensure(AIOwner == nullptr || AIOwner->GetAIPerceptionComponent() == nullptr || AIOwner->GetAIPerceptionComponent() == this
-		|| (AIOwner->GetWorld() && AIOwner->GetWorld()->WorldType != EWorldType::Editor));
-	if (AIOwner && AIOwner->GetAIPerceptionComponent() == nullptr)
+	
+	AIOwner = Cast<AAIController>(Owner);
+	if (AIOwner)
+	{
+		if (AIOwner->GetAIPerceptionComponent() == nullptr)
+		{
+			// this should not be needed but apparently AAIController::PostRegisterAllComponents
+			// gets called component's OnRegister
+			AIOwner->SetPerceptionComponent(*this);
+		}
+		else
+		{
+			ensure(AIOwner->GetAIPerceptionComponent() == this || (AIOwner->GetWorld() && AIOwner->GetWorld()->WorldType != EWorldType::Editor));
+		}
+	}
+	else if (Owner)
 	{
-		AIOwner->SetPerceptionComponent(*this);
+		UE_CVLOG_UELOG(Owner->GetWorld() && (Owner->GetWorld()->WorldType != EWorldType::Editor), Owner, LogAIPerception, Log, TEXT("%s: Owner %s is not an AIController so make sure you bind to OnPerceptionUpdated, otherwise you won't get notified about updated actor perception events"), ANSI_TO_TCHAR(__FUNCTION__), *Owner->GetName());
 	}
 }
 
@@ -595,7 +600,7 @@ void UAIPerceptionComponent::ProcessStimuli()
 
 	if (UpdatedActors.Num() > 0)
 	{
-		if (AIOwner != NULL)
+		if (AIOwner != nullptr)
 		{
 			AIOwner->ActorsPerceptionUpdated(UpdatedActors);
 		}
