Refactor Navmesh GameplayDebugger category to be able to display other kind of navigation data

diff --git a/Engine/Source/Runtime/AIModule/Private/GameplayDebugger/GameplayDebuggerCategory_Navmesh.cpp b/Engine/Source/Runtime/AIModule/Private/GameplayDebugger/GameplayDebuggerCategory_Navmesh.cpp
--- a/Engine/Source/Runtime/AIModule/Private/GameplayDebugger/GameplayDebuggerCategory_Navmesh.cpp
+++ b/Engine/Source/Runtime/AIModule/Private/GameplayDebugger/GameplayDebuggerCategory_Navmesh.cpp
@@ -111,14 +111,14 @@ void FGameplayDebuggerCategory_Navmesh::FRepData::Serialize(FArchive& Ar)
 
 void FGameplayDebuggerCategory_Navmesh::CollectData(APlayerController* OwnerPC, AActor* DebugActor)
 {
-#if WITH_RECAST
+	UNavigationSystemV1* NavSys = nullptr;
 	ANavigationData* NavData = nullptr;
 	const APawn* RefPawn = nullptr;
 	int32 NumNavData = 0;
 
 	if (OwnerPC != nullptr)
 	{
-		UNavigationSystemV1* NavSys = FNavigationSystem::GetCurrent<UNavigationSystemV1>(OwnerPC->GetWorld());
+		NavSys = FNavigationSystem::GetCurrent<UNavigationSystemV1>(OwnerPC->GetWorld());
 		if (NavSys) 
 		{
 			DataPack.NumDirtyAreas = NavSys->GetNumDirtyAreas();
@@ -217,6 +217,12 @@ void FGameplayDebuggerCategory_Navmesh::CollectData(APlayerController* OwnerPC,
 		}
 	}
 
+	CollectNavigationData(NavSys, NavData, RefPawn);
+}
+
+void FGameplayDebuggerCategory_Navmesh::CollectNavigationData(UNavigationSystemV1* NavSys, ANavigationData* NavData, const APawn* RefPawn)
+{
+#if WITH_RECAST
 	const ARecastNavMesh* RecastNavMesh = Cast<const ARecastNavMesh>(NavData);
 	if (RecastNavMesh && RefPawn)
 	{
@@ -224,32 +230,18 @@ void FGameplayDebuggerCategory_Navmesh::CollectData(APlayerController* OwnerPC,
 		// Note that we round up to the next odd number to keep the reference position in the middle tile
 		const FVector TargetLocation = RefPawn->GetActorLocation();
 
-		int32 NumTilesPerSide = FMath::Max(FGameplayDebuggerCategoryNavmeshTweakables::DisplaySize, 1);
-		NumTilesPerSide += (NumTilesPerSide % 2 == 0) ? 1 : 0;
-
-		const int32 NumTilesToDisplay = NumTilesPerSide * NumTilesPerSide;
-
-		TArray<int32> DeltaX;
-		TArray<int32> DeltaY;
-		DeltaX.AddUninitialized(NumTilesToDisplay);
-		DeltaY.AddUninitialized(NumTilesToDisplay);
-
-		const int32 MinIdx = -(NumTilesPerSide >> 1);
-		for (int32 i=0; i < NumTilesToDisplay; ++i)
-		{
-			DeltaX[i] = MinIdx + (i % NumTilesPerSide);
-			DeltaY[i] = MinIdx + (i / NumTilesPerSide);
-		}
+		TArray<FIntPoint> TileDeltas;
+		RetrieveRelativeTilesToDisplay(TileDeltas);
 
 		int32 TargetTileX = 0;
 		int32 TargetTileY = 0;
 		RecastNavMesh->GetNavMeshTileXY(TargetLocation, TargetTileX, TargetTileY);
 
 		TArray<FNavTileRef> TileSet;
-		for (int32 Idx = 0; Idx < NumTilesToDisplay; Idx++)
+		for (const FIntPoint& TileDelta : TileDeltas)
 		{
-			const int32 NeiX = TargetTileX + DeltaX[Idx];
-			const int32 NeiY = TargetTileY + DeltaY[Idx];
+			const int32 NeiX = TargetTileX + TileDelta.X;
+			const int32 NeiY = TargetTileY + TileDelta.Y;
 			RecastNavMesh->GetNavMeshTilesAt(NeiX, NeiY, TileSet);
 		}
 
@@ -273,6 +265,22 @@ void FGameplayDebuggerCategory_Navmesh::CollectData(APlayerController* OwnerPC,
 #endif // WITH_RECAST
 }
 
+void FGameplayDebuggerCategory_Navmesh::RetrieveRelativeTilesToDisplay(TArray<FIntPoint>& OutTileDelta)
+{
+	int32 NumTilesPerSide = FMath::Max(FGameplayDebuggerCategoryNavmeshTweakables::DisplaySize, 1);
+	NumTilesPerSide += (NumTilesPerSide % 2 == 0) ? 1 : 0;
+
+	const int32 NumTilesToDisplay = NumTilesPerSide * NumTilesPerSide;
+
+	OutTileDelta.AddUninitialized(NumTilesToDisplay);
+
+	const int32 MinIdx = -(NumTilesPerSide >> 1);
+	for (int32 i=0; i < NumTilesToDisplay; ++i)
+	{
+		OutTileDelta[i] = FIntPoint(MinIdx + (i % NumTilesPerSide), MinIdx + (i / NumTilesPerSide));
+	}
+}
+
 void FGameplayDebuggerCategory_Navmesh::DrawData(APlayerController* OwnerPC, FGameplayDebuggerCanvasContext& CanvasContext)
 {
 	CanvasContext.Printf(TEXT("Num dirty areas: {%s}%d"), DataPack.NumDirtyAreas > 0 ? TEXT("red") : TEXT("green"), DataPack.NumDirtyAreas);
diff --git a/Engine/Source/Runtime/AIModule/Public/GameplayDebugger/GameplayDebuggerCategory_Navmesh.h b/Engine/Source/Runtime/AIModule/Public/GameplayDebugger/GameplayDebuggerCategory_Navmesh.h
--- a/Engine/Source/Runtime/AIModule/Public/GameplayDebugger/GameplayDebuggerCategory_Navmesh.h
+++ b/Engine/Source/Runtime/AIModule/Public/GameplayDebugger/GameplayDebuggerCategory_Navmesh.h
@@ -9,6 +9,7 @@
 #include "NavMesh/NavMeshRenderingComponent.h"
 
 class APlayerController;
+class UNavigationSystemV1;
 
 class FGameplayDebuggerCategory_Navmesh : public FGameplayDebuggerCategory
 {
@@ -26,6 +27,8 @@ protected:
 
 	void CycleNavData();
 	void CycleActorReference();
+	AIMODULE_API virtual void CollectNavigationData(UNavigationSystemV1* NavSys, ANavigationData* NavData, const APawn* RefPawn);
+	AIMODULE_API void RetrieveRelativeTilesToDisplay(TArray<FIntPoint>& OutTileDelta);
 
 	struct FRepData
 	{
