Change FGameplayDebuggerCategory_Navmesh::CollectNavigationData to use const pointer parameters

diff --git a/Engine/Source/Runtime/AIModule/Private/GameplayDebugger/GameplayDebuggerCategory_Navmesh.cpp b/Engine/Source/Runtime/AIModule/Private/GameplayDebugger/GameplayDebuggerCategory_Navmesh.cpp
--- a/Engine/Source/Runtime/AIModule/Private/GameplayDebugger/GameplayDebuggerCategory_Navmesh.cpp
+++ b/Engine/Source/Runtime/AIModule/Private/GameplayDebugger/GameplayDebuggerCategory_Navmesh.cpp
@@ -220,7 +220,7 @@ void FGameplayDebuggerCategory_Navmesh::CollectData(APlayerController* OwnerPC,
 	CollectNavigationData(NavSys, NavData, RefPawn);
 }
 
-void FGameplayDebuggerCategory_Navmesh::CollectNavigationData(UNavigationSystemV1* NavSys, ANavigationData* NavData, const APawn* RefPawn)
+void FGameplayDebuggerCategory_Navmesh::CollectNavigationData(const UNavigationSystemV1* NavSys, const ANavigationData* NavData, const APawn* RefPawn)
 {
 #if WITH_RECAST
 	const ARecastNavMesh* RecastNavMesh = Cast<const ARecastNavMesh>(NavData);
diff --git a/Engine/Source/Runtime/AIModule/Public/GameplayDebugger/GameplayDebuggerCategory_Navmesh.h b/Engine/Source/Runtime/AIModule/Public/GameplayDebugger/GameplayDebuggerCategory_Navmesh.h
--- a/Engine/Source/Runtime/AIModule/Public/GameplayDebugger/GameplayDebuggerCategory_Navmesh.h
+++ b/Engine/Source/Runtime/AIModule/Public/GameplayDebugger/GameplayDebuggerCategory_Navmesh.h
@@ -27,7 +27,7 @@ protected:
 
 	void CycleNavData();
 	void CycleActorReference();
-	AIMODULE_API virtual void CollectNavigationData(UNavigationSystemV1* NavSys, ANavigationData* NavData, const APawn* RefPawn);
+	AIMODULE_API virtual void CollectNavigationData(const UNavigationSystemV1* NavSys, const ANavigationData* NavData, const APawn* RefPawn);
 	AIMODULE_API void RetrieveRelativeTilesToDisplay(TArray<FIntPoint>& OutTileDelta);
 
 	struct FRepData
