Create BlackboardKeyType_Struct and associated ValueOrBBKey_Struct

diff --git a/Engine/Source/Editor/BehaviorTreeEditor/Classes/ValueOrBBKeyDetails.h b/Engine/Source/Editor/BehaviorTreeEditor/Classes/ValueOrBBKeyDetails.h
--- a/Engine/Source/Editor/BehaviorTreeEditor/Classes/ValueOrBBKeyDetails.h
+++ b/Engine/Source/Editor/BehaviorTreeEditor/Classes/ValueOrBBKeyDetails.h
@@ -111,3 +111,16 @@ protected:
 	void OnObjectChanged(const FAssetData& AssetData);
 	FString OnGetObjectPath() const;
 };
+
+class FValueOrBBKeyDetails_Struct : public FValueOrBBKeyDetails
+{
+public:
+	/** Makes a new instance of this detail layout class for a specific detail view requesting it */
+	static TSharedRef<IPropertyTypeCustomization> MakeInstance();
+
+	virtual void CustomizeHeader(TSharedRef<class IPropertyHandle> StructPropertyHandle, class FDetailWidgetRow& HeaderRow, IPropertyTypeCustomizationUtils& StructCustomizationUtils) override;
+	virtual void CustomizeChildren(TSharedRef<class IPropertyHandle> StructPropertyHandle, class IDetailChildrenBuilder& StructBuilder, IPropertyTypeCustomizationUtils& StructCustomizationUtils) override;
+
+protected:
+	TSharedPtr<IPropertyHandle> EditDefaultsOnlyProperty;
+};
\ No newline at end of file
diff --git a/Engine/Source/Editor/BehaviorTreeEditor/Private/BehaviorTreeEditorModule.cpp b/Engine/Source/Editor/BehaviorTreeEditor/Private/BehaviorTreeEditorModule.cpp
--- a/Engine/Source/Editor/BehaviorTreeEditor/Private/BehaviorTreeEditorModule.cpp
+++ b/Engine/Source/Editor/BehaviorTreeEditor/Private/BehaviorTreeEditorModule.cpp
@@ -70,6 +70,7 @@ void FBehaviorTreeEditorModule::StartupModule()
 	PropertyModule.RegisterCustomPropertyTypeLayout("ValueOrBBKey_Object", FOnGetPropertyTypeCustomizationInstance::CreateStatic(&FValueOrBBKeyDetails_Object::MakeInstance));
 	PropertyModule.RegisterCustomPropertyTypeLayout("ValueOrBBKey_Rotator", FOnGetPropertyTypeCustomizationInstance::CreateStatic(&FValueOrBBKeyDetails::MakeInstance));
 	PropertyModule.RegisterCustomPropertyTypeLayout("ValueOrBBKey_Vector", FOnGetPropertyTypeCustomizationInstance::CreateStatic(&FValueOrBBKeyDetails::MakeInstance));
+	PropertyModule.RegisterCustomPropertyTypeLayout("ValueOrBBKey_Struct", FOnGetPropertyTypeCustomizationInstance::CreateStatic(&FValueOrBBKeyDetails_Struct::MakeInstance));
 	PropertyModule.NotifyCustomizationModuleChanged();
 }
 
@@ -110,6 +111,7 @@ void FBehaviorTreeEditorModule::ShutdownModule()
 		PropertyModule.UnregisterCustomPropertyTypeLayout("ValueOrBBKey_Object");
 		PropertyModule.UnregisterCustomPropertyTypeLayout("ValueOrBBKey_Rotator");
 		PropertyModule.UnregisterCustomPropertyTypeLayout("ValueOrBBKey_Vector");
+		PropertyModule.UnregisterCustomPropertyTypeLayout("ValueOrBBKey_Struct");
 		PropertyModule.NotifyCustomizationModuleChanged();
 	}
 }
diff --git a/Engine/Source/Editor/BehaviorTreeEditor/Private/DetailCustomizations/ValueOrBBKeyDetails.cpp b/Engine/Source/Editor/BehaviorTreeEditor/Private/DetailCustomizations/ValueOrBBKeyDetails.cpp
--- a/Engine/Source/Editor/BehaviorTreeEditor/Private/DetailCustomizations/ValueOrBBKeyDetails.cpp
+++ b/Engine/Source/Editor/BehaviorTreeEditor/Private/DetailCustomizations/ValueOrBBKeyDetails.cpp
@@ -59,6 +59,11 @@ TSharedRef<IPropertyTypeCustomization> FValueOrBBKeyDetails_Object::MakeInstance
 	return MakeShareable(new FValueOrBBKeyDetails_Object);
 }
 
+TSharedRef<IPropertyTypeCustomization> FValueOrBBKeyDetails_Struct::MakeInstance()
+{
+	return MakeShareable(new FValueOrBBKeyDetails_Struct);
+}
+
 BEGIN_SLATE_FUNCTION_BUILD_OPTIMIZATION
 void FValueOrBBKeyDetails::CustomizeHeader(TSharedRef<IPropertyHandle> StructPropertyHandle, FDetailWidgetRow& HeaderRow, IPropertyTypeCustomizationUtils& StructCustomizationUtils)
 {
@@ -239,6 +244,27 @@ TSharedRef<SWidget> FValueOrBBKeyDetails_Object::CreateDefaultValueWidget()
 	}
 	return FValueOrBBKeyDetails::CreateDefaultValueWidget();
 }
+
+void FValueOrBBKeyDetails_Struct::CustomizeHeader(TSharedRef<class IPropertyHandle> StructPropertyHandle, class FDetailWidgetRow& HeaderRow, IPropertyTypeCustomizationUtils& StructCustomizationUtils)
+{
+	FValueOrBBKeyDetails::CustomizeHeader(StructPropertyHandle, HeaderRow, StructCustomizationUtils);
+	HeaderRow.ShouldAutoExpand(true);
+
+	EditDefaultsOnlyProperty = StructPropertyHandle->GetChildHandle(TEXT("bCanEditDefaultValueType"));
+	// We dont have access to the EditDefaultsOnly property or it's not editable. Lock the type so node instance can't change the DefaultValue type.
+	if (!EditDefaultsOnlyProperty.IsValid() || !EditDefaultsOnlyProperty->IsEditable())
+	{
+		DefaultValueProperty->SetInstanceMetaData(TEXT("StructTypeConst"), TEXT(""));
+	}
+}
+
+void FValueOrBBKeyDetails_Struct::CustomizeChildren(TSharedRef<class IPropertyHandle> StructPropertyHandle, class IDetailChildrenBuilder& StructBuilder, IPropertyTypeCustomizationUtils& StructCustomizationUtils)
+{
+	FValueOrBBKeyDetails::CustomizeChildren(StructPropertyHandle, StructBuilder, StructCustomizationUtils);
+	StructBuilder.AddProperty(DefaultValueProperty.ToSharedRef())
+	.IsEnabled(TAttribute<bool>(this, &FValueOrBBKeyDetails::CanEditDefaultValue));
+}
+
 END_SLATE_FUNCTION_BUILD_OPTIMIZATION
 
 bool FValueOrBBKeyDetails::HasAccessToBlackboard() const
diff --git a/Engine/Source/Editor/EditorStyle/Private/SlateEditorStyle.cpp b/Engine/Source/Editor/EditorStyle/Private/SlateEditorStyle.cpp
--- a/Engine/Source/Editor/EditorStyle/Private/SlateEditorStyle.cpp
+++ b/Engine/Source/Editor/EditorStyle/Private/SlateEditorStyle.cpp
@@ -4163,6 +4163,7 @@ void FSlateEditorStyle::FStyle::SetupGraphEditorStyles()
 			Set( "ClassIcon.BlackboardKeyType_Object", new IMAGE_BRUSH( "Icons/pill_16x", Icon16x16, FLinearColor(0.0f, 0.4f, 0.910000f, 1.0f) ) );
 			Set( "ClassIcon.BlackboardKeyType_Rotator", new IMAGE_BRUSH( "Icons/pill_16x", Icon16x16, FLinearColor(0.353393f, 0.454175f, 1.0f, 1.0f) ) );
 			Set( "ClassIcon.BlackboardKeyType_String", new IMAGE_BRUSH( "Icons/pill_16x", Icon16x16, FLinearColor(1.0f, 0.0f, 0.660537f, 1.0f) ) );
+			Set( "ClassIcon.BlackboardKeyType_Struct", new IMAGE_BRUSH("Icons/pill_16x", Icon16x16, FLinearColor(0.f, 0.349019f, 0.796070f, 1.0f)));
 			Set( "ClassIcon.BlackboardKeyType_Vector", new IMAGE_BRUSH( "Icons/pill_16x", Icon16x16, FLinearColor(1.0f, 0.591255f, 0.016512f, 1.0f) ) );
 
 			Set( "BTEditor.Common.NewBlackboard", new IMAGE_BRUSH( "BehaviorTree/NewBlackboard_40x", Icon40x40));
diff --git a/Engine/Source/Editor/EditorStyle/Private/StarshipStyle.cpp b/Engine/Source/Editor/EditorStyle/Private/StarshipStyle.cpp
--- a/Engine/Source/Editor/EditorStyle/Private/StarshipStyle.cpp
+++ b/Engine/Source/Editor/EditorStyle/Private/StarshipStyle.cpp
@@ -3646,6 +3646,7 @@ void FStarshipEditorStyle::FStyle::SetupGraphEditorStyles()
 			Set( "ClassIcon.BlackboardKeyType_Object", new FSlateRoundedBoxBrush( FLinearColor(0.0f, 0.4f, 0.910000f, 1.0f), 2.5f, FVector2D(16.f, 5.f) ) );
 			Set( "ClassIcon.BlackboardKeyType_Rotator", new FSlateRoundedBoxBrush( FLinearColor(0.353393f, 0.454175f, 1.0f, 1.0f), 2.5f, FVector2D(16.f, 5.f) ) );
 			Set( "ClassIcon.BlackboardKeyType_String", new FSlateRoundedBoxBrush( FLinearColor(1.0f, 0.0f, 0.660537f, 1.0f), 2.5f, FVector2D(16.f, 5.f) ) );
+			Set( "ClassIcon.BlackboardKeyType_Struct", new FSlateRoundedBoxBrush( FLinearColor(0.f, 0.349019f, 0.796070f, 1.0f), 2.5f, FVector2D(16.f, 5.f)));
 			Set( "ClassIcon.BlackboardKeyType_Vector", new FSlateRoundedBoxBrush( FLinearColor(1.0f, 0.591255f, 0.016512f, 1.0f), 2.5f, FVector2D(16.f, 5.f) ) );
 
 			Set( "BTEditor.Common.NewBlackboard", new IMAGE_BRUSH_SVG( "Starship/BehaviorTree/Blackboard_20", Icon20x20));
diff --git a/Engine/Source/Runtime/AIModule/Classes/BehaviorTree/BehaviorTreeTypes.h b/Engine/Source/Runtime/AIModule/Classes/BehaviorTree/BehaviorTreeTypes.h
--- a/Engine/Source/Runtime/AIModule/Classes/BehaviorTree/BehaviorTreeTypes.h
+++ b/Engine/Source/Runtime/AIModule/Classes/BehaviorTree/BehaviorTreeTypes.h
@@ -665,6 +665,13 @@ public:
 	AIMODULE_API void AddRotatorFilter(UObject* Owner, FName PropertyName);
 	AIMODULE_API void AddStringFilter(UObject* Owner, FName PropertyName);
 	AIMODULE_API void AddNameFilter(UObject* Owner, FName PropertyName);
+	AIMODULE_API void AddStructFilter(UObject* Owner, FName PropertyName, const UScriptStruct* AllowedStruct);
+
+	template<typename T>
+	void AddStructFilter(UObject* Owner, FName PropertyName)
+	{
+		AddStructFilter(Owner, PropertyName, TBaseStructure<T>::Get());
+	}
 
 	FORCEINLINE bool IsNone() const { return bNoneIsAllowedValue && GetSelectedKeyID() == FBlackboard::InvalidKey; }
 	FORCEINLINE bool IsSet() const { return GetSelectedKeyID() != FBlackboard::InvalidKey; }
diff --git a/Engine/Source/Runtime/AIModule/Classes/BehaviorTree/Blackboard/BlackboardKeyAllTypes.h b/Engine/Source/Runtime/AIModule/Classes/BehaviorTree/Blackboard/BlackboardKeyAllTypes.h
--- a/Engine/Source/Runtime/AIModule/Classes/BehaviorTree/Blackboard/BlackboardKeyAllTypes.h
+++ b/Engine/Source/Runtime/AIModule/Classes/BehaviorTree/Blackboard/BlackboardKeyAllTypes.h
@@ -13,5 +13,6 @@
 #include "BehaviorTree/Blackboard/BlackboardKeyType_Object.h"
 #include "BehaviorTree/Blackboard/BlackboardKeyType_Rotator.h"
 #include "BehaviorTree/Blackboard/BlackboardKeyType_String.h"
+#include "BehaviorTree/Blackboard/BlackboardKeyType_Struct.h"
 #include "BehaviorTree/Blackboard/BlackboardKeyType_Vector.h"
 #include "BehaviorTree/BlackboardComponent.h"
diff --git a/Engine/Source/Runtime/AIModule/Classes/BehaviorTree/Blackboard/BlackboardKeyType_Struct.h b/Engine/Source/Runtime/AIModule/Classes/BehaviorTree/Blackboard/BlackboardKeyType_Struct.h
new file mode 100644
--- /dev/null
+++ b/Engine/Source/Runtime/AIModule/Classes/BehaviorTree/Blackboard/BlackboardKeyType_Struct.h
@@ -0,0 +1,56 @@
+// Copyright Epic Games, Inc. All Rights Reserved.
+
+#pragma once
+
+#include "BehaviorTree/Blackboard/BlackboardKeyType.h"
+#include "StructUtils/InstancedStruct.h"
+
+#include "BlackboardKeyType_Struct.generated.h"
+
+class UBlackboardComponent;
+
+struct FConstStructView;
+
+UCLASS(EditInlineNew, MinimalAPI, meta = (DisplayName = "Struct"))
+class UBlackboardKeyType_Struct : public UBlackboardKeyType
+{
+	GENERATED_BODY()
+
+public:
+	using FDataType = FConstStructView;
+	static const FDataType InvalidValue;
+
+	UBlackboardKeyType_Struct(const FObjectInitializer& ObjectInitializer);
+
+	static FConstStructView GetValue(UBlackboardKeyType_Struct* KeyOb, const uint8* RawData);
+	static bool SetValue(UBlackboardKeyType_Struct* KeyOb, uint8* RawData, FConstStructView Value);
+
+	virtual EBlackboardCompare::Type CompareValues(const UBlackboardComponent& OwnerComp, const uint8* MemoryBlock, const UBlackboardKeyType* OtherKeyOb, const uint8* OtherMemoryBlock) const override;
+	virtual void CopyValues(UBlackboardComponent& OwnerComp, uint8* MemoryBlock, const UBlackboardKeyType* SourceKeyOb, const uint8* SourceBlock) override;
+
+	virtual FString DescribeSelf() const override;
+	virtual bool IsAllowedByFilter(UBlackboardKeyType* FilterOb) const override;
+
+	virtual void InitializeMemory(UBlackboardComponent& OwnerComp, uint8* MemoryBlock) override;
+	virtual void FreeMemory(UBlackboardComponent& OwnerComp, uint8* MemoryBlock) override;
+	virtual void Clear(UBlackboardComponent& OwnerComp, uint8* MemoryBlock) override;
+
+	virtual void PostLoad() override;
+
+#if WITH_EDITOR
+	virtual void PostEditChangeProperty(struct FPropertyChangedEvent& PropertyChangedEvent) override;
+#endif // WITH_EDITOR
+
+	UPROPERTY(Category = Blackboard, EditDefaultsOnly, DisplayName = DefaultValue)
+	FInstancedStruct DefaultValue;
+
+protected:
+	// Runtime value of the key if it is instanced otherwise is empty.
+	UPROPERTY(Transient)
+	FInstancedStruct Value;
+
+	void UpdateNeedsInstance();
+
+	virtual FString DescribeValue(const UBlackboardComponent& OwnerComp, const uint8* RawData) const override;
+	virtual bool TestBasicOperation(const UBlackboardComponent& OwnerComp, const uint8* MemoryBlock, EBasicKeyOperation::Type Op) const override;
+};
diff --git a/Engine/Source/Runtime/AIModule/Classes/BehaviorTree/ValueOrBBKey.h b/Engine/Source/Runtime/AIModule/Classes/BehaviorTree/ValueOrBBKey.h
--- a/Engine/Source/Runtime/AIModule/Classes/BehaviorTree/ValueOrBBKey.h
+++ b/Engine/Source/Runtime/AIModule/Classes/BehaviorTree/ValueOrBBKey.h
@@ -7,6 +7,8 @@
 #include "BehaviorTree/BlackboardComponent.h"
 #include "GameplayTagContainer.h"
 #include "Misc/Guid.h"
+#include "StructUtils/InstancedStruct.h"
+#include "StructUtils/StructView.h"
 
 #include "ValueOrBBKey.generated.h"
 
@@ -14,6 +16,7 @@ class UBlackboardKeyType;
 class FValueOrBBKeyDetails;
 
 struct FPropertyTag;
+struct FConstStructView;
 
 // Drop in replacement for Property in bt nodes that allows easy binding to blackboard key. Replace the old property with the corresponding type and call the GetValue() function to retrieve the value.
 struct FValueOrBBKey_Bool;
@@ -25,6 +28,7 @@ struct FValueOrBBKey_Name;
 struct FValueOrBBKey_String;
 struct FValueOrBBKey_Object;
 struct FValueOrBBKey_Rotator;
+struct FValueOrBBKey_Struct;
 struct FValueOrBBKey_Vector;
 
 namespace FBlackboard
@@ -65,6 +69,12 @@ namespace FBlackboard
 		}
 		return DefaultValue;
 	}
+
+	FConstStructView TryGetBlackboardKeyStruct(const UBlackboardComponent& Blackboard, const FName& Name, FBlackboard::FKey& InOutCachedKey, const UScriptStruct* TargetStruct);
+
+	FConstStructView GetStructValue(const UBlackboardComponent& Blackboard, const FName& Name, FBlackboard::FKey& InOutCachedKey, const FConstStructView& DefaultValue);
+
+	FConstStructView GetStructValue(const UBehaviorTreeComponent& BehaviorComp, const FName& Name, FBlackboard::FKey& InOutCachedKey, const FConstStructView& DefaultValue);
 } // namespace FBlackboard
 
 // Base struct to simplify edition in the editor, shouldn't be used elsewhere
@@ -436,7 +446,53 @@ protected:
 	FVector DefaultValue;
 };
 
-template<>
+USTRUCT(BlueprintType)
+struct FValueOrBBKey_Struct : public FValueOrBlackboardKeyBase
+{
+	GENERATED_BODY()
+
+public:
+	template <typename T>
+	FValueOrBBKey_Struct(const T& Value)
+	{
+		DefaultValue.InitializeAs<T>(Value);
+	}
+	FValueOrBBKey_Struct() = default;
+
+#if WITH_EDITOR
+	AIMODULE_API virtual bool IsCompatibleType(const UBlackboardKeyType* KeyType) const;
+#endif // WITH_EDITOR
+
+	AIMODULE_API FString ToString() const;
+
+	template <typename T>
+	const T& GetValue(const UBehaviorTreeComponent& BehaviorComp) const
+	{
+		check(DefaultValue.GetScriptStruct() == TBaseStructure<T>::Get());
+		return GetValue(BehaviorComp).Get<const T>();
+	}
+
+	bool SerializeFromMismatchedTag(const FPropertyTag& Tag, FStructuredArchive::FSlot Slot);
+
+	AIMODULE_API FConstStructView GetValue(const UBehaviorTreeComponent& BehaviorComp) const;
+	AIMODULE_API FConstStructView GetValue(const UBehaviorTreeComponent* BehaviorComp) const;
+	AIMODULE_API FConstStructView GetValue(const UBlackboardComponent& Blackboard) const;
+	AIMODULE_API FConstStructView GetValue(const UBlackboardComponent* Blackboard) const;
+
+protected:
+	friend class FValueOrBBKeyDetails_Struct;
+
+	UPROPERTY(EditAnywhere, Category = Value)
+	FInstancedStruct DefaultValue;
+
+#if WITH_EDITORONLY_DATA
+	/** Exist to give a way for the details to know if the DefaultValue type can be edited. */
+	UPROPERTY(EditDefaultsOnly, Category = Value)
+	bool bCanEditDefaultValueType = false;
+#endif // WITH_EDITORONLY_DATA
+};
+
+template <>
 struct TStructOpsTypeTraits<FValueOrBBKey_Bool> : public TStructOpsTypeTraitsBase2<FValueOrBBKey_Bool>
 {
 	enum
@@ -445,7 +501,7 @@ struct TStructOpsTypeTraits<FValueOrBBKey_Bool> : public TStructOpsTypeTraitsBas
 	};
 };
 
-template<>
+template <>
 struct TStructOpsTypeTraits<FValueOrBBKey_Class> : public TStructOpsTypeTraitsBase2<FValueOrBBKey_Class>
 {
 	enum
@@ -454,7 +510,7 @@ struct TStructOpsTypeTraits<FValueOrBBKey_Class> : public TStructOpsTypeTraitsBa
 	};
 };
 
-template<>
+template <>
 struct TStructOpsTypeTraits<FValueOrBBKey_Enum> : public TStructOpsTypeTraitsBase2<FValueOrBBKey_Enum>
 {
 	enum
@@ -463,7 +519,7 @@ struct TStructOpsTypeTraits<FValueOrBBKey_Enum> : public TStructOpsTypeTraitsBas
 	};
 };
 
-template<>
+template <>
 struct TStructOpsTypeTraits<FValueOrBBKey_Float> : public TStructOpsTypeTraitsBase2<FValueOrBBKey_Float>
 {
 	enum
@@ -472,7 +528,7 @@ struct TStructOpsTypeTraits<FValueOrBBKey_Float> : public TStructOpsTypeTraitsBa
 	};
 };
 
-template<>
+template <>
 struct TStructOpsTypeTraits<FValueOrBBKey_Int32> : public TStructOpsTypeTraitsBase2<FValueOrBBKey_Int32>
 {
 	enum
@@ -481,7 +537,7 @@ struct TStructOpsTypeTraits<FValueOrBBKey_Int32> : public TStructOpsTypeTraitsBa
 	};
 };
 
-template<>
+template <>
 struct TStructOpsTypeTraits<FValueOrBBKey_Name> : public TStructOpsTypeTraitsBase2<FValueOrBBKey_Name>
 {
 	enum
@@ -490,7 +546,7 @@ struct TStructOpsTypeTraits<FValueOrBBKey_Name> : public TStructOpsTypeTraitsBas
 	};
 };
 
-template<>
+template <>
 struct TStructOpsTypeTraits<FValueOrBBKey_String> : public TStructOpsTypeTraitsBase2<FValueOrBBKey_String>
 {
 	enum
@@ -499,7 +555,7 @@ struct TStructOpsTypeTraits<FValueOrBBKey_String> : public TStructOpsTypeTraitsB
 	};
 };
 
-template<>
+template <>
 struct TStructOpsTypeTraits<FValueOrBBKey_Object> : public TStructOpsTypeTraitsBase2<FValueOrBBKey_Object>
 {
 	enum
@@ -508,7 +564,7 @@ struct TStructOpsTypeTraits<FValueOrBBKey_Object> : public TStructOpsTypeTraitsB
 	};
 };
 
-template<>
+template <>
 struct TStructOpsTypeTraits<FValueOrBBKey_Rotator> : public TStructOpsTypeTraitsBase2<FValueOrBBKey_Rotator>
 {
 	enum
@@ -517,7 +573,16 @@ struct TStructOpsTypeTraits<FValueOrBBKey_Rotator> : public TStructOpsTypeTraits
 	};
 };
 
-template<>
+template <>
+struct TStructOpsTypeTraits<FValueOrBBKey_Struct> : public TStructOpsTypeTraitsBase2<FValueOrBBKey_Struct>
+{
+	enum
+	{
+		WithStructuredSerializeFromMismatchedTag = true,
+	};
+};
+
+template <>
 struct TStructOpsTypeTraits<FValueOrBBKey_Vector> : public TStructOpsTypeTraitsBase2<FValueOrBBKey_Vector>
 {
 	enum
diff --git a/Engine/Source/Runtime/AIModule/Private/BehaviorTree/Blackboard/BlackboardKeyType.cpp b/Engine/Source/Runtime/AIModule/Private/BehaviorTree/Blackboard/BlackboardKeyType.cpp
--- a/Engine/Source/Runtime/AIModule/Private/BehaviorTree/Blackboard/BlackboardKeyType.cpp
+++ b/Engine/Source/Runtime/AIModule/Private/BehaviorTree/Blackboard/BlackboardKeyType.cpp
@@ -64,7 +64,7 @@ void UBlackboardKeyType::WrappedFree(UBlackboardComponent& OwnerComp, uint8* Mem
 	{
 		UBlackboardKeyType* InstancedKey = GetKeyInstance(OwnerComp, MemoryBlock);
 		uint8* InstanceMemoryBlock = MemoryBlock + sizeof(FBlackboardInstancedKeyMemory);
-		InstancedKey->FreeMemory(OwnerComp, InstanceMemoryBlock);
+		return InstancedKey->FreeMemory(OwnerComp, InstanceMemoryBlock);
 	}
 
 	return FreeMemory(OwnerComp, MemoryBlock);
diff --git a/Engine/Source/Runtime/AIModule/Private/BehaviorTree/Blackboard/BlackboardKeyType_Struct.cpp b/Engine/Source/Runtime/AIModule/Private/BehaviorTree/Blackboard/BlackboardKeyType_Struct.cpp
new file mode 100644
--- /dev/null
+++ b/Engine/Source/Runtime/AIModule/Private/BehaviorTree/Blackboard/BlackboardKeyType_Struct.cpp
@@ -0,0 +1,211 @@
+// Copyright Epic Games, Inc. All Rights Reserved.
+
+#include "BehaviorTree/Blackboard/BlackboardKeyType_Struct.h"
+
+#include "StructUtils/StructView.h"
+
+const UBlackboardKeyType_Struct::FDataType UBlackboardKeyType_Struct::InvalidValue;
+
+UBlackboardKeyType_Struct::UBlackboardKeyType_Struct(const FObjectInitializer& ObjectInitializer)
+	: Super(ObjectInitializer)
+{
+	bCreateKeyInstance = false;
+	SupportedOp = EBlackboardKeyOperation::Basic;
+}
+
+FConstStructView UBlackboardKeyType_Struct::GetValue(UBlackboardKeyType_Struct* KeyOb, const uint8* RawData)
+{
+	if (KeyOb->bIsInstanced)
+	{
+		return FConstStructView(KeyOb->Value);
+	}
+	else
+	{
+		return FConstStructView(KeyOb->DefaultValue.GetScriptStruct(), RawData);
+	}
+}
+
+bool UBlackboardKeyType_Struct::SetValue(UBlackboardKeyType_Struct* KeyOb, uint8* RawData, FConstStructView NewValue)
+{
+	if (KeyOb->bIsInstanced)
+	{
+		if (NewValue.GetScriptStruct() == KeyOb->DefaultValue.GetScriptStruct())
+		{
+			KeyOb->Value = NewValue;
+			return true;
+		}
+	}
+	else
+	{
+		if (const UScriptStruct* ScriptStruct = KeyOb->DefaultValue.GetScriptStruct())
+		{
+			if (NewValue.GetScriptStruct() == ScriptStruct)
+			{
+				ScriptStruct->CopyScriptStruct(RawData, NewValue.GetMemory());
+				return true;
+			}
+		}
+	}
+	return false;
+}
+
+EBlackboardCompare::Type UBlackboardKeyType_Struct::CompareValues(const UBlackboardComponent& OwnerComp, const uint8* MemoryBlock,
+	const UBlackboardKeyType* OtherKeyOb, const uint8* OtherMemoryBlock) const
+{
+	if (const UBlackboardKeyType_Struct* OtherKey = Cast<UBlackboardKeyType_Struct>(OtherKeyOb))
+	{
+		if (DefaultValue.GetScriptStruct() == OtherKey->DefaultValue.GetScriptStruct())
+		{
+			if (bIsInstanced)
+			{
+				return Value == OtherKey->Value ? EBlackboardCompare::Equal : EBlackboardCompare::NotEqual;
+			}
+			else
+			{
+				if (const UScriptStruct* ScriptStruct = DefaultValue.GetScriptStruct())
+				{
+					return ScriptStruct->CompareScriptStruct(MemoryBlock, OtherMemoryBlock, PPF_None) ? EBlackboardCompare::Equal : EBlackboardCompare::NotEqual;
+				}
+			}
+		}
+	}
+	return EBlackboardCompare::NotEqual;
+}
+
+void UBlackboardKeyType_Struct::CopyValues(UBlackboardComponent& OwnerComp, uint8* MemoryBlock, const UBlackboardKeyType* SourceKeyOb, const uint8* SourceBlock)
+{
+	if (const UBlackboardKeyType_Struct* SourceKey = Cast<UBlackboardKeyType_Struct>(SourceKeyOb))
+	{
+		if (DefaultValue.GetScriptStruct() == SourceKey->DefaultValue.GetScriptStruct())
+		{
+			if (bIsInstanced)
+			{
+				Value = SourceKey->Value;
+			}
+			else
+			{
+				if (const UScriptStruct* ScriptStruct = DefaultValue.GetScriptStruct())
+				{
+					ScriptStruct->CopyScriptStruct(MemoryBlock, SourceBlock);
+				}
+			}
+		}
+	}
+}
+
+void UBlackboardKeyType_Struct::UpdateNeedsInstance()
+{
+	if (const UScriptStruct* ScriptStruct = DefaultValue.GetScriptStruct())
+	{
+		for (const FProperty* PropertyLink = ScriptStruct->PropertyLink; PropertyLink; PropertyLink = PropertyLink->PropertyLinkNext)
+		{
+			// Struct with strong object reference needs to have their key instantiated to be taken into account by the gc.
+			// Otherwise struct can safely lives inside the blackboard buffer
+			TArray<const FStructProperty*> EncounteredStructProps;
+			if (PropertyLink->ContainsObjectReference(EncounteredStructProps, EPropertyObjectReferenceType::Strong))
+			{
+				bCreateKeyInstance = true;
+				ValueSize = 0;
+				return;
+			}
+		}
+
+		bCreateKeyInstance = false;
+		ValueSize = ScriptStruct->GetStructureSize();
+	}
+	else
+	{
+		bCreateKeyInstance = false;
+		ValueSize = 0;
+	}
+}
+
+FString UBlackboardKeyType_Struct::DescribeValue(const UBlackboardComponent& OwnerComp, const uint8* RawData) const
+{
+	const UScriptStruct* ScriptStruct = nullptr;
+	const uint8* ValuePtr = nullptr;
+	if (bIsInstanced)
+	{
+		ScriptStruct = Value.GetScriptStruct();
+		ValuePtr = Value.GetMemory();
+	}
+	else
+	{
+		ScriptStruct = DefaultValue.GetScriptStruct();
+		ValuePtr = RawData;
+	}
+
+	if (ScriptStruct && ValuePtr)
+	{
+		FString ExportedValue;
+		ScriptStruct->ExportText(ExportedValue, ValuePtr, nullptr, const_cast<UBlackboardKeyType_Struct*>(this), PPF_None, nullptr);
+		return ExportedValue;
+	}
+	return DescribeSelf();
+}
+
+FString UBlackboardKeyType_Struct::DescribeSelf() const
+{
+	return *GetNameSafe(DefaultValue.GetScriptStruct());
+}
+
+bool UBlackboardKeyType_Struct::IsAllowedByFilter(UBlackboardKeyType* FilterOb) const
+{
+	const UBlackboardKeyType_Struct* FilterClass = Cast<UBlackboardKeyType_Struct>(FilterOb);
+	return FilterClass && (FilterClass->DefaultValue.GetScriptStruct() == DefaultValue.GetScriptStruct());
+}
+
+void UBlackboardKeyType_Struct::InitializeMemory(UBlackboardComponent& OwnerComp, uint8* MemoryBlock)
+{
+	SetValue(this, MemoryBlock, DefaultValue);
+}
+
+void UBlackboardKeyType_Struct::FreeMemory(UBlackboardComponent& OwnerComp, uint8* MemoryBlock)
+{
+	// Let normal object destruction flow take care of the instanced key.
+	if (!bIsInstanced)
+	{
+		if (const UScriptStruct* ScriptStruct = DefaultValue.GetScriptStruct())
+		{
+			ScriptStruct->DestroyStruct(MemoryBlock);
+		}
+	}
+}
+
+void UBlackboardKeyType_Struct::Clear(UBlackboardComponent& OwnerComp, uint8* MemoryBlock)
+{
+	InitializeMemory(OwnerComp, MemoryBlock);
+}
+
+void UBlackboardKeyType_Struct::PostLoad()
+{
+	Super::PostLoad();
+
+	UpdateNeedsInstance();
+}
+
+#if WITH_EDITOR
+void UBlackboardKeyType_Struct::PostEditChangeProperty(struct FPropertyChangedEvent& PropertyChangedEvent)
+{
+	Super::PostEditChangeProperty(PropertyChangedEvent);
+
+	UpdateNeedsInstance();
+}
+#endif // WITH_EDITOR
+
+bool UBlackboardKeyType_Struct::TestBasicOperation(const UBlackboardComponent& OwnerComp, const uint8* MemoryBlock, EBasicKeyOperation::Type Op) const
+{
+	if (bIsInstanced)
+	{
+		return (Op == EBasicKeyOperation::Set) ? Value != DefaultValue : Value == DefaultValue;
+	}
+	else
+	{
+		if (const UScriptStruct* ScriptStruct = DefaultValue.GetScriptStruct())
+		{
+			const bool Identical = ScriptStruct->CompareScriptStruct(MemoryBlock, DefaultValue.GetMemory(), PPF_None);
+			return Op == EBasicKeyOperation::Set ? !Identical : Identical;
+		}
+	}
+	return Op == EBasicKeyOperation::NotSet;
+}
\ No newline at end of file
diff --git a/Engine/Source/Runtime/AIModule/Private/BehaviorTree/ValueOrBBKey.cpp b/Engine/Source/Runtime/AIModule/Private/BehaviorTree/ValueOrBBKey.cpp
--- a/Engine/Source/Runtime/AIModule/Private/BehaviorTree/ValueOrBBKey.cpp
+++ b/Engine/Source/Runtime/AIModule/Private/BehaviorTree/ValueOrBBKey.cpp
@@ -11,8 +11,58 @@
 #include "BehaviorTree/Blackboard/BlackboardKeyType_Object.h"
 #include "BehaviorTree/Blackboard/BlackboardKeyType_Rotator.h"
 #include "BehaviorTree/Blackboard/BlackboardKeyType_String.h"
+#include "BehaviorTree/Blackboard/BlackboardKeyType_Struct.h"
 #include "BehaviorTree/Blackboard/BlackboardKeyType_Vector.h"
 
+namespace FBlackboard
+{
+	FConstStructView TryGetBlackboardKeyStruct(const UBlackboardComponent& Blackboard, const FName& Name, FBlackboard::FKey& InOutCachedKey, const UScriptStruct* TargetStruct)
+	{
+		if (InOutCachedKey == FBlackboard::InvalidKey)
+		{
+			InOutCachedKey = Blackboard.GetKeyID(Name);
+		}
+
+		if (const UBlackboardData* BlackBoarData = Blackboard.GetBlackboardAsset())
+		{
+			if (const FBlackboardEntry* BBEntry = BlackBoarData->GetKey(InOutCachedKey))
+			{
+				if (UBlackboardKeyType_Struct* StructKeyType = Cast<UBlackboardKeyType_Struct>(BBEntry->KeyType))
+				{
+					if (StructKeyType->DefaultValue.GetScriptStruct() == TargetStruct)
+					{
+						return Blackboard.GetValue<UBlackboardKeyType_Struct>(InOutCachedKey);
+					}
+				}
+			}
+		}
+
+		return FConstStructView{};
+	}
+
+	FConstStructView GetStructValue(const UBlackboardComponent& Blackboard, const FName& Name, FBlackboard::FKey& InOutCachedKey, const FConstStructView& DefaultValue)
+	{
+		if (!Name.IsNone())
+		{
+			FConstStructView KeyValue = FBlackboard::TryGetBlackboardKeyStruct(Blackboard, Name, InOutCachedKey, DefaultValue.GetScriptStruct());
+			if (KeyValue.IsValid())
+			{
+				return KeyValue;
+			}
+		}
+		return DefaultValue;
+	}
+
+	FConstStructView GetStructValue(const UBehaviorTreeComponent& BehaviorComp, const FName& Name, FBlackboard::FKey& InOutCachedKey, const FConstStructView& DefaultValue)
+	{
+		if (const UBlackboardComponent* Blackboard = BehaviorComp.GetBlackboardComponent())
+		{
+			return GetStructValue(*Blackboard, Name, InOutCachedKey, DefaultValue);
+		}
+		return DefaultValue;
+	}
+} // namespace FBlackboard
+
 FBlackboard::FKey FValueOrBlackboardKeyBase::GetKeyId(const UBehaviorTreeComponent& OwnerComp) const
 {
 	if (KeyId == FBlackboard::InvalidKey)
@@ -107,6 +157,18 @@ FString FValueOrBBKey_Rotator::ToString() const
 	}
 }
 
+FString FValueOrBBKey_Struct::ToString() const
+{
+	if (!Key.IsNone())
+	{
+		return ToStringKeyName();
+	}
+	else
+	{
+		return GetNameSafe(DefaultValue.GetScriptStruct());
+	}
+}
+
 FString FValueOrBBKey_Vector::ToString() const
 {
 	if (!Key.IsNone())
@@ -299,6 +361,26 @@ FRotator FValueOrBBKey_Rotator::GetValue(const UBehaviorTreeComponent* BehaviorC
 	return BehaviorComp ? GetValue(*BehaviorComp) : DefaultValue;
 }
 
+FConstStructView FValueOrBBKey_Struct::GetValue(const UBlackboardComponent& Blackboard) const
+{
+	return FBlackboard::GetStructValue(Blackboard, Key, KeyId, FConstStructView(DefaultValue));
+}
+
+FConstStructView FValueOrBBKey_Struct::GetValue(const UBlackboardComponent* Blackboard) const
+{
+	return Blackboard ? GetValue(*Blackboard) : DefaultValue;
+}
+
+FConstStructView FValueOrBBKey_Struct::GetValue(const UBehaviorTreeComponent& BehaviorComp) const
+{
+	return FBlackboard::GetStructValue(BehaviorComp, Key, KeyId, FConstStructView(DefaultValue));
+}
+
+FConstStructView FValueOrBBKey_Struct::GetValue(const UBehaviorTreeComponent* BehaviorComp) const
+{
+	return BehaviorComp ? GetValue(*BehaviorComp) : DefaultValue;
+}
+
 FVector FValueOrBBKey_Vector::GetValue(const UBlackboardComponent& Blackboard) const
 {
 	return FBlackboard::GetValue<UBlackboardKeyType_Vector>(Blackboard, Key, KeyId, DefaultValue);
@@ -414,6 +496,25 @@ bool FValueOrBBKey_Rotator::SerializeFromMismatchedTag(const FPropertyTag& Tag,
 	return false;
 }
 
+bool FValueOrBBKey_Struct::SerializeFromMismatchedTag(const FPropertyTag& Tag, FStructuredArchive::FSlot Slot)
+{
+	if (Tag.Type == NAME_StructProperty && Tag.GetType().GetParameterCount() > 0)
+	{
+		if (UScriptStruct* Struct = FindFirstObject<UScriptStruct>(*Tag.GetType().GetParameterName(0).ToString(), EFindFirstObjectOptions::NativeFirst))
+		{
+			if (Struct == DefaultValue.GetScriptStruct() || DefaultValue.GetScriptStruct() == nullptr)
+			{
+				TUniquePtr<uint8[]> SerializedStruct = MakeUnique<uint8[]>(Struct->GetStructureSize());
+				Struct->InitializeStruct(SerializedStruct.Get());
+				Struct->SerializeItem(Slot, SerializedStruct.Get(), nullptr);
+				DefaultValue.InitializeAs(Struct, SerializedStruct.Get());
+				return true;
+			}
+		}
+	}
+	return false;
+}
+
 bool FValueOrBBKey_Vector::SerializeFromMismatchedTag(const FPropertyTag& Tag, FStructuredArchive::FSlot Slot)
 {
 	if (Tag.Type == NAME_VectorProperty)
@@ -491,6 +592,15 @@ bool FValueOrBBKey_Rotator::IsCompatibleType(const UBlackboardKeyType* KeyType)
 	return KeyType && KeyType->GetClass() == UBlackboardKeyType_Rotator::StaticClass();
 }
 
+bool FValueOrBBKey_Struct::IsCompatibleType(const UBlackboardKeyType* KeyType) const
+{
+	if (const UBlackboardKeyType_Struct* StructKey = Cast<UBlackboardKeyType_Struct>(KeyType))
+	{
+		return StructKey->DefaultValue.GetScriptStruct() == DefaultValue.GetScriptStruct();
+	}
+	return false;
+}
+
 bool FValueOrBBKey_Vector::IsCompatibleType(const UBlackboardKeyType* KeyType) const
 {
 	return KeyType && KeyType->GetClass() == UBlackboardKeyType_Vector::StaticClass();
diff --git a/Engine/Source/Runtime/AIModule/Private/BehaviorTree/ValueOrBBKeyBlueprintUtility.cpp b/Engine/Source/Runtime/AIModule/Private/BehaviorTree/ValueOrBBKeyBlueprintUtility.cpp
--- a/Engine/Source/Runtime/AIModule/Private/BehaviorTree/ValueOrBBKeyBlueprintUtility.cpp
+++ b/Engine/Source/Runtime/AIModule/Private/BehaviorTree/ValueOrBBKeyBlueprintUtility.cpp
@@ -3,6 +3,7 @@
 #include "BehaviorTree/ValueOrBBKeyBlueprintUtility.h"
 
 #include "BehaviorTree/ValueOrBBKey.h"
+#include "StructUtils/StructView.h"
 
 bool UValueOrBBKeyBlueprintUtility::GetBool(const FValueOrBBKey_Bool& Value, const UBehaviorTreeComponent* BehaviorTreeComp)
 {
@@ -53,3 +54,8 @@ FVector UValueOrBBKeyBlueprintUtility::GetVector(const FValueOrBBKey_Vector& Val
 {
 	return Value.GetValue(BehaviorTreeComp);
 }
+
+FInstancedStruct UValueOrBBKeyBlueprintUtility::GetStruct(const FValueOrBBKey_Struct& Value, const UBehaviorTreeComponent* BehaviorTreeComp)
+{
+	return FInstancedStruct(Value.GetValue(BehaviorTreeComp));
+}
diff --git a/Engine/Source/Runtime/AIModule/Private/BehaviorTree/ValueOrBBKeyBlueprintUtility.h b/Engine/Source/Runtime/AIModule/Private/BehaviorTree/ValueOrBBKeyBlueprintUtility.h
--- a/Engine/Source/Runtime/AIModule/Private/BehaviorTree/ValueOrBBKeyBlueprintUtility.h
+++ b/Engine/Source/Runtime/AIModule/Private/BehaviorTree/ValueOrBBKeyBlueprintUtility.h
@@ -51,4 +51,7 @@ public:
 
 	UFUNCTION(BlueprintPure, Category = Blackboard)
 	static FVector GetVector(const FValueOrBBKey_Vector& Value, const UBehaviorTreeComponent* BehaviorTreeComp);
+
+	UFUNCTION(BlueprintPure, Category = Blackboard)
+	static FInstancedStruct GetStruct(const FValueOrBBKey_Struct& Value, const UBehaviorTreeComponent* BehaviorTreeComp);
 };
\ No newline at end of file
