Missing file from cl 34530408
#rnx

diff --git a/Engine/Source/Runtime/AIModule/Private/BehaviorTree/BehaviorTreeTypes.cpp b/Engine/Source/Runtime/AIModule/Private/BehaviorTree/BehaviorTreeTypes.cpp
--- a/Engine/Source/Runtime/AIModule/Private/BehaviorTree/BehaviorTreeTypes.cpp
+++ b/Engine/Source/Runtime/AIModule/Private/BehaviorTree/BehaviorTreeTypes.cpp
@@ -18,6 +18,7 @@
 #include "BehaviorTree/Blackboard/BlackboardKeyType_Int.h"
 #include "BehaviorTree/Blackboard/BlackboardKeyType_Name.h"
 #include "BehaviorTree/Blackboard/BlackboardKeyType_String.h"
+#include "BehaviorTree/Blackboard/BlackboardKeyType_Struct.h"
 #include "BehaviorTree/BTTaskNode.h"
 #include "BehaviorTree/BTCompositeNode.h"
 
@@ -664,6 +665,14 @@ void FBlackboardKeySelector::AddNameFilter(UObject* Owner, FName PropertyName)
 	AllowedTypes.Add(NewObject<UBlackboardKeyType_Name>(Owner, *FilterName));
 }
 
+void FBlackboardKeySelector::AddStructFilter(UObject* Owner, FName PropertyName, const UScriptStruct* AllowedStruct)
+{
+	const FString FilterName = PropertyName.ToString() + TEXT("_Struct_") + GetNameSafe(AllowedStruct);
+	UBlackboardKeyType_Struct* FilterOb = NewObject<UBlackboardKeyType_Struct>(Owner, *FilterName);
+	FilterOb->DefaultValue.InitializeAs(AllowedStruct);
+	AllowedTypes.Add(FilterOb);
+}
+
 //----------------------------------------------------------------------//
 // UBehaviorTreeTypes
 //----------------------------------------------------------------------//
