PR #12042: Deprecate old perception (pawn sensing component)

diff --git a/Engine/Source/Runtime/AIModule/Classes/Perception/PawnSensingComponent.h b/Engine/Source/Runtime/AIModule/Classes/Perception/PawnSensingComponent.h
--- a/Engine/Source/Runtime/AIModule/Classes/Perception/PawnSensingComponent.h
+++ b/Engine/Source/Runtime/AIModule/Classes/Perception/PawnSensingComponent.h
@@ -17,6 +17,7 @@ class UPawnNoiseEmitterComponent;
  * SensingComponent encapsulates sensory (ie sight and hearing) settings and functionality for an Actor,
  * allowing the actor to see/hear Pawns in the world. It does nothing on network clients.
  */
+class UE_DEPRECATED_FORGAME(5.5, "Pawn sensing is deprecated. AI Perception should be used instead.") UPawnSensingComponent;
 UCLASS(ClassGroup=AI, HideCategories=(Activation, "Components|Activation", Collision), meta=(BlueprintSpawnableComponent), MinimalAPI)
 class UPawnSensingComponent : public UActorComponent
 {
