Add Points in cylinder EQS generator based on Donut one.

diff --git a/Engine/Source/Runtime/AIModule/Classes/EnvironmentQuery/Generators/EnvQueryGenerator_Donut.h b/Engine/Source/Runtime/AIModule/Classes/EnvironmentQuery/Generators/EnvQueryGenerator_Donut.h
--- a/Engine/Source/Runtime/AIModule/Classes/EnvironmentQuery/Generators/EnvQueryGenerator_Donut.h
+++ b/Engine/Source/Runtime/AIModule/Classes/EnvironmentQuery/Generators/EnvQueryGenerator_Donut.h
@@ -63,4 +63,5 @@ class UEnvQueryGenerator_Donut : public UEnvQueryGenerator_ProjectedPoints
 protected:
 	AIMODULE_API FVector::FReal GetArcBisectorAngle(FEnvQueryInstance& QueryInstance) const;
 	AIMODULE_API bool IsAngleAllowed(FVector::FReal TestAngleRad, FVector::FReal BisectAngleDeg, FVector::FReal AngleRangeDeg, bool bConstrainAngle) const;
+	AIMODULE_API bool GenerateDonutNavPoints(const TArray<FVector>& CenterPoints, TArray<FNavLocation>& Points, FEnvQueryInstance& QueryInstance) const;
 };
diff --git a/Engine/Source/Runtime/AIModule/Private/EnvironmentQuery/Generators/EnvQueryGenerator_Donut.cpp b/Engine/Source/Runtime/AIModule/Private/EnvironmentQuery/Generators/EnvQueryGenerator_Donut.cpp
--- a/Engine/Source/Runtime/AIModule/Private/EnvironmentQuery/Generators/EnvQueryGenerator_Donut.cpp
+++ b/Engine/Source/Runtime/AIModule/Private/EnvironmentQuery/Generators/EnvQueryGenerator_Donut.cpp
@@ -80,6 +80,15 @@ void UEnvQueryGenerator_Donut::GenerateItems(FEnvQueryInstance& QueryInstance) c
 		return;
 	}
 
+	TArray<FNavLocation> Points;
+	GenerateDonutNavPoints(CenterPoints, Points, QueryInstance);
+
+	ProjectAndFilterNavPoints(Points, QueryInstance);
+	StoreNavPoints(Points, QueryInstance);
+}
+
+bool UEnvQueryGenerator_Donut::GenerateDonutNavPoints(const TArray<FVector>& CenterPoints, TArray<FNavLocation>& Points, FEnvQueryInstance& QueryInstance) const
+{
 	UObject* BindOwner = QueryInstance.Owner.Get();
 	InnerRadius.BindData(BindOwner, QueryInstance.QueryID);
 	OuterRadius.BindData(BindOwner, QueryInstance.QueryID);
@@ -97,7 +106,7 @@ void UEnvQueryGenerator_Donut::GenerateItems(FEnvQueryInstance& QueryInstance) c
 		(InnerRadiusValue > OuterRadiusValue) ||
 		(NumRings < 1) || (NumPoints < 1))
 	{
-		return;
+		return false;
 	}
 
 	const FVector::FReal ArcBisectDeg = GetArcBisectorAngle(QueryInstance);
@@ -107,7 +116,6 @@ void UEnvQueryGenerator_Donut::GenerateItems(FEnvQueryInstance& QueryInstance) c
 	const FVector::FReal AngleDelta = 2. * UE_DOUBLE_PI / NumPoints;
 	FVector::FReal SectionAngle = FMath::DegreesToRadians(ArcBisectDeg);
 
-	TArray<FNavLocation> Points;
 	Points.Reserve(NumPoints * NumRings);
 
 	if (!bUseSpiralPattern)
@@ -162,8 +170,7 @@ void UEnvQueryGenerator_Donut::GenerateItems(FEnvQueryInstance& QueryInstance) c
 		}
 	}
 
-	ProjectAndFilterNavPoints(Points, QueryInstance);
-	StoreNavPoints(Points, QueryInstance);
+	return Points.Num() > 0;
 }
 
 FText UEnvQueryGenerator_Donut::GetDescriptionTitle() const
