Clarify UAITask_MoveTo params

diff --git a/Engine/Source/Runtime/AIModule/Classes/Tasks/AITask_MoveTo.h b/Engine/Source/Runtime/AIModule/Classes/Tasks/AITask_MoveTo.h
--- a/Engine/Source/Runtime/AIModule/Classes/Tasks/AITask_MoveTo.h
+++ b/Engine/Source/Runtime/AIModule/Classes/Tasks/AITask_MoveTo.h
@@ -31,6 +31,11 @@ public:
 	bool WasMoveSuccessful() const { return MoveResult == EPathFollowingResult::Success; }
 	bool WasMovePartial() const { return Path.IsValid() && Path->IsPartial(); }
 
+	/**
+	 * Move to a location or an actor
+	 * @param ProjectGoalOnNavigation Try to move the goal to the navigation surface before requesting the move, fails if it can't.
+	 * @param RequireNavigableEndLocation Set to No to allow pursuing the request even if no navigation surface is found at the goal location. 
+	 */
 	UFUNCTION(BlueprintCallable, Category = "AI|Tasks", meta = (AdvancedDisplay = "AcceptanceRadius,StopOnOverlap,AcceptPartialPath,bUsePathfinding,bUseContinuousGoalTracking,ProjectGoalOnNavigation,RequireNavigableEndLocation", DefaultToSelf = "Controller", BlueprintInternalUseOnly = "TRUE", DisplayName = "Move To Location or Actor"))
 	static AIMODULE_API UAITask_MoveTo* AIMoveTo(AAIController* Controller, FVector GoalLocation, AActor* GoalActor = nullptr,
 		float AcceptanceRadius = -1.f, EAIOptionFlag::Type StopOnOverlap = EAIOptionFlag::Default, EAIOptionFlag::Type AcceptPartialPath = EAIOptionFlag::Default,
