PR #12186: Fix Crash in UAISense_Sight::Update

diff --git a/Engine/Source/Runtime/AIModule/Private/Perception/AISense_Sight.cpp b/Engine/Source/Runtime/AIModule/Private/Perception/AISense_Sight.cpp
--- a/Engine/Source/Runtime/AIModule/Private/Perception/AISense_Sight.cpp
+++ b/Engine/Source/Runtime/AIModule/Private/Perception/AISense_Sight.cpp
@@ -748,36 +748,33 @@ void UAISense_Sight::UnregisterSource(AActor& SourceActor)
 	{
 		AActor* TargetActor = AsTarget.Target.Get();
 
-		if (TargetActor)
+		// notify all interested observers that this source is no longer
+		// visible		
+		AIPerception::FListenerMap& ListenersMap = *GetListeners();
+		auto RemoveQuery = [this,&ListenersMap,&AsTargetId,&TargetActor](TArray<FAISightQuery>& SightQueries, const int32 QueryIndex)->EReverseForEachResult
 		{
-			// notify all interested observers that this source is no longer
-			// visible		
-			AIPerception::FListenerMap& ListenersMap = *GetListeners();
-			auto RemoveQuery = [this,&ListenersMap,&AsTargetId,&TargetActor](TArray<FAISightQuery>& SightQueries, const int32 QueryIndex)->EReverseForEachResult
+			FAISightQuery* SightQuery = &SightQueries[QueryIndex];
+			if (SightQuery->TargetId == AsTargetId)
 			{
-				FAISightQuery* SightQuery = &SightQueries[QueryIndex];
-				if (SightQuery->TargetId == AsTargetId)
+				if (SightQuery->GetLastResult() && TargetActor)
 				{
-					if (SightQuery->GetLastResult())
-					{
-						FPerceptionListener& Listener = ListenersMap[SightQuery->ObserverId];
-						ensure(Listener.Listener.IsValid());
-
-						Listener.RegisterStimulus(TargetActor, FAIStimulus(*this, 0.f, SightQuery->LastSeenLocation, Listener.CachedLocation, FAIStimulus::SensingFailed));
-					}
+					FPerceptionListener& Listener = ListenersMap[SightQuery->ObserverId];
+					ensure(Listener.Listener.IsValid());
 
-					SightQueries.RemoveAtSwap(QueryIndex, EAllowShrinking::No);
-					return EReverseForEachResult::Modified;
+					Listener.RegisterStimulus(TargetActor, FAIStimulus(*this, 0.f, SightQuery->LastSeenLocation, Listener.CachedLocation, FAIStimulus::SensingFailed));
 				}
-				return EReverseForEachResult::UnTouched;
-			};
-			ReverseForEach(SightQueriesInRange, RemoveQuery);
-			if (ReverseForEach(SightQueriesOutOfRange, RemoveQuery) == EReverseForEachResult::Modified)
-			{
-				bSightQueriesOutOfRangeDirty = true;
+
+				SightQueries.RemoveAtSwap(QueryIndex, EAllowShrinking::No);
+				return EReverseForEachResult::Modified;
 			}
-			ReverseForEach(SightQueriesPending, RemoveQuery);
+			return EReverseForEachResult::UnTouched;
+		};
+		ReverseForEach(SightQueriesInRange, RemoveQuery);
+		if (ReverseForEach(SightQueriesOutOfRange, RemoveQuery) == EReverseForEachResult::Modified)
+		{
+			bSightQueriesOutOfRangeDirty = true;
 		}
+		ReverseForEach(SightQueriesPending, RemoveQuery);
 	}
 }
 
