Fixed wrong comment

diff --git a/Engine/Source/Runtime/AIModule/Classes/Perception/AIPerceptionSystem.h b/Engine/Source/Runtime/AIModule/Classes/Perception/AIPerceptionSystem.h
--- a/Engine/Source/Runtime/AIModule/Classes/Perception/AIPerceptionSystem.h
+++ b/Engine/Source/Runtime/AIModule/Classes/Perception/AIPerceptionSystem.h
@@ -24,7 +24,7 @@ AIMODULE_API DECLARE_LOG_CATEGORY_EXTERN(LogAIPerception, Warning, All);
 class APawn;
 
 /**
- *	By design checks perception between hostile teams
+ *	AI Subsystem managing AI Perception through registered AI Senses between Listeners and Stimuli Sources
  */
 UCLASS(ClassGroup=AI, config=Game, defaultconfig, MinimalAPI)
 class UAIPerceptionSystem : public UAISubsystem
