PR #12184: Add `MinimalAPI` specifier to `UEnvQueryTest_Distance` class

diff --git a/Engine/Source/Runtime/AIModule/Classes/EnvironmentQuery/Tests/EnvQueryTest_Distance.h b/Engine/Source/Runtime/AIModule/Classes/EnvironmentQuery/Tests/EnvQueryTest_Distance.h
--- a/Engine/Source/Runtime/AIModule/Classes/EnvironmentQuery/Tests/EnvQueryTest_Distance.h
+++ b/Engine/Source/Runtime/AIModule/Classes/EnvironmentQuery/Tests/EnvQueryTest_Distance.h
@@ -21,7 +21,7 @@ namespace EEnvTestDistance
 	};
 }
 
-UCLASS()
+UCLASS(MinimalAPI)
 class UEnvQueryTest_Distance : public UEnvQueryTest
 {
 	GENERATED_UCLASS_BODY()
