Change NavLinkProxy to update navigation data in PostEditChangeProperty similar to what is done in other navigation code.

diff --git a/Engine/Source/Runtime/AIModule/Private/Navigation/NavLinkProxy.cpp b/Engine/Source/Runtime/AIModule/Private/Navigation/NavLinkProxy.cpp
--- a/Engine/Source/Runtime/AIModule/Private/Navigation/NavLinkProxy.cpp
+++ b/Engine/Source/Runtime/AIModule/Private/Navigation/NavLinkProxy.cpp
@@ -182,7 +182,7 @@ void ANavLinkProxy::PostEditChangeProperty(FPropertyChangedEvent& PropertyChange
 		UNavigationSystemV1* NavSys = FNavigationSystem::GetCurrent<UNavigationSystemV1>(GetWorld());
 		if (NavSys)
 		{
-			NavSys->UpdateActorInNavOctree(*this);
+			FNavigationSystem::UpdateActorData(*this);
 		}
 	}
 
