Allow to lock a location for navmesh gathering in the GameplayDebugger

diff --git a/Engine/Source/Runtime/AIModule/Private/GameplayDebugger/GameplayDebuggerCategory_Navmesh.cpp b/Engine/Source/Runtime/AIModule/Private/GameplayDebugger/GameplayDebuggerCategory_Navmesh.cpp
--- a/Engine/Source/Runtime/AIModule/Private/GameplayDebugger/GameplayDebuggerCategory_Navmesh.cpp
+++ b/Engine/Source/Runtime/AIModule/Private/GameplayDebugger/GameplayDebuggerCategory_Navmesh.cpp
@@ -45,9 +45,11 @@ FGameplayDebuggerCategory_Navmesh::FGameplayDebuggerCategory_Navmesh()
 
 	const FGameplayDebuggerInputHandlerConfig CycleActorReference(TEXT("Cycle Actor Reference"), TEXT("Subtract"), FGameplayDebuggerInputModifier::Shift);
 	const FGameplayDebuggerInputHandlerConfig CycleNavigationData(TEXT("Cycle NavData"), TEXT("Add"), FGameplayDebuggerInputModifier::Shift);
-	
+	const FGameplayDebuggerInputHandlerConfig LocReferenceLocation(TEXT("Lock Reference Location"), TEXT("Multiply"), FGameplayDebuggerInputModifier::Shift);
+
 	BindKeyPress(CycleActorReference, this, &FGameplayDebuggerCategory_Navmesh::CycleActorReference, EGameplayDebuggerInputMode::Replicated);
 	BindKeyPress(CycleNavigationData, this, &FGameplayDebuggerCategory_Navmesh::CycleNavData, EGameplayDebuggerInputMode::Replicated);
+	BindKeyPress(LocReferenceLocation, this, &FGameplayDebuggerCategory_Navmesh::ToggleLockedReferenceLocation, EGameplayDebuggerInputMode::Replicated);
 }
 
 void FGameplayDebuggerCategory_Navmesh::CycleNavData()
@@ -76,6 +78,11 @@ void FGameplayDebuggerCategory_Navmesh::CycleActorReference()
 	}
 }
 
+void FGameplayDebuggerCategory_Navmesh::ToggleLockedReferenceLocation()
+{
+	bToggleLockedReferenceLocation = true;
+}
+
 TSharedRef<FGameplayDebuggerCategory> FGameplayDebuggerCategory_Navmesh::MakeInstance()
 {
 	return MakeShareable(new FGameplayDebuggerCategory_Navmesh());
@@ -88,6 +95,8 @@ void FGameplayDebuggerCategory_Navmesh::FRepData::Serialize(FArchive& Ar)
 	Ar << NumRemainingTasks;
 	Ar << NavDataName;
 
+	Ar << LockedReferenceLocation;
+
 	Ar << NavBuildLockStatusDesc;
 	Ar << SupportedAgents;
 	Ar << NumSuspendedDirtyAreas;
@@ -201,6 +210,13 @@ void FGameplayDebuggerCategory_Navmesh::CollectData(APlayerController* OwnerPC,
 		}
 	}
 
+	if (bToggleLockedReferenceLocation)
+	{
+		LockedReferenceLocation = (!LockedReferenceLocation.IsSet() && RefPawn) ? RefPawn->GetActorLocation() : TOptional<FVector>();
+		bToggleLockedReferenceLocation = false;
+	}
+	DataPack.LockedReferenceLocation = LockedReferenceLocation.Get(FNavigationSystem::InvalidLocation);
+
 	if (NavData)
 	{
 		DataPack.bIsUsingPlayerActor = (ActorReferenceMode != EActorReferenceMode::DebugActor);
@@ -221,21 +237,26 @@ void FGameplayDebuggerCategory_Navmesh::CollectData(APlayerController* OwnerPC,
 }
 
 void FGameplayDebuggerCategory_Navmesh::CollectNavigationData(const UNavigationSystemV1* NavSys, const ANavigationData* NavData, const APawn* RefPawn)
+{
+	const FVector RefLocation = LockedReferenceLocation.Get(RefPawn ? RefPawn->GetActorLocation() : FNavigationSystem::InvalidLocation);
+	CollectNavigationData(NavSys, NavData, RefLocation);
+}
+
+void FGameplayDebuggerCategory_Navmesh::CollectNavigationData(const UNavigationSystemV1* NavSys, const ANavigationData* NavData, const FVector& RefLocation)
 {
 #if WITH_RECAST
 	const ARecastNavMesh* RecastNavMesh = Cast<const ARecastNavMesh>(NavData);
-	if (RecastNavMesh && RefPawn)
+	if (RecastNavMesh && RefLocation != FNavigationSystem::InvalidLocation)
 	{
 		// add NxN neighborhood of target (where N is the number of tiles)
 		// Note that we round up to the next odd number to keep the reference position in the middle tile
-		const FVector TargetLocation = RefPawn->GetActorLocation();
 
 		TArray<FIntPoint> TileDeltas;
 		RetrieveRelativeTilesToDisplay(TileDeltas);
 
 		int32 TargetTileX = 0;
 		int32 TargetTileY = 0;
-		RecastNavMesh->GetNavMeshTileXY(TargetLocation, TargetTileX, TargetTileY);
+		RecastNavMesh->GetNavMeshTileXY(RefLocation, TargetTileX, TargetTileY);
 
 		TArray<FNavTileRef> TileSet;
 		for (const FIntPoint& TileDelta : TileDeltas)
@@ -320,6 +341,11 @@ void FGameplayDebuggerCategory_Navmesh::DrawData(APlayerController* OwnerPC, FGa
 	{
 		CanvasContext.Printf(TEXT("[{yellow}%s{white}]: Display around %s actor"), *GetInputHandlerDescription(0), DataPack.bIsUsingPlayerActor ? TEXT("Debug") : TEXT("Player"));
 	}
+
+	CanvasContext.Printf(TEXT("[{yellow}%s{white}]: %s"),
+		*GetInputHandlerDescription(2),
+		DataPack.LockedReferenceLocation != FNavigationSystem::InvalidLocation ? *FString::Printf(TEXT("Locked at %s"), *DataPack.LockedReferenceLocation.ToString()) : TEXT("Not Locked") );
+
 }
 
 void FGameplayDebuggerCategory_Navmesh::OnDataPackReplicated(int32 DataPackId)
diff --git a/Engine/Source/Runtime/AIModule/Public/GameplayDebugger/GameplayDebuggerCategory_Navmesh.h b/Engine/Source/Runtime/AIModule/Public/GameplayDebugger/GameplayDebuggerCategory_Navmesh.h
--- a/Engine/Source/Runtime/AIModule/Public/GameplayDebugger/GameplayDebuggerCategory_Navmesh.h
+++ b/Engine/Source/Runtime/AIModule/Public/GameplayDebugger/GameplayDebuggerCategory_Navmesh.h
@@ -27,7 +27,13 @@ protected:
 
 	void CycleNavData();
 	void CycleActorReference();
+	void ToggleLockedReferenceLocation();
+
+	/** Called on the server to collect data for the specified navigation data relative to RefPawn.
+	 * RefPawn will be used to define the location where we need to collect the data and then call CollectNavigationData's version that receive a location as parameter. */
 	AIMODULE_API virtual void CollectNavigationData(const UNavigationSystemV1* NavSys, const ANavigationData* NavData, const APawn* RefPawn);
+	/** Called on the server to collect data for the specified navigation data around a location. */
+	AIMODULE_API virtual void CollectNavigationData(const UNavigationSystemV1* NavSys, const ANavigationData* NavData, const FVector& RefLocation);
 	AIMODULE_API void RetrieveRelativeTilesToDisplay(TArray<FIntPoint>& OutTileDelta);
 
 	struct FRepData
@@ -37,6 +43,7 @@ protected:
 		FString NavDataName;
 		FString NavBuildLockStatusDesc;
 		FString SupportedAgents;
+		FVector LockedReferenceLocation = FNavigationSystem::InvalidLocation;
 		int32 NumDirtyAreas = 0;
 		int32 NumSuspendedDirtyAreas = 0;
 		uint16 NumRunningTasks = 0;
@@ -64,6 +71,8 @@ protected:
 	int32 NavDataIndexToDisplay = INDEX_NONE;
 	bool bSwitchToNextNavigationData = false;
 	TWeakObjectPtr<const APawn> PrevDebugActorReference;
+	bool bToggleLockedReferenceLocation = false;
+	TOptional<FVector> LockedReferenceLocation;
 };
 
 #endif // WITH_GAMEPLAY_DEBUGGER_MENU
