fix typo in GameplayDebugger category navmesh

diff --git a/Engine/Source/Runtime/AIModule/Private/GameplayDebugger/GameplayDebuggerCategory_Navmesh.cpp b/Engine/Source/Runtime/AIModule/Private/GameplayDebugger/GameplayDebuggerCategory_Navmesh.cpp
--- a/Engine/Source/Runtime/AIModule/Private/GameplayDebugger/GameplayDebuggerCategory_Navmesh.cpp
+++ b/Engine/Source/Runtime/AIModule/Private/GameplayDebugger/GameplayDebuggerCategory_Navmesh.cpp
@@ -45,11 +45,11 @@ FGameplayDebuggerCategory_Navmesh::FGameplayDebuggerCategory_Navmesh()
 
 	const FGameplayDebuggerInputHandlerConfig CycleActorReference(TEXT("Cycle Actor Reference"), TEXT("Subtract"), FGameplayDebuggerInputModifier::Shift);
 	const FGameplayDebuggerInputHandlerConfig CycleNavigationData(TEXT("Cycle NavData"), TEXT("Add"), FGameplayDebuggerInputModifier::Shift);
-	const FGameplayDebuggerInputHandlerConfig LocReferenceLocation(TEXT("Lock Reference Location"), TEXT("Multiply"), FGameplayDebuggerInputModifier::Shift);
+	const FGameplayDebuggerInputHandlerConfig LockReferenceLocation(TEXT("Lock Reference Location"), TEXT("Multiply"), FGameplayDebuggerInputModifier::Shift);
 
 	BindKeyPress(CycleActorReference, this, &FGameplayDebuggerCategory_Navmesh::CycleActorReference, EGameplayDebuggerInputMode::Replicated);
 	BindKeyPress(CycleNavigationData, this, &FGameplayDebuggerCategory_Navmesh::CycleNavData, EGameplayDebuggerInputMode::Replicated);
-	BindKeyPress(LocReferenceLocation, this, &FGameplayDebuggerCategory_Navmesh::ToggleLockedReferenceLocation, EGameplayDebuggerInputMode::Replicated);
+	BindKeyPress(LockReferenceLocation, this, &FGameplayDebuggerCategory_Navmesh::ToggleLockedReferenceLocation, EGameplayDebuggerInputMode::Replicated);
 }
 
 void FGameplayDebuggerCategory_Navmesh::CycleNavData()
