[Behavior Tree]
- Added ensure to track BehaviorTreeComponent edge case with null ExecuteNode from an ExecutionRequest

diff --git a/Engine/Source/Runtime/AIModule/Private/BehaviorTree/BehaviorTreeComponent.cpp b/Engine/Source/Runtime/AIModule/Private/BehaviorTree/BehaviorTreeComponent.cpp
--- a/Engine/Source/Runtime/AIModule/Private/BehaviorTree/BehaviorTreeComponent.cpp
+++ b/Engine/Source/Runtime/AIModule/Private/BehaviorTree/BehaviorTreeComponent.cpp
@@ -45,6 +45,13 @@ namespace UE::BehaviorTree
 		TEXT("BehaviorTree.ApplyAuxNodesFromFailedSearches"),
 		bApplyAuxNodesFromFailedSearches,
 		TEXT("Apply Aux Nodes From Failed Searches"));
+
+	// Enables the ensure in ScheduleExecutionUpdate to validate that either the pending execution or execution node is valid
+	static bool bEnsureOnScheduleExecutionUpdateWithInvalidExecutionRequest = true;
+	static FAutoConsoleVariableRef CVarEnsureOnScheduleExecutionUpdateWithInvalidExecutionRequest(
+		TEXT("BehaviorTree.EnsureOnScheduleExecutionUpdateWithInvalidExecutionRequest"), 
+		bEnsureOnScheduleExecutionUpdateWithInvalidExecutionRequest,
+		TEXT("Should we ensure when pending execution and execution node are invalid when scheduling execution update."));
 }
 
 //----------------------------------------------------------------------//
@@ -1137,6 +1144,8 @@ static void FindCommonParent(const TArray<FBehaviorTreeInstance>& Instances, con
 void UBehaviorTreeComponent::ScheduleExecutionUpdate()
 {
 	ScheduleNextTick(0.0f);
+
+	ensureMsgf(!UE::BehaviorTree::bEnsureOnScheduleExecutionUpdateWithInvalidExecutionRequest || PendingExecution.IsSet() || ExecutionRequest.ExecuteNode, TEXT("Expecting either a pending execution or an execution node in the request."));
 	bRequestedFlowUpdate = true;
 }
 
@@ -1966,6 +1975,12 @@ void UBehaviorTreeComponent::ProcessExecutionRequest()
 		return;
 	}
 
+	if (ExecutionRequest.ExecuteNode == nullptr)
+	{
+		UE_VLOG(GetOwner(), LogBehaviorTree, Display, TEXT("Ignoring ProcessExecutionRequest call, node to be executed is not set."));
+		return;
+	}
+
 	bool bIsSearchValid = true;
 	SearchData.RollbackInstanceIdx = ActiveInstanceIdx;
 	SearchData.RollbackDeactivatedBranchStart = SearchData.DeactivatedBranchStart;
