PR #12238: Allow using components for IAISightTargetInterface

diff --git a/Engine/Source/Runtime/AIModule/Private/Perception/AISense_Sight.cpp b/Engine/Source/Runtime/AIModule/Private/Perception/AISense_Sight.cpp
--- a/Engine/Source/Runtime/AIModule/Private/Perception/AISense_Sight.cpp
+++ b/Engine/Source/Runtime/AIModule/Private/Perception/AISense_Sight.cpp
@@ -784,20 +784,24 @@ bool UAISense_Sight::RegisterTarget(AActor& TargetActor, const TFunction<void(FA
 	
 	FAISightTarget* SightTarget = ObservedTargets.Find(TargetActor.GetUniqueID());
 	
-	if (SightTarget != nullptr && SightTarget->GetTargetActor() != &TargetActor)
+	// Check if the target is recycled OR new
+	if (SightTarget == nullptr || SightTarget->GetTargetActor() != &TargetActor)
 	{
-		// this means given unique ID has already been recycled. 
 		FAISightTarget NewSightTarget(&TargetActor);
 
 		SightTarget = &(ObservedTargets.Add(NewSightTarget.TargetId, NewSightTarget));
-		SightTarget->SightTargetInterface = Cast<IAISightTargetInterface>(&TargetActor);
-	}
-	else if (SightTarget == nullptr)
-	{
-		FAISightTarget NewSightTarget(&TargetActor);
 
-		SightTarget = &(ObservedTargets.Add(NewSightTarget.TargetId, NewSightTarget));
-		SightTarget->SightTargetInterface = Cast<IAISightTargetInterface>(&TargetActor);
+		// we're looking at components first and only if nothing is found we proceed to check 
+		// if the TargetActor implements IAISightTargetInterface. The advantage of doing it in 
+		// this order is that you can have components override the original Actor's implementation
+		if (IAISightTargetInterface* InterfaceComponent = TargetActor.FindComponentByInterface<IAISightTargetInterface>())
+		{
+			SightTarget->SightTargetInterface = InterfaceComponent;
+		}
+		else 
+		{
+			SightTarget->SightTargetInterface = Cast<IAISightTargetInterface>(&TargetActor);
+		}
 	}
 
 	// set/update data
