[AI Perception] Follow up to PR#12238 to ensure that accessing the stored IAISightTargetInterface instance is safe.

The original change made it possible for FAISightTarget.SightTargetInterface to become invalid while FAISightTarget's contents indicated that it can still be used (due to FAISightTarget.Target being still valid). This could happen if SightTargetInterface was pointing at a component implementing the interface, and the component got removed from the owning actor, without destroying the actor.

diff --git a/Engine/Source/Runtime/AIModule/Classes/Perception/AISense_Sight.h b/Engine/Source/Runtime/AIModule/Classes/Perception/AISense_Sight.h
--- a/Engine/Source/Runtime/AIModule/Classes/Perception/AISense_Sight.h
+++ b/Engine/Source/Runtime/AIModule/Classes/Perception/AISense_Sight.h
@@ -54,10 +54,13 @@ struct FAISightTarget
 	static AIMODULE_API const FTargetId InvalidTargetId;
 
 	TWeakObjectPtr<AActor> Target;
-	IAISightTargetInterface* SightTargetInterface;
+	TWeakInterfacePtr<IAISightTargetInterface> WeakSightTargetInterface;
 	FGenericTeamId TeamId;
 	FTargetId TargetId;
 
+	UE_DEPRECATED_FORGAME(5.5, "SightTargetInterface is deprecated. Use WeakSightTargetInterface instead.")
+	IAISightTargetInterface* SightTargetInterface = nullptr;
+
 	AIMODULE_API FAISightTarget(AActor* InTarget = NULL, FGenericTeamId InTeamId = FGenericTeamId::NoTeam);
 
 	FORCEINLINE FVector GetLocationSimple() const
diff --git a/Engine/Source/Runtime/AIModule/Private/Perception/AISense_Sight.cpp b/Engine/Source/Runtime/AIModule/Private/Perception/AISense_Sight.cpp
--- a/Engine/Source/Runtime/AIModule/Private/Perception/AISense_Sight.cpp
+++ b/Engine/Source/Runtime/AIModule/Private/Perception/AISense_Sight.cpp
@@ -95,7 +95,7 @@ EReverseForEachResult ReverseForEach(T& Array, const PREDICATE_CLASS& Predicate)
 const FAISightTarget::FTargetId FAISightTarget::InvalidTargetId = FAISystem::InvalidUnsignedID;
 
 FAISightTarget::FAISightTarget(AActor* InTarget, FGenericTeamId InTeamId)
-	: Target(InTarget), SightTargetInterface(nullptr), TeamId(InTeamId)
+	: Target(InTarget), TeamId(InTeamId)
 {
 	if (InTarget)
 	{
@@ -530,7 +530,7 @@ UAISense_Sight::EVisibilityResult UAISense_Sight::ComputeVisibility(UWorld* Worl
 		return EVisibilityResult::NotVisible;
 	}
 
-	if (Target.SightTargetInterface != nullptr)
+	if (IAISightTargetInterface* SightTargetInterface = Target.WeakSightTargetInterface.Get())
 	{
 		const bool bWasVisible = SightQuery.GetLastResult();
 		FCanBeSeenFromContext Context;
@@ -539,7 +539,7 @@ UAISense_Sight::EVisibilityResult UAISense_Sight::ComputeVisibility(UWorld* Worl
 		Context.IgnoreActor = ListenerActor;
 		Context.bWasVisible = &bWasVisible;
 
-		const EVisibilityResult Result = Target.SightTargetInterface->CanBeSeenFrom(Context, OutSeenLocation, OutNumberOfLoSChecksPerformed, OutNumberOfAsyncLosCheckRequested, OutStimulusStrength, &SightQuery.UserData, &OnPendingCanBeSeenQueryProcessedDelegate);
+		const EVisibilityResult Result = SightTargetInterface->CanBeSeenFrom(Context, OutSeenLocation, OutNumberOfLoSChecksPerformed, OutNumberOfAsyncLosCheckRequested, OutStimulusStrength, &SightQuery.UserData, &OnPendingCanBeSeenQueryProcessedDelegate);
 		if (Result == EVisibilityResult::Pending)
 		{
 			// we need to clear the trace info value in order to avoid interfering with the engine processed asynchronous queries
@@ -796,11 +796,11 @@ bool UAISense_Sight::RegisterTarget(AActor& TargetActor, const TFunction<void(FA
 		// this order is that you can have components override the original Actor's implementation
 		if (IAISightTargetInterface* InterfaceComponent = TargetActor.FindComponentByInterface<IAISightTargetInterface>())
 		{
-			SightTarget->SightTargetInterface = InterfaceComponent;
+			SightTarget->WeakSightTargetInterface = InterfaceComponent;
 		}
 		else 
 		{
-			SightTarget->SightTargetInterface = Cast<IAISightTargetInterface>(&TargetActor);
+			SightTarget->WeakSightTargetInterface = Cast<IAISightTargetInterface>(&TargetActor);
 		}
 	}
 
