Make sure reservation for donut EQS generator is big enough : we expect to have NumPoints * NumRings around each center, so reserving for only one isn't enough

diff --git a/Engine/Source/Runtime/AIModule/Private/EnvironmentQuery/Generators/EnvQueryGenerator_Donut.cpp b/Engine/Source/Runtime/AIModule/Private/EnvironmentQuery/Generators/EnvQueryGenerator_Donut.cpp
--- a/Engine/Source/Runtime/AIModule/Private/EnvironmentQuery/Generators/EnvQueryGenerator_Donut.cpp
+++ b/Engine/Source/Runtime/AIModule/Private/EnvironmentQuery/Generators/EnvQueryGenerator_Donut.cpp
@@ -116,7 +116,7 @@ bool UEnvQueryGenerator_Donut::GenerateDonutNavPoints(const TArray<FVector>& Cen
 	const FVector::FReal AngleDelta = 2. * UE_DOUBLE_PI / NumPoints;
 	FVector::FReal SectionAngle = FMath::DegreesToRadians(ArcBisectDeg);
 
-	Points.Reserve(NumPoints * NumRings);
+	Points.Reserve(NumPoints * NumRings * CenterPoints.Num());
 
 	if (!bUseSpiralPattern)
 	{
