Exposing FValueOrBBKeyBase's Key property to BP

diff --git a/Engine/Source/Runtime/AIModule/Classes/BehaviorTree/ValueOrBBKey.h b/Engine/Source/Runtime/AIModule/Classes/BehaviorTree/ValueOrBBKey.h
--- a/Engine/Source/Runtime/AIModule/Classes/BehaviorTree/ValueOrBBKey.h
+++ b/Engine/Source/Runtime/AIModule/Classes/BehaviorTree/ValueOrBBKey.h
@@ -78,7 +78,7 @@ namespace FBlackboard
 } // namespace FBlackboard
 
 // Base struct to simplify edition in the editor, shouldn't be used elsewhere
-USTRUCT()
+USTRUCT(BlueprintType)
 struct FValueOrBlackboardKeyBase
 {
 	GENERATED_BODY()
@@ -112,7 +112,7 @@ protected:
 
 	AIMODULE_API FString ToStringKeyName() const;
 
-	UPROPERTY(EditAnywhere, Category = "Value")
+	UPROPERTY(EditAnywhere, BlueprintReadOnly, Category = "Value")
 	FName Key;
 
 	mutable FBlackboard::FKey KeyId = FBlackboard::InvalidKey;
