Make EnvQueryGenerator derive from UEnvQueryGenerator_ProjectedPoints so the locations can be projected.
#rnx

diff --git a/Engine/Source/Runtime/AIModule/Classes/EnvironmentQuery/Generators/EnvQueryGenerator_CurrentLocation.h b/Engine/Source/Runtime/AIModule/Classes/EnvironmentQuery/Generators/EnvQueryGenerator_CurrentLocation.h
--- a/Engine/Source/Runtime/AIModule/Classes/EnvironmentQuery/Generators/EnvQueryGenerator_CurrentLocation.h
+++ b/Engine/Source/Runtime/AIModule/Classes/EnvironmentQuery/Generators/EnvQueryGenerator_CurrentLocation.h
@@ -2,15 +2,17 @@
 
 #pragma once
 
-#include "CoreMinimal.h"
+#if UE_ENABLE_INCLUDE_ORDER_DEPRECATED_IN_5_6
+	#include "CoreMinimal.h"
+#endif
 #include "UObject/ObjectMacros.h"
 #include "Templates/SubclassOf.h"
 #include "EnvironmentQuery/EnvQueryContext.h"
-#include "EnvironmentQuery/EnvQueryGenerator.h"
+#include "EnvironmentQuery/Generators/EnvQueryGenerator_ProjectedPoints.h"
 #include "EnvQueryGenerator_CurrentLocation.generated.h"
 
 UCLASS(meta = (DisplayName = "Current Location"), MinimalAPI)
-class UEnvQueryGenerator_CurrentLocation : public UEnvQueryGenerator
+class UEnvQueryGenerator_CurrentLocation : public UEnvQueryGenerator_ProjectedPoints
 {
 	GENERATED_BODY()
 
diff --git a/Engine/Source/Runtime/AIModule/Private/EnvironmentQuery/Generators/EnvQueryGenerator_CurrentLocation.cpp b/Engine/Source/Runtime/AIModule/Private/EnvironmentQuery/Generators/EnvQueryGenerator_CurrentLocation.cpp
--- a/Engine/Source/Runtime/AIModule/Private/EnvironmentQuery/Generators/EnvQueryGenerator_CurrentLocation.cpp
+++ b/Engine/Source/Runtime/AIModule/Private/EnvironmentQuery/Generators/EnvQueryGenerator_CurrentLocation.cpp
@@ -4,6 +4,7 @@
 #include "AI/Navigation/NavigationTypes.h"
 #include "EnvironmentQuery/Items/EnvQueryItemType_Point.h"
 #include "EnvironmentQuery/Contexts/EnvQueryContext_Querier.h"
+#include "UObject/Class.h"
 
 #include UE_INLINE_GENERATED_CPP_BY_NAME(EnvQueryGenerator_CurrentLocation)
 
@@ -14,6 +15,7 @@ UEnvQueryGenerator_CurrentLocation::UEnvQueryGenerator_CurrentLocation(const FOb
 {
 	QueryContext = UEnvQueryContext_Querier::StaticClass();
 	ItemType = UEnvQueryItemType_Point::StaticClass();
+	ProjectionData.TraceMode = EEnvQueryTrace::None; // Set to none by default to preserve behavior to before the node could project location.
 }
 
 void UEnvQueryGenerator_CurrentLocation::GenerateItems(FEnvQueryInstance& QueryInstance) const
@@ -21,21 +23,32 @@ void UEnvQueryGenerator_CurrentLocation::GenerateItems(FEnvQueryInstance& QueryI
 	TArray<FVector> ContextLocations;
 	QueryInstance.PrepareContext(QueryContext, ContextLocations);
 
+	TArray<FNavLocation> NavLocations;
+	NavLocations.Reserve(ContextLocations.Num());
 	for (const FVector& Location : ContextLocations)
 	{
-		FNavLocation NavLoc(Location);
-		QueryInstance.AddItemData<UEnvQueryItemType_Point>(NavLoc);
+		NavLocations.Emplace(Location);
 	}
+
+	ProjectAndFilterNavPoints(NavLocations, QueryInstance);
+	StoreNavPoints(NavLocations, QueryInstance);
 }
 
 FText UEnvQueryGenerator_CurrentLocation::GetDescriptionTitle() const
 {
-	return FText::Format(LOCTEXT("CurrentLocationOn", "Current Location of {0}"), UEnvQueryTypes::DescribeContext(QueryContext));
+	if (ProjectionData.TraceMode != EEnvQueryTrace::None)
+	{
+		return FText::Format(LOCTEXT("CurrentLocationProjectedOn", "Current Location of {0} projected on {1}"), UEnvQueryTypes::DescribeContext(QueryContext), UEnum::GetDisplayValueAsText(ProjectionData.TraceMode));
+	}
+	else
+	{
+		return FText::Format(LOCTEXT("CurrentLocationOn", "Current Location of {0}"), UEnvQueryTypes::DescribeContext(QueryContext));
+	}
 };
 
 FText UEnvQueryGenerator_CurrentLocation::GetDescriptionDetails() const
 {
-	return FText::GetEmpty();
+	return ProjectionData.ToText(FEnvTraceData::EDescriptionMode::Brief);
 }
 
 #undef LOCTEXT_NAMESPACE
