Create BTTask_SetKeyValueX for all blackboard key type in the engine

diff --git a/Engine/Plugins/Experimental/GameplayBehaviors/Source/GameplayBehaviorsModule/Private/AI/BTTask_SetKeyValueGameplayTag.cpp b/Engine/Plugins/Experimental/GameplayBehaviors/Source/GameplayBehaviorsModule/Private/AI/BTTask_SetKeyValueGameplayTag.cpp
new file mode 100644
--- /dev/null
+++ b/Engine/Plugins/Experimental/GameplayBehaviors/Source/GameplayBehaviorsModule/Private/AI/BTTask_SetKeyValueGameplayTag.cpp
@@ -0,0 +1,38 @@
+// Copyright Epic Games, Inc. All Rights Reserved.
+
+#include "AI/BTTask_SetKeyValueGameplayTag.h"
+
+#include "BehaviorTree/BehaviorTreeComponent.h"
+#include "BlackboardKeyType_GameplayTag.h"
+
+UBTTask_SetKeyValueGameplayTag::UBTTask_SetKeyValueGameplayTag(const FObjectInitializer& ObjectInitializer /*= FObjectInitializer::Get()*/)
+{
+	NodeName = "Set GameplayTag Key";
+	BlackboardKey.AllowNoneAsValue(false);
+	const FName FilterName = MakeUniqueObjectName(this, UBlackboardKeyType_GameplayTag::StaticClass(), *FString::Printf(TEXT("%s_GameplayTag"), *GET_MEMBER_NAME_CHECKED(UBTTask_SetKeyValueGameplayTag, BlackboardKey).ToString()));
+	UBlackboardKeyType_GameplayTag* FilterOb = ObjectInitializer.CreateDefaultSubobject<UBlackboardKeyType_GameplayTag>(this, FilterName, /*Transient*/ true);
+	BlackboardKey.AllowedTypes.Add(FilterOb);
+	INIT_TASK_NODE_NOTIFY_FLAGS();
+}
+
+EBTNodeResult::Type UBTTask_SetKeyValueGameplayTag::ExecuteTask(UBehaviorTreeComponent& OwnerComp, uint8* NodeMemory)
+{
+	UBlackboardComponent* Blackboard = OwnerComp.GetBlackboardComponent();
+	if (!Blackboard)
+	{
+		return EBTNodeResult::Failed;
+	}
+
+	if (BlackboardKey.GetSelectedKeyID() == FBlackboard::InvalidKey)
+	{
+		return EBTNodeResult::Failed;
+	}
+
+	Blackboard->SetValue<UBlackboardKeyType_GameplayTag>(BlackboardKey.GetSelectedKeyID(), Value.GetValue(*Blackboard));
+	return EBTNodeResult::Succeeded;
+}
+
+FString UBTTask_SetKeyValueGameplayTag::GetStaticDescription() const
+{
+	return FString::Format(TEXT("Setting {0} to {1}"), { BlackboardKey.SelectedKeyName.ToString(), Value.ToString() });
+}
diff --git a/Engine/Plugins/Experimental/GameplayBehaviors/Source/GameplayBehaviorsModule/Public/AI/BTTask_SetKeyValueGameplayTag.h b/Engine/Plugins/Experimental/GameplayBehaviors/Source/GameplayBehaviorsModule/Public/AI/BTTask_SetKeyValueGameplayTag.h
new file mode 100644
--- /dev/null
+++ b/Engine/Plugins/Experimental/GameplayBehaviors/Source/GameplayBehaviorsModule/Public/AI/BTTask_SetKeyValueGameplayTag.h
@@ -0,0 +1,23 @@
+// Copyright Epic Games, Inc. All Rights Reserved.
+#pragma once
+
+#include "BehaviorTree/Tasks/BTTask_BlackboardBase.h"
+#include "AI/ValueOrBBKey_GameplayTag.h"
+
+#include "BTTask_SetKeyValueGameplayTag.generated.h"
+
+UCLASS(DisplayName = "Set GameplayTag Key")
+class UBTTask_SetKeyValueGameplayTag : public UBTTask_BlackboardBase
+{
+	GENERATED_BODY()
+
+public:
+	UBTTask_SetKeyValueGameplayTag(const FObjectInitializer& ObjectInitializer = FObjectInitializer::Get());
+
+	virtual EBTNodeResult::Type ExecuteTask(UBehaviorTreeComponent& OwnerComp, uint8* NodeMemory) override;
+	virtual FString GetStaticDescription() const override;
+
+private:
+	UPROPERTY(EditAnywhere, Category = Blackboard)
+	FValueOrBBKey_GameplayTagContainer Value;
+};
\ No newline at end of file
diff --git a/Engine/Source/Runtime/AIModule/Classes/BehaviorTree/Tasks/BTTask_SetKeyValue.h b/Engine/Source/Runtime/AIModule/Classes/BehaviorTree/Tasks/BTTask_SetKeyValue.h
new file mode 100644
--- /dev/null
+++ b/Engine/Source/Runtime/AIModule/Classes/BehaviorTree/Tasks/BTTask_SetKeyValue.h
@@ -0,0 +1,211 @@
+// Copyright Epic Games, Inc. All Rights Reserved.
+#pragma once
+
+#include "BehaviorTree/Tasks/BTTask_BlackboardBase.h"
+#include "BehaviorTree/ValueOrBBKey.h"
+
+#include "BTTask_SetKeyValue.generated.h"
+
+UCLASS(DisplayName = "Set Bool Key")
+class UBTTask_SetKeyValueBool : public UBTTask_BlackboardBase
+{
+	GENERATED_BODY()
+
+public:
+	UBTTask_SetKeyValueBool(const FObjectInitializer& ObjectInitializer = FObjectInitializer::Get());
+
+	virtual EBTNodeResult::Type ExecuteTask(UBehaviorTreeComponent& OwnerComp, uint8* NodeMemory) override;
+	virtual FString GetStaticDescription() const override;
+
+private:
+	UPROPERTY(EditAnywhere, Category = Blackboard)
+	FValueOrBBKey_Bool Value;
+};
+
+UCLASS(DisplayName = "Set Class Key")
+class UBTTask_SetKeyValueClass : public UBTTask_BlackboardBase
+{
+	GENERATED_BODY()
+
+public:
+	UBTTask_SetKeyValueClass(const FObjectInitializer& ObjectInitializer = FObjectInitializer::Get());
+
+	virtual EBTNodeResult::Type ExecuteTask(UBehaviorTreeComponent& OwnerComp, uint8* NodeMemory) override;
+	virtual FString GetStaticDescription() const override;
+
+	virtual void PostLoad() override;
+#if WITH_EDITOR
+	virtual void PostEditChangeProperty(FPropertyChangedEvent& PropertyChangedEvent) override;
+#endif // WITH_EDITOR
+private:
+	UPROPERTY(EditAnywhere, Category = Blackboard, meta = (AllowAbstract = "true", NoClear))
+	TSubclassOf<UObject> BaseClass = UObject::StaticClass();
+
+	UPROPERTY(EditAnywhere, Category = Blackboard)
+	FValueOrBBKey_Class Value;
+};
+
+UCLASS(DisplayName = "Set Enum Key")
+class UBTTask_SetKeyValueEnum : public UBTTask_BlackboardBase
+{
+	GENERATED_BODY()
+
+public:
+	UBTTask_SetKeyValueEnum(const FObjectInitializer& ObjectInitializer = FObjectInitializer::Get());
+
+	virtual EBTNodeResult::Type ExecuteTask(UBehaviorTreeComponent& OwnerComp, uint8* NodeMemory) override;
+	virtual FString GetStaticDescription() const override;
+
+	virtual void PostLoad() override;
+#if WITH_EDITOR
+	virtual void PostEditChangeProperty(FPropertyChangedEvent& PropertyChangedEvent) override;
+#endif // WITH_EDITOR
+private:
+	UPROPERTY(EditAnywhere, Category = Blackboard, meta = (AllowAbstract = "true", NoClear))
+	TObjectPtr<UEnum> EnumType;
+
+	UPROPERTY(EditAnywhere, Category = Blackboard)
+	FValueOrBBKey_Enum Value;
+};
+
+UCLASS(DisplayName = "Set Int Key")
+class UBTTask_SetKeyValueInt32 : public UBTTask_BlackboardBase
+{
+	GENERATED_BODY()
+
+public:
+	UBTTask_SetKeyValueInt32(const FObjectInitializer& ObjectInitializer = FObjectInitializer::Get());
+
+	virtual EBTNodeResult::Type ExecuteTask(UBehaviorTreeComponent& OwnerComp, uint8* NodeMemory) override;
+	virtual FString GetStaticDescription() const override;
+
+private:
+	UPROPERTY(EditAnywhere, Category = Blackboard)
+	FValueOrBBKey_Int32 Value;
+};
+
+UCLASS(DisplayName = "Set Float Key")
+class UBTTask_SetKeyValueFloat : public UBTTask_BlackboardBase
+{
+	GENERATED_BODY()
+
+public:
+	UBTTask_SetKeyValueFloat(const FObjectInitializer& ObjectInitializer = FObjectInitializer::Get());
+
+	virtual EBTNodeResult::Type ExecuteTask(UBehaviorTreeComponent& OwnerComp, uint8* NodeMemory) override;
+	virtual FString GetStaticDescription() const override;
+
+private:
+	UPROPERTY(EditAnywhere, Category = Blackboard)
+	FValueOrBBKey_Float Value;
+};
+
+UCLASS(DisplayName = "Set Name Key")
+class UBTTask_SetKeyValueName : public UBTTask_BlackboardBase
+{
+	GENERATED_BODY()
+
+public:
+	UBTTask_SetKeyValueName(const FObjectInitializer& ObjectInitializer = FObjectInitializer::Get());
+
+	virtual EBTNodeResult::Type ExecuteTask(UBehaviorTreeComponent& OwnerComp, uint8* NodeMemory) override;
+	virtual FString GetStaticDescription() const override;
+
+private:
+	UPROPERTY(EditAnywhere, Category = Blackboard)
+	FValueOrBBKey_Name Value;
+};
+
+UCLASS(DisplayName = "Set String Key")
+class UBTTask_SetKeyValueString : public UBTTask_BlackboardBase
+{
+	GENERATED_BODY()
+
+public:
+	UBTTask_SetKeyValueString(const FObjectInitializer& ObjectInitializer = FObjectInitializer::Get());
+
+	virtual EBTNodeResult::Type ExecuteTask(UBehaviorTreeComponent& OwnerComp, uint8* NodeMemory) override;
+	virtual FString GetStaticDescription() const override;
+
+private:
+	UPROPERTY(EditAnywhere, Category = Blackboard)
+	FValueOrBBKey_String Value;
+};
+
+UCLASS(DisplayName = "Set Object Key")
+class UBTTask_SetKeyValueObject : public UBTTask_BlackboardBase
+{
+	GENERATED_BODY()
+
+public:
+	UBTTask_SetKeyValueObject(const FObjectInitializer& ObjectInitializer = FObjectInitializer::Get());
+
+	virtual EBTNodeResult::Type ExecuteTask(UBehaviorTreeComponent& OwnerComp, uint8* NodeMemory) override;
+	virtual FString GetStaticDescription() const override;
+
+	virtual void PostLoad() override;
+#if WITH_EDITOR
+	virtual void PostEditChangeProperty(FPropertyChangedEvent& PropertyChangedEvent) override;
+#endif // WITH_EDITOR
+private:
+	UPROPERTY(EditAnywhere, Category = Blackboard, meta = (AllowAbstract = "true", NoClear))
+	TSubclassOf<UObject> BaseClass = UObject::StaticClass();
+
+	UPROPERTY(EditAnywhere, Category = Blackboard)
+	FValueOrBBKey_Object Value;
+};
+
+UCLASS(DisplayName = "Set Rotator Key")
+class UBTTask_SetKeyValueRotator : public UBTTask_BlackboardBase
+{
+	GENERATED_BODY()
+
+public:
+	UBTTask_SetKeyValueRotator(const FObjectInitializer& ObjectInitializer = FObjectInitializer::Get());
+
+	virtual EBTNodeResult::Type ExecuteTask(UBehaviorTreeComponent& OwnerComp, uint8* NodeMemory) override;
+	virtual FString GetStaticDescription() const override;
+
+private:
+	UPROPERTY(EditAnywhere, Category = Blackboard)
+	FValueOrBBKey_Rotator Value;
+};
+
+UCLASS(DisplayName = "Set Struct Key")
+class UBTTask_SetKeyValueStruct : public UBTTask_BlackboardBase
+{
+	GENERATED_BODY()
+
+public:
+	UBTTask_SetKeyValueStruct(const FObjectInitializer& ObjectInitializer = FObjectInitializer::Get());
+
+	virtual EBTNodeResult::Type ExecuteTask(UBehaviorTreeComponent& OwnerComp, uint8* NodeMemory) override;
+	virtual FString GetStaticDescription() const override;
+
+	virtual void PostLoad() override;
+#if WITH_EDITOR
+	virtual void PostEditChangeProperty(FPropertyChangedEvent& PropertyChangedEvent) override;
+#endif // WITH_EDITOR
+private:
+	UPROPERTY(EditAnywhere, Category = Blackboard)
+	TObjectPtr<UScriptStruct> StructType;
+
+	UPROPERTY(EditAnywhere, Category = Blackboard)
+	FValueOrBBKey_Struct Value;
+};
+
+UCLASS(DisplayName = "Set Vector Key")
+class UBTTask_SetKeyValueVector : public UBTTask_BlackboardBase
+{
+	GENERATED_BODY()
+
+public:
+	UBTTask_SetKeyValueVector(const FObjectInitializer& ObjectInitializer = FObjectInitializer::Get());
+
+	virtual EBTNodeResult::Type ExecuteTask(UBehaviorTreeComponent& OwnerComp, uint8* NodeMemory) override;
+	virtual FString GetStaticDescription() const override;
+
+private:
+	UPROPERTY(EditAnywhere, Category = Blackboard)
+	FValueOrBBKey_Vector Value;
+};
diff --git a/Engine/Source/Runtime/AIModule/Classes/BehaviorTree/ValueOrBBKey.h b/Engine/Source/Runtime/AIModule/Classes/BehaviorTree/ValueOrBBKey.h
--- a/Engine/Source/Runtime/AIModule/Classes/BehaviorTree/ValueOrBBKey.h
+++ b/Engine/Source/Runtime/AIModule/Classes/BehaviorTree/ValueOrBBKey.h
@@ -181,6 +181,8 @@ struct FValueOrBBKey_Class : public FValueOrBlackboardKeyBase
 	UE_DEPRECATED_FORGAME(5.5, "Implicit conversion will be removed next version. Call GetValue instead")
 	operator UClass*() const { return DefaultValue; }
 
+	AIMODULE_API void SetBaseClass(UClass* NewBaseClass);
+
 protected:
 	UPROPERTY(EditAnywhere, Category = "Value")
 	TObjectPtr<UClass> DefaultValue = nullptr;
@@ -224,6 +226,8 @@ struct FValueOrBBKey_Enum : public FValueOrBlackboardKeyBase
 	UE_DEPRECATED_FORGAME(5.5, "Implicit conversion will be removed next version. Call GetValue instead")
 	operator uint8() const { return DefaultValue; }
 
+	AIMODULE_API void SetEnumType(UEnum* NewEnumType);
+
 protected:
 	UPROPERTY(EditAnywhere, Category = "Value")
 	uint8 DefaultValue = 0;
@@ -382,6 +386,8 @@ struct FValueOrBBKey_Object : public FValueOrBlackboardKeyBase
 	UE_DEPRECATED_FORGAME(5.5, "Implicit conversion will be removed next version. Call GetValue instead")
 	operator UObject*() const { return DefaultValue; }
 
+	AIMODULE_API void SetBaseClass(UClass* NewBaseClass);
+
 protected:
 	UPROPERTY(EditAnywhere, Category = "Value")
 	TObjectPtr<UObject> DefaultValue = nullptr;
@@ -479,6 +485,8 @@ public:
 	AIMODULE_API FConstStructView GetValue(const UBlackboardComponent& Blackboard) const;
 	AIMODULE_API FConstStructView GetValue(const UBlackboardComponent* Blackboard) const;
 
+	AIMODULE_API void SetStructType(UScriptStruct* NewStructType);
+
 protected:
 	friend class FValueOrBBKeyDetails_Struct;
 
diff --git a/Engine/Source/Runtime/AIModule/Private/BehaviorTree/Tasks/BTTask_SetKeyValue.cpp b/Engine/Source/Runtime/AIModule/Private/BehaviorTree/Tasks/BTTask_SetKeyValue.cpp
new file mode 100644
--- /dev/null
+++ b/Engine/Source/Runtime/AIModule/Private/BehaviorTree/Tasks/BTTask_SetKeyValue.cpp
@@ -0,0 +1,348 @@
+// Copyright Epic Games, Inc. All Rights Reserved.
+
+#include "BehaviorTree/Tasks/BTTask_SetKeyValue.h"
+
+#include "BehaviorTree/BehaviorTreeComponent.h"
+#include "BehaviorTree/Blackboard/BlackboardKeyAllTypes.h"
+
+namespace FBlackboard
+{
+	template <typename KeyType, typename ValueType>
+	EBTNodeResult::Type SetBlackboardKeyValue(UBehaviorTreeComponent& OwnerComp, const FBlackboardKeySelector& BlackboardKey, const ValueType& Value)
+	{
+		UBlackboardComponent* Blackboard = OwnerComp.GetBlackboardComponent();
+		if (!Blackboard)
+		{
+			return EBTNodeResult::Failed;
+		}
+
+		if (BlackboardKey.GetSelectedKeyID() == FBlackboard::InvalidKey)
+		{
+			return EBTNodeResult::Failed;
+		}
+
+		Blackboard->SetValue<KeyType>(BlackboardKey.GetSelectedKeyID(), Value.GetValue(*Blackboard));
+		return EBTNodeResult::Succeeded;
+	}
+} // namespace FBlackboard
+
+UBTTask_SetKeyValueBool::UBTTask_SetKeyValueBool(const FObjectInitializer& ObjectInitializer /*= FObjectInitializer::Get()*/)
+	: Super(ObjectInitializer)
+{
+	NodeName = "Set Bool Key";
+	BlackboardKey.AllowNoneAsValue(false);
+	BlackboardKey.AddBoolFilter(this, GET_MEMBER_NAME_CHECKED(UBTTask_SetKeyValueBool, BlackboardKey));
+	INIT_TASK_NODE_NOTIFY_FLAGS();
+}
+
+UBTTask_SetKeyValueClass::UBTTask_SetKeyValueClass(const FObjectInitializer& ObjectInitializer /*= FObjectInitializer::Get()*/)
+{
+	NodeName = "Set Class Key";
+	BlackboardKey.AllowNoneAsValue(false);
+	BlackboardKey.AddClassFilter(this, GET_MEMBER_NAME_CHECKED(UBTTask_SetKeyValueClass, BlackboardKey), BaseClass);
+	INIT_TASK_NODE_NOTIFY_FLAGS();
+}
+
+UBTTask_SetKeyValueEnum::UBTTask_SetKeyValueEnum(const FObjectInitializer& ObjectInitializer /*= FObjectInitializer::Get()*/)
+{
+	NodeName = "Set Class Key";
+	// Since there are no good default value for enum don't add a filter and allow none as a value so the key isn't bound to a random key.
+	INIT_TASK_NODE_NOTIFY_FLAGS();
+}
+
+UBTTask_SetKeyValueInt32::UBTTask_SetKeyValueInt32(const FObjectInitializer& ObjectInitializer /*= FObjectInitializer::Get()*/)
+{
+	NodeName = "Set Int Key";
+	BlackboardKey.AllowNoneAsValue(false);
+	BlackboardKey.AddIntFilter(this, GET_MEMBER_NAME_CHECKED(UBTTask_SetKeyValueInt32, BlackboardKey));
+	INIT_TASK_NODE_NOTIFY_FLAGS();
+}
+
+UBTTask_SetKeyValueFloat::UBTTask_SetKeyValueFloat(const FObjectInitializer& ObjectInitializer /*= FObjectInitializer::Get()*/)
+{
+	NodeName = "Set Float Key";
+	BlackboardKey.AllowNoneAsValue(false);
+	BlackboardKey.AddFloatFilter(this, GET_MEMBER_NAME_CHECKED(UBTTask_SetKeyValueFloat, BlackboardKey));
+	INIT_TASK_NODE_NOTIFY_FLAGS();
+}
+
+UBTTask_SetKeyValueName::UBTTask_SetKeyValueName(const FObjectInitializer& ObjectInitializer /*= FObjectInitializer::Get()*/)
+{
+	NodeName = "Set Name Key";
+	BlackboardKey.AllowNoneAsValue(false);
+	BlackboardKey.AddNameFilter(this, GET_MEMBER_NAME_CHECKED(UBTTask_SetKeyValueName, BlackboardKey));
+	INIT_TASK_NODE_NOTIFY_FLAGS();
+}
+
+UBTTask_SetKeyValueString::UBTTask_SetKeyValueString(const FObjectInitializer& ObjectInitializer /*= FObjectInitializer::Get()*/)
+{
+	NodeName = "Set String Key";
+	BlackboardKey.AllowNoneAsValue(false);
+	BlackboardKey.AddStringFilter(this, GET_MEMBER_NAME_CHECKED(UBTTask_SetKeyValueString, BlackboardKey));
+	INIT_TASK_NODE_NOTIFY_FLAGS();
+}
+
+UBTTask_SetKeyValueObject::UBTTask_SetKeyValueObject(const FObjectInitializer& ObjectInitializer /*= FObjectInitializer::Get()*/)
+{
+	NodeName = "Set Object Key";
+	BlackboardKey.AllowNoneAsValue(false);
+	BlackboardKey.AddObjectFilter(this, GET_MEMBER_NAME_CHECKED(UBTTask_SetKeyValueObject, BlackboardKey), BaseClass);
+	INIT_TASK_NODE_NOTIFY_FLAGS();
+}
+
+UBTTask_SetKeyValueRotator::UBTTask_SetKeyValueRotator(const FObjectInitializer& ObjectInitializer /*= FObjectInitializer::Get()*/)
+{
+	NodeName = "Set Rotator Key";
+	BlackboardKey.AllowNoneAsValue(false);
+	BlackboardKey.AddRotatorFilter(this, GET_MEMBER_NAME_CHECKED(UBTTask_SetKeyValueRotator, BlackboardKey));
+	INIT_TASK_NODE_NOTIFY_FLAGS();
+}
+
+UBTTask_SetKeyValueStruct::UBTTask_SetKeyValueStruct(const FObjectInitializer& ObjectInitializer /*= FObjectInitializer::Get()*/)
+{
+	NodeName = "Set Struct Key";
+	//Since there are no good default value for structtype don't add a filter and allow none as a value so the key isn't bound to a random key.
+	INIT_TASK_NODE_NOTIFY_FLAGS();
+}
+
+UBTTask_SetKeyValueVector::UBTTask_SetKeyValueVector(const FObjectInitializer& ObjectInitializer /*= FObjectInitializer::Get()*/)
+{
+	NodeName = "Set Vector Key";
+	BlackboardKey.AllowNoneAsValue(false);
+	BlackboardKey.AddVectorFilter(this, GET_MEMBER_NAME_CHECKED(UBTTask_SetKeyValueVector, BlackboardKey));
+	INIT_TASK_NODE_NOTIFY_FLAGS();
+}
+
+EBTNodeResult::Type UBTTask_SetKeyValueBool::ExecuteTask(UBehaviorTreeComponent& OwnerComp, uint8* NodeMemory)
+{
+	return FBlackboard::SetBlackboardKeyValue<UBlackboardKeyType_Bool>(OwnerComp, BlackboardKey, Value);
+}
+
+EBTNodeResult::Type UBTTask_SetKeyValueClass::ExecuteTask(UBehaviorTreeComponent& OwnerComp, uint8* NodeMemory)
+{
+	return FBlackboard::SetBlackboardKeyValue<UBlackboardKeyType_Class>(OwnerComp, BlackboardKey, Value);
+}
+
+EBTNodeResult::Type UBTTask_SetKeyValueEnum::ExecuteTask(UBehaviorTreeComponent& OwnerComp, uint8* NodeMemory)
+{
+	return FBlackboard::SetBlackboardKeyValue<UBlackboardKeyType_Enum>(OwnerComp, BlackboardKey, Value);
+}
+
+EBTNodeResult::Type UBTTask_SetKeyValueInt32::ExecuteTask(UBehaviorTreeComponent& OwnerComp, uint8* NodeMemory)
+{
+	return FBlackboard::SetBlackboardKeyValue<UBlackboardKeyType_Int>(OwnerComp, BlackboardKey, Value);
+}
+
+EBTNodeResult::Type UBTTask_SetKeyValueFloat::ExecuteTask(UBehaviorTreeComponent& OwnerComp, uint8* NodeMemory)
+{
+	return FBlackboard::SetBlackboardKeyValue<UBlackboardKeyType_Float>(OwnerComp, BlackboardKey, Value);
+}
+
+EBTNodeResult::Type UBTTask_SetKeyValueName::ExecuteTask(UBehaviorTreeComponent& OwnerComp, uint8* NodeMemory)
+{
+	return FBlackboard::SetBlackboardKeyValue<UBlackboardKeyType_Name>(OwnerComp, BlackboardKey, Value);
+}
+
+EBTNodeResult::Type UBTTask_SetKeyValueString::ExecuteTask(UBehaviorTreeComponent& OwnerComp, uint8* NodeMemory)
+{
+	return FBlackboard::SetBlackboardKeyValue<UBlackboardKeyType_String>(OwnerComp, BlackboardKey, Value);
+}
+
+EBTNodeResult::Type UBTTask_SetKeyValueObject::ExecuteTask(UBehaviorTreeComponent& OwnerComp, uint8* NodeMemory)
+{
+	return FBlackboard::SetBlackboardKeyValue<UBlackboardKeyType_Object>(OwnerComp, BlackboardKey, Value);
+}
+
+EBTNodeResult::Type UBTTask_SetKeyValueRotator::ExecuteTask(UBehaviorTreeComponent& OwnerComp, uint8* NodeMemory)
+{
+	return FBlackboard::SetBlackboardKeyValue<UBlackboardKeyType_Rotator>(OwnerComp, BlackboardKey, Value);
+}
+
+EBTNodeResult::Type UBTTask_SetKeyValueStruct::ExecuteTask(UBehaviorTreeComponent& OwnerComp, uint8* NodeMemory)
+{
+	return FBlackboard::SetBlackboardKeyValue<UBlackboardKeyType_Struct>(OwnerComp, BlackboardKey, Value);
+}
+
+EBTNodeResult::Type UBTTask_SetKeyValueVector::ExecuteTask(UBehaviorTreeComponent& OwnerComp, uint8* NodeMemory)
+{
+	return FBlackboard::SetBlackboardKeyValue<UBlackboardKeyType_Vector>(OwnerComp, BlackboardKey, Value);
+}
+
+FString UBTTask_SetKeyValueBool::GetStaticDescription() const
+{
+	return FString::Format(TEXT("Setting {0} to {1}"), { BlackboardKey.SelectedKeyName.ToString(), Value.ToString() });
+}
+
+FString UBTTask_SetKeyValueClass::GetStaticDescription() const
+{
+	return FString::Format(TEXT("Setting {0} to {1}"), { BlackboardKey.SelectedKeyName.ToString(), Value.ToString() });
+}
+
+FString UBTTask_SetKeyValueEnum::GetStaticDescription() const
+{
+	return FString::Format(TEXT("Setting {0} to {1}"), { BlackboardKey.SelectedKeyName.ToString(), Value.ToString() });
+}
+
+FString UBTTask_SetKeyValueInt32::GetStaticDescription() const
+{
+	return FString::Format(TEXT("Setting {0} to {1}"), { BlackboardKey.SelectedKeyName.ToString(), Value.ToString() });
+}
+
+FString UBTTask_SetKeyValueFloat::GetStaticDescription() const
+{
+	return FString::Format(TEXT("Setting {0} to {1}"), { BlackboardKey.SelectedKeyName.ToString(), Value.ToString() });
+}
+
+FString UBTTask_SetKeyValueName::GetStaticDescription() const
+{
+	return FString::Format(TEXT("Setting {0} to {1}"), { BlackboardKey.SelectedKeyName.ToString(), Value.ToString() });
+}
+
+FString UBTTask_SetKeyValueString::GetStaticDescription() const
+{
+	return FString::Format(TEXT("Setting {0} to {1}"), { BlackboardKey.SelectedKeyName.ToString(), Value.ToString() });
+}
+
+FString UBTTask_SetKeyValueObject::GetStaticDescription() const
+{
+	return FString::Format(TEXT("Setting {0} to {1}"), { BlackboardKey.SelectedKeyName.ToString(), Value.ToString() });
+}
+
+FString UBTTask_SetKeyValueRotator::GetStaticDescription() const
+{
+	return FString::Format(TEXT("Setting {0} to {1}"), { BlackboardKey.SelectedKeyName.ToString(), Value.ToString() });
+}
+
+FString UBTTask_SetKeyValueStruct::GetStaticDescription() const
+{
+	return FString::Format(TEXT("Setting {0} to {1}"), { BlackboardKey.SelectedKeyName.ToString(), Value.ToString() });
+}
+
+FString UBTTask_SetKeyValueVector::GetStaticDescription() const
+{
+	return FString::Format(TEXT("Setting {0} to {1}"), { BlackboardKey.SelectedKeyName.ToString(), Value.ToString() });
+}
+
+void UBTTask_SetKeyValueClass::PostLoad()
+{
+	Super::PostLoad();
+
+	if (BaseClass != UObject::StaticClass())
+	{
+		BlackboardKey.AllowedTypes.Empty();
+		BlackboardKey.AddClassFilter(this, GET_MEMBER_NAME_CHECKED(UBTTask_SetKeyValueClass, BlackboardKey), BaseClass);
+		Value.SetBaseClass(BaseClass);
+	}
+}
+
+void UBTTask_SetKeyValueEnum::PostLoad()
+{
+	Super::PostLoad();
+
+	if (EnumType)
+	{
+		BlackboardKey.AllowedTypes.Empty();
+		BlackboardKey.AddEnumFilter(this, GET_MEMBER_NAME_CHECKED(UBTTask_SetKeyValueEnum, BlackboardKey), EnumType);
+		BlackboardKey.AllowNoneAsValue(false);
+		Value.SetEnumType(EnumType);
+	}
+}
+
+void UBTTask_SetKeyValueObject::PostLoad()
+{
+	Super::PostLoad();
+
+	if (BaseClass && BaseClass != UObject::StaticClass())
+	{
+		BlackboardKey.AllowedTypes.Empty();
+		BlackboardKey.AddObjectFilter(this, GET_MEMBER_NAME_CHECKED(UBTTask_SetKeyValueObject, BlackboardKey), BaseClass);
+		Value.SetBaseClass(BaseClass);
+	}
+}
+
+void UBTTask_SetKeyValueStruct::PostLoad()
+{
+	Super::PostLoad();
+
+	if (StructType)
+	{
+		BlackboardKey.AllowedTypes.Empty();
+		BlackboardKey.AddStructFilter(this, GET_MEMBER_NAME_CHECKED(UBTTask_SetKeyValueStruct, BlackboardKey), StructType);
+		BlackboardKey.AllowNoneAsValue(false);
+		Value.SetStructType(StructType);
+	}
+}
+
+#if WITH_EDITOR
+void UBTTask_SetKeyValueClass::PostEditChangeProperty(FPropertyChangedEvent& PropertyChangedEvent)
+{
+	if (PropertyChangedEvent.GetPropertyName() == GET_MEMBER_NAME_CHECKED(UBTTask_SetKeyValueClass, BaseClass))
+	{
+		BlackboardKey.AllowedTypes.Empty();
+		BlackboardKey.AddClassFilter(this, GET_MEMBER_NAME_CHECKED(UBTTask_SetKeyValueClass, BlackboardKey), BaseClass);
+		BlackboardKey.SelectedKeyName = NAME_None;
+		BlackboardKey.InvalidateResolvedKey();
+		if (const UBlackboardData* Blackboard = GetBlackboardAsset())
+		{
+			BlackboardKey.ResolveSelectedKey(*Blackboard);
+		}
+		Value.SetBaseClass(BaseClass);
+	}
+	Super::PostEditChangeProperty(PropertyChangedEvent);
+}
+
+void UBTTask_SetKeyValueEnum::PostEditChangeProperty(FPropertyChangedEvent& PropertyChangedEvent)
+{
+	if (PropertyChangedEvent.GetPropertyName() == GET_MEMBER_NAME_CHECKED(UBTTask_SetKeyValueEnum, EnumType))
+	{
+		BlackboardKey.AllowedTypes.Empty();
+		BlackboardKey.AddEnumFilter(this, GET_MEMBER_NAME_CHECKED(UBTTask_SetKeyValueEnum, BlackboardKey), EnumType);
+		BlackboardKey.SelectedKeyName = NAME_None;
+		BlackboardKey.InvalidateResolvedKey();
+		BlackboardKey.AllowNoneAsValue(EnumType != nullptr);
+		if (const UBlackboardData* Blackboard = GetBlackboardAsset())
+		{
+			BlackboardKey.ResolveSelectedKey(*Blackboard);
+		}
+		Value.SetEnumType(EnumType);
+
+
+	}
+	Super::PostEditChangeProperty(PropertyChangedEvent);
+}
+
+void UBTTask_SetKeyValueObject::PostEditChangeProperty(FPropertyChangedEvent& PropertyChangedEvent)
+{
+	if (PropertyChangedEvent.GetPropertyName() == GET_MEMBER_NAME_CHECKED(UBTTask_SetKeyValueObject, BaseClass))
+	{
+		BlackboardKey.AllowedTypes.Empty();
+		BlackboardKey.AddObjectFilter(this, GET_MEMBER_NAME_CHECKED(UBTTask_SetKeyValueObject, BlackboardKey), BaseClass);
+		BlackboardKey.SelectedKeyName = NAME_None;
+		BlackboardKey.InvalidateResolvedKey();
+		if (const UBlackboardData* Blackboard = GetBlackboardAsset())
+		{
+			BlackboardKey.ResolveSelectedKey(*Blackboard);
+		}
+		Value.SetBaseClass(BaseClass);
+	}
+	Super::PostEditChangeProperty(PropertyChangedEvent);
+}
+
+void UBTTask_SetKeyValueStruct::PostEditChangeProperty(FPropertyChangedEvent& PropertyChangedEvent)
+{
+	if (PropertyChangedEvent.GetPropertyName() == GET_MEMBER_NAME_CHECKED(UBTTask_SetKeyValueStruct, StructType))
+	{
+		BlackboardKey.AllowedTypes.Empty();
+		BlackboardKey.AddStructFilter(this, GET_MEMBER_NAME_CHECKED(UBTTask_SetKeyValueStruct, BlackboardKey), StructType);
+		BlackboardKey.SelectedKeyName = NAME_None;
+		BlackboardKey.InvalidateResolvedKey();
+		BlackboardKey.AllowNoneAsValue(StructType != nullptr);
+		if (const UBlackboardData* Blackboard = GetBlackboardAsset())
+		{
+			BlackboardKey.ResolveSelectedKey(*Blackboard);
+		}
+		Value.SetStructType(StructType);
+	}
+	Super::PostEditChangeProperty(PropertyChangedEvent);
+}
+#endif // WITH_EDITOR
diff --git a/Engine/Source/Runtime/AIModule/Private/BehaviorTree/ValueOrBBKey.cpp b/Engine/Source/Runtime/AIModule/Private/BehaviorTree/ValueOrBBKey.cpp
--- a/Engine/Source/Runtime/AIModule/Private/BehaviorTree/ValueOrBBKey.cpp
+++ b/Engine/Source/Runtime/AIModule/Private/BehaviorTree/ValueOrBBKey.cpp
@@ -525,6 +525,54 @@ bool FValueOrBBKey_Vector::SerializeFromMismatchedTag(const FPropertyTag& Tag, F
 	return false;
 }
 
+void FValueOrBBKey_Class::SetBaseClass(UClass* NewBaseClass)
+{
+	if (NewBaseClass != BaseClass)
+	{
+		if (!BaseClass || !BaseClass->IsChildOf(NewBaseClass))
+		{
+			Key = NAME_None;
+			DefaultValue = nullptr;
+		}
+		BaseClass = NewBaseClass;
+	}
+}
+
+void FValueOrBBKey_Enum::SetEnumType(UEnum* NewEnumType)
+{
+	if (EnumType != NewEnumType)
+	{
+		Key = NAME_None;
+		EnumType = NewEnumType;
+		if (EnumType)
+		{
+			DefaultValue = static_cast<uint8>(EnumType->GetValueByIndex(0));
+		}
+	}
+}
+
+void FValueOrBBKey_Object::SetBaseClass(UClass* NewBaseClass)
+{
+	if (NewBaseClass != BaseClass)
+	{
+		if (!BaseClass || !BaseClass->IsChildOf(NewBaseClass))
+		{
+			Key = NAME_None;
+			DefaultValue = nullptr;
+		}
+		BaseClass = NewBaseClass;
+	}
+}
+
+void FValueOrBBKey_Struct::SetStructType(UScriptStruct* NewStructType)
+{
+	if (NewStructType != DefaultValue.GetScriptStruct())
+	{
+		DefaultValue.InitializeAs(NewStructType, nullptr);
+		Key = NAME_None;
+	}
+}
+
 #if WITH_EDITOR
 bool FValueOrBBKey_Bool::IsCompatibleType(const UBlackboardKeyType* KeyType) const
 {
