[AIDataProvider] prevent provider from being replaced by an incompatible one when using copy/paste

diff --git a/Engine/Source/Editor/DetailCustomizations/Private/AIDataProviderValueDetails.cpp b/Engine/Source/Editor/DetailCustomizations/Private/AIDataProviderValueDetails.cpp
--- a/Engine/Source/Editor/DetailCustomizations/Private/AIDataProviderValueDetails.cpp
+++ b/Engine/Source/Editor/DetailCustomizations/Private/AIDataProviderValueDetails.cpp
@@ -114,10 +114,24 @@ void FAIDataProviderValueDetails::OnBindingChanged()
 
 		FName NameValue;
 		DataFieldProperty->GetValue(NameValue);
+
+		// If current field name is not in the list of matching names then default on the first match
 		if (MatchingProperties.Num() && !MatchingProperties.Contains(NameValue))
 		{
+			ensureAlwaysMsgf(false, TEXT("We expect UEnvQueryNode::PostEditChangeProperty to prevent "
+				"copying an incompatible provider value on top of an existing one before calling NotifyAssetUpdate and refreshing the details."));
 			OnDataFieldNameChange(0);
 		}
+
+		// If current field name is set but there are no matching names then reset the binding to the default since it's invalid.
+		// This could happen when copy /pasting FAIDataProviderValue based properties of different types.
+		else if (MatchingProperties.IsEmpty() && !NameValue.IsNone())
+		{
+			ensureAlwaysMsgf(false, TEXT("We expect UEnvQueryNode::PostEditChangeProperty to prevent "
+				"copying an incompatible provider value on top of an existing one before calling NotifyAssetUpdate and refreshing the details."));
+			DataBindingProperty->ResetToDefault();
+			DefaultValueProperty->ResetToDefault();
+		}
 	}
 }
 
diff --git a/Engine/Source/Runtime/AIModule/Private/DataProviders/AIDataProvider.cpp b/Engine/Source/Runtime/AIModule/Private/DataProviders/AIDataProvider.cpp
--- a/Engine/Source/Runtime/AIModule/Private/DataProviders/AIDataProvider.cpp
+++ b/Engine/Source/Runtime/AIModule/Private/DataProviders/AIDataProvider.cpp
@@ -36,8 +36,17 @@ void FAIDataProviderValue::BindData(const UObject* Owner, int32 RequestId) const
 {
 	if (DataBinding && ensure(Owner))
 	{
+		CachedProperty = nullptr;
 		DataBinding->BindData(*Owner, RequestId);
-		CachedProperty = DataBinding->GetClass()->FindPropertyByName(DataField);
+		if (!DataField.IsNone())
+		{
+			FProperty* Prop = DataBinding->GetClass()->FindPropertyByName(DataField);
+			if (ensureAlwaysMsgf(Prop != nullptr && IsMatchingType(Prop)
+				, TEXT("We expect UEnvQueryNode::PostEditChangeProperty to prevent copying an incompatible provider value on top of an existing one.")))
+			{
+				CachedProperty = Prop;
+			}
+		}
 	}
 }
 
diff --git a/Engine/Source/Runtime/AIModule/Private/EnvironmentQuery/EnvQueryNode.cpp b/Engine/Source/Runtime/AIModule/Private/EnvironmentQuery/EnvQueryNode.cpp
--- a/Engine/Source/Runtime/AIModule/Private/EnvironmentQuery/EnvQueryNode.cpp
+++ b/Engine/Source/Runtime/AIModule/Private/EnvironmentQuery/EnvQueryNode.cpp
@@ -34,12 +34,33 @@ void UEnvQueryNode::PostEditChangeProperty(FPropertyChangedEvent& PropertyChange
 	
 	Super::PostEditChangeProperty(PropertyChangedEvent);
 
-	const FName PropertyName = PropertyChangedEvent.Property ? PropertyChangedEvent.Property->GetFName() : FName();
-	if (PropertyName == NAME_DataBinding)
+	// Validate that DataField is still a valid value for DataBinding
+	if (const FStructProperty* TestStruct = CastField<FStructProperty>(PropertyChangedEvent.MemberProperty))
 	{
+		if (FAIDataProviderValue* AIDataProviderValue = TestStruct->ContainerPtrToValuePtr<FAIDataProviderValue>(this))
+		{
+			TArray<FName> MatchingProperties;
+			AIDataProviderValue->GetMatchingProperties(MatchingProperties);
+
+			const FName NameValue = AIDataProviderValue->DataField;
+
+			// If current field name is not in the list of matching names then default on the first match
+			if (MatchingProperties.Num() && !MatchingProperties.Contains(NameValue))
+			{
+				AIDataProviderValue->DataField = MatchingProperties[0];
+			}
+			// If current field name is set but there are no matching names then reset the binding to the default since it's invalid.
+			// This could happen when copy /pasting FAIDataProviderValue based properties of different types.
+			else if (MatchingProperties.IsEmpty() && !NameValue.IsNone())
+			{
+				AIDataProviderValue->DataField = FName();
+				AIDataProviderValue->DataBinding = nullptr;
+			}
+		}
+
 		// populate ParamName value
-		FStructProperty* TestStruct = CastField<FStructProperty>(PropertyChangedEvent.MemberProperty);
-		if (TestStruct != nullptr)
+		const FName PropertyName = PropertyChangedEvent.Property ? PropertyChangedEvent.Property->GetFName() : FName();
+		if (PropertyName == NAME_DataBinding)
 		{
 			const FString TypeDesc = TestStruct->GetCPPType(nullptr, CPPF_None);
 
@@ -48,9 +69,8 @@ void UEnvQueryNode::PostEditChangeProperty(FPropertyChangedEvent& PropertyChange
 			if (TypeDesc.Contains(GET_STRUCT_NAME_CHECKED(FAIDataProviderIntValue))
 				|| TypeDesc.Contains(GET_STRUCT_NAME_CHECKED(FAIDataProviderFloatValue))
 				|| TypeDesc.Contains(GET_STRUCT_NAME_CHECKED(FAIDataProviderBoolValue)))
-			{				
-				const FAIDataProviderTypedValue* PropertyValue = TestStruct->ContainerPtrToValuePtr<FAIDataProviderTypedValue>(this);
-				if (PropertyValue)
+			{
+				if (const FAIDataProviderTypedValue* PropertyValue = TestStruct->ContainerPtrToValuePtr<FAIDataProviderTypedValue>(this))
 				{
 					UAIDataProvider_QueryParams* QueryParamProvider = Cast<UAIDataProvider_QueryParams>(PropertyValue->DataBinding);
 					if (QueryParamProvider && QueryParamProvider->ParamName.IsNone())
@@ -73,7 +93,7 @@ void UEnvQueryNode::PostEditChangeProperty(FPropertyChangedEvent& PropertyChange
 	}
 
 #if USE_EQS_DEBUGGER
-	UEnvQueryManager::NotifyAssetUpdate(NULL);
+	UEnvQueryManager::NotifyAssetUpdate(nullptr);
 #endif // USE_EQS_DEBUGGER
 }
 #endif //WITH_EDITOR && USE_EQS_DEBUGGER
