Add save time validation to BlackboardKeySelector and ValueOrBBKey_X to warn if their binded key no longer exist.

diff --git a/Engine/Source/Runtime/AIModule/Classes/BehaviorTree/BTNode.h b/Engine/Source/Runtime/AIModule/Classes/BehaviorTree/BTNode.h
--- a/Engine/Source/Runtime/AIModule/Classes/BehaviorTree/BTNode.h
+++ b/Engine/Source/Runtime/AIModule/Classes/BehaviorTree/BTNode.h
@@ -169,6 +169,8 @@ class UBTNode : public UObject, public IGameplayTaskOwnerInterface
 
 	/** Called after creating new node in behavior tree editor, use for versioning */
 	virtual void OnNodeCreated() {}
+
+	AIMODULE_API void PreSave(FObjectPreSaveContext SaveContext) override;
 #endif
 
 	/** Gets called only for instanced nodes(bCreateNodeInstance == true). In practive overridden by BP-implemented BT nodes */
diff --git a/Engine/Source/Runtime/AIModule/Classes/BehaviorTree/BehaviorTreeTypes.h b/Engine/Source/Runtime/AIModule/Classes/BehaviorTree/BehaviorTreeTypes.h
--- a/Engine/Source/Runtime/AIModule/Classes/BehaviorTree/BehaviorTreeTypes.h
+++ b/Engine/Source/Runtime/AIModule/Classes/BehaviorTree/BehaviorTreeTypes.h
@@ -678,6 +678,10 @@ public:
 	FORCEINLINE bool NeedsResolving() const { return GetSelectedKeyID() == FBlackboard::InvalidKey && SelectedKeyName.IsNone() == false; }
 	FORCEINLINE void InvalidateResolvedKey() { SelectedKeyID = FBlackboard::InvalidKey; }
 
+#if WITH_EDITOR
+	void PreSave(const UObject* Outer, const UBlackboardData& Blackboard, FName PropertyName);
+#endif // WITH_EDITOR
+
 	friend FBlackboardDecoratorDetails;
 };
 
diff --git a/Engine/Source/Runtime/AIModule/Classes/BehaviorTree/ValueOrBBKey.h b/Engine/Source/Runtime/AIModule/Classes/BehaviorTree/ValueOrBBKey.h
--- a/Engine/Source/Runtime/AIModule/Classes/BehaviorTree/ValueOrBBKey.h
+++ b/Engine/Source/Runtime/AIModule/Classes/BehaviorTree/ValueOrBBKey.h
@@ -89,6 +89,7 @@ struct FValueOrBlackboardKeyBase
 	AIMODULE_API virtual ~FValueOrBlackboardKeyBase() = default;
 	AIMODULE_API virtual bool IsCompatibleType(const UBlackboardKeyType* KeyType) const { return false; };
 	AIMODULE_API virtual FString ToString() const { return FString(); }
+	AIMODULE_API void PreSave(const UObject* Outer, const UBlackboardData& Blackboard, FName PropertyName);
 #endif // WITH_EDITOR
 
 	const FName& GetKey() const { return Key; }
diff --git a/Engine/Source/Runtime/AIModule/Private/BehaviorTree/BTNode.cpp b/Engine/Source/Runtime/AIModule/Private/BehaviorTree/BTNode.cpp
--- a/Engine/Source/Runtime/AIModule/Private/BehaviorTree/BTNode.cpp
+++ b/Engine/Source/Runtime/AIModule/Private/BehaviorTree/BTNode.cpp
@@ -1,13 +1,15 @@
 // Copyright Epic Games, Inc. All Rights Reserved.
 
 #include "BehaviorTree/BTNode.h"
-#include "Engine/World.h"
 #include "AIController.h"
-#include "BehaviorTree/BTCompositeNode.h"
 #include "BehaviorTree/BehaviorTree.h"
+#include "BehaviorTree/BTCompositeNode.h"
+#include "BehaviorTree/ValueOrBBKey.h"
+#include "Engine/World.h"
 #include "GameplayTasksComponent.h"
-#include "UObject/Package.h"
 #include "Styling/AppStyle.h"
+#include "UObject/ObjectSaveContext.h"
+#include "UObject/Package.h"
 
 #include UE_INLINE_GENERATED_CPP_BY_NAME(BTNode)
 
@@ -226,6 +228,32 @@ FString UBTNode::GetErrorMessage() const
 	return FString();
 }
 
+void UBTNode::PreSave(FObjectPreSaveContext SaveContext)
+{
+	if (const UBlackboardData* Blackboard = GetBlackboardAsset())
+	{
+		for (TFieldIterator<FStructProperty> StructPropIt(GetClass()); StructPropIt; ++StructPropIt)
+		{
+			if (StructPropIt->Struct && StructPropIt->Struct->IsChildOf(FValueOrBlackboardKeyBase::StaticStruct()))
+			{
+				if (FValueOrBlackboardKeyBase* ValueOrBBKey = StructPropIt->ContainerPtrToValuePtr<FValueOrBlackboardKeyBase>(this))
+				{
+					ValueOrBBKey->PreSave(this, *Blackboard, StructPropIt->GetFName());
+				}
+			}
+			else if (StructPropIt->Struct && StructPropIt->Struct->IsChildOf(FBlackboardKeySelector::StaticStruct()))
+			{
+				if (FBlackboardKeySelector* KeySelector = StructPropIt->ContainerPtrToValuePtr<FBlackboardKeySelector>(this))
+				{
+					KeySelector->PreSave(this, *Blackboard, StructPropIt->GetFName());
+				}
+			}
+		}
+	}
+
+	Super::PreSave(SaveContext);
+}
+
 #endif
 
 UGameplayTasksComponent* UBTNode::GetGameplayTasksComponent(const UGameplayTask& Task) const
diff --git a/Engine/Source/Runtime/AIModule/Private/BehaviorTree/BehaviorTreeTypes.cpp b/Engine/Source/Runtime/AIModule/Private/BehaviorTree/BehaviorTreeTypes.cpp
--- a/Engine/Source/Runtime/AIModule/Private/BehaviorTree/BehaviorTreeTypes.cpp
+++ b/Engine/Source/Runtime/AIModule/Private/BehaviorTree/BehaviorTreeTypes.cpp
@@ -673,6 +673,55 @@ void FBlackboardKeySelector::AddStructFilter(UObject* Owner, FName PropertyName,
 	AllowedTypes.Add(FilterOb);
 }
 
+#if WITH_EDITOR
+void FBlackboardKeySelector::PreSave(const UObject* Outer, const UBlackboardData& Blackboard, FName PropertyName)
+{
+	if(!SelectedKeyName.IsNone())
+	{
+		const TCHAR* ResetValueMessage = bNoneIsAllowedValue ? TEXT("Resetting selected key to none.") : TEXT("Setting the key to first matching key.");
+		const FBlackboard::FKey Id = Blackboard.GetKeyID(SelectedKeyName);
+		if (Id != FBlackboard::InvalidKey)
+		{
+			if (const FBlackboardEntry* BlackboardEntry = Blackboard.GetKey(Id))
+			{
+				if(AllowedTypes.Num() > 0)
+				{
+					for (UBlackboardKeyType* AllowedType : AllowedTypes)
+					{
+						if (BlackboardEntry->KeyType->IsAllowedByFilter(AllowedType))
+						{
+							return;
+						}
+					}
+				}
+				else
+				{
+					return;
+				}
+				UE_LOG(LogBehaviorTree, Warning, TEXT("%s in node %s is bound to key %s in blackboard %s but the key type isn't allowed. %s"), *PropertyName.ToString(), *GetNameSafe(Outer), *SelectedKeyName.ToString(), *Blackboard.GetFullName(), ResetValueMessage);
+			}
+			else
+			{
+				UE_LOG(LogBehaviorTree, Warning, TEXT("%s in node %s is bound to key %s in blackboard %s but the key doesn't exist. %s"), *PropertyName.ToString(), *GetNameSafe(Outer), *SelectedKeyName.ToString(), *Blackboard.GetFullName(), ResetValueMessage);
+			}
+		}
+		else
+		{
+			UE_LOG(LogBehaviorTree, Warning, TEXT("%s in node %s is bound to key %s in blackboard %s but the key doesn't exist. %s"), *PropertyName.ToString(), *GetNameSafe(Outer), *SelectedKeyName.ToString(), *Blackboard.GetFullName(), ResetValueMessage);
+		}
+
+		if (bNoneIsAllowedValue)
+		{
+			InitSelection(Blackboard);
+		}
+		else
+		{
+			SelectedKeyName = NAME_None;
+		}
+	}
+}
+#endif // WITH_EDITOR
+
 //----------------------------------------------------------------------//
 // UBehaviorTreeTypes
 //----------------------------------------------------------------------//
diff --git a/Engine/Source/Runtime/AIModule/Private/BehaviorTree/ValueOrBBKey.cpp b/Engine/Source/Runtime/AIModule/Private/BehaviorTree/ValueOrBBKey.cpp
--- a/Engine/Source/Runtime/AIModule/Private/BehaviorTree/ValueOrBBKey.cpp
+++ b/Engine/Source/Runtime/AIModule/Private/BehaviorTree/ValueOrBBKey.cpp
@@ -13,6 +13,8 @@
 #include "BehaviorTree/Blackboard/BlackboardKeyType_String.h"
 #include "BehaviorTree/Blackboard/BlackboardKeyType_Struct.h"
 #include "BehaviorTree/Blackboard/BlackboardKeyType_Vector.h"
+#include "BehaviorTree/BlackboardData.h"
+#include "BehaviorTree/BTNode.h" // For LogBehaviorTree
 
 namespace FBlackboard
 {
@@ -574,6 +576,37 @@ void FValueOrBBKey_Struct::SetStructType(UScriptStruct* NewStructType)
 }
 
 #if WITH_EDITOR
+void FValueOrBlackboardKeyBase::PreSave(const UObject* Outer, const UBlackboardData& Blackboard, FName PropertyName)
+{
+	if (!Key.IsNone())
+	{
+		const FBlackboard::FKey Id = Blackboard.GetKeyID(Key);
+		if (Id != FBlackboard::InvalidKey)
+		{
+			if (const FBlackboardEntry* BlackboardEntry = Blackboard.GetKey(Id))
+			{
+				if (IsCompatibleType(BlackboardEntry->KeyType))
+				{
+					return;
+				}
+				else
+				{
+					UE_LOG(LogBehaviorTree, Warning, TEXT("%s in node %s is bound to key %s but the type doesn't match in blackboard %s. Resetting the key to none."), *PropertyName.ToString(), *GetNameSafe(Outer), *Key.ToString(), *Blackboard.GetFullName());
+				}
+			}
+			else
+			{
+				UE_LOG(LogBehaviorTree, Warning, TEXT("%s in node %s is bound to key %s but the key doesn't exist in blackboard %s. Resetting the key to none."), *PropertyName.ToString(), *GetNameSafe(Outer), *Key.ToString(), *Blackboard.GetFullName());
+			}
+		}
+		else
+		{
+			UE_LOG(LogBehaviorTree, Warning, TEXT("%s in node %s is bound to key %s but the key doesn't exist in blackboard %s. Resetting the key to none."), *PropertyName.ToString(), *GetNameSafe(Outer), *Key.ToString(), *Blackboard.GetFullName());
+		}
+		Key = NAME_None;
+	}
+}
+
 bool FValueOrBBKey_Bool::IsCompatibleType(const UBlackboardKeyType* KeyType) const
 {
 	return KeyType && KeyType->GetClass() == UBlackboardKeyType_Bool::StaticClass();
