Add bt name in warning on presave for KeySeletor and ValueOrBBkey

diff --git a/Engine/Source/Runtime/AIModule/Private/BehaviorTree/BehaviorTreeTypes.cpp b/Engine/Source/Runtime/AIModule/Private/BehaviorTree/BehaviorTreeTypes.cpp
--- a/Engine/Source/Runtime/AIModule/Private/BehaviorTree/BehaviorTreeTypes.cpp
+++ b/Engine/Source/Runtime/AIModule/Private/BehaviorTree/BehaviorTreeTypes.cpp
@@ -698,16 +698,16 @@ void FBlackboardKeySelector::PreSave(const UObject* Outer, const UBlackboardData
 				{
 					return;
 				}
-				UE_LOG(LogBehaviorTree, Warning, TEXT("%s in node %s is bound to key %s in blackboard %s but the key type isn't allowed. %s"), *PropertyName.ToString(), *GetNameSafe(Outer), *SelectedKeyName.ToString(), *Blackboard.GetFullName(), ResetValueMessage);
+				UE_LOG(LogBehaviorTree, Warning, TEXT("%s in node %s in bt %s is bound to key %s in blackboard %s but the key type isn't allowed. %s"), *PropertyName.ToString(), *GetNameSafe(Outer), Outer ? *GetNameSafe(Outer->GetTypedOuter<UBehaviorTree>()) : TEXT("None"), *SelectedKeyName.ToString(), *Blackboard.GetFullName(), ResetValueMessage);
 			}
 			else
 			{
-				UE_LOG(LogBehaviorTree, Warning, TEXT("%s in node %s is bound to key %s in blackboard %s but the key doesn't exist. %s"), *PropertyName.ToString(), *GetNameSafe(Outer), *SelectedKeyName.ToString(), *Blackboard.GetFullName(), ResetValueMessage);
+				UE_LOG(LogBehaviorTree, Warning, TEXT("%s in node %s in bt %s is bound to key %s in blackboard %s but the key doesn't exist. %s"), *PropertyName.ToString(), *GetNameSafe(Outer), Outer ? *GetNameSafe(Outer->GetTypedOuter<UBehaviorTree>()) : TEXT("None"), *SelectedKeyName.ToString(), *Blackboard.GetFullName(), ResetValueMessage);
 			}
 		}
 		else
 		{
-			UE_LOG(LogBehaviorTree, Warning, TEXT("%s in node %s is bound to key %s in blackboard %s but the key doesn't exist. %s"), *PropertyName.ToString(), *GetNameSafe(Outer), *SelectedKeyName.ToString(), *Blackboard.GetFullName(), ResetValueMessage);
+			UE_LOG(LogBehaviorTree, Warning, TEXT("%s in node %s in bt %s is bound to key %s in blackboard %s but the key doesn't exist. %s"), *PropertyName.ToString(), *GetNameSafe(Outer), Outer ? *GetNameSafe(Outer->GetTypedOuter<UBehaviorTree>()) : TEXT("None"), *SelectedKeyName.ToString(), *Blackboard.GetFullName(), ResetValueMessage);
 		}
 
 		if (bNoneIsAllowedValue)
diff --git a/Engine/Source/Runtime/AIModule/Private/BehaviorTree/ValueOrBBKey.cpp b/Engine/Source/Runtime/AIModule/Private/BehaviorTree/ValueOrBBKey.cpp
--- a/Engine/Source/Runtime/AIModule/Private/BehaviorTree/ValueOrBBKey.cpp
+++ b/Engine/Source/Runtime/AIModule/Private/BehaviorTree/ValueOrBBKey.cpp
@@ -2,6 +2,7 @@
 
 #include "BehaviorTree/ValueOrBBKey.h"
 
+#include "BehaviorTree/BehaviorTree.h"
 #include "BehaviorTree/Blackboard/BlackboardKeyType_Bool.h"
 #include "BehaviorTree/Blackboard/BlackboardKeyType_Class.h"
 #include "BehaviorTree/Blackboard/BlackboardKeyType_Enum.h"
@@ -591,17 +592,17 @@ void FValueOrBlackboardKeyBase::PreSave(const UObject* Outer, const UBlackboardD
 				}
 				else
 				{
-					UE_LOG(LogBehaviorTree, Warning, TEXT("%s in node %s is bound to key %s but the type doesn't match in blackboard %s. Resetting the key to none."), *PropertyName.ToString(), *GetNameSafe(Outer), *Key.ToString(), *Blackboard.GetFullName());
+					UE_LOG(LogBehaviorTree, Warning, TEXT("%s in node %s in bt %s is bound to key %s but the type doesn't match in blackboard %s. Resetting the key to none."), *PropertyName.ToString(), *GetNameSafe(Outer), Outer ? *GetNameSafe(Outer->GetTypedOuter<UBehaviorTree>()) : TEXT("None"), *Key.ToString(), *Blackboard.GetFullName());
 				}
 			}
 			else
 			{
-				UE_LOG(LogBehaviorTree, Warning, TEXT("%s in node %s is bound to key %s but the key doesn't exist in blackboard %s. Resetting the key to none."), *PropertyName.ToString(), *GetNameSafe(Outer), *Key.ToString(), *Blackboard.GetFullName());
+				UE_LOG(LogBehaviorTree, Warning, TEXT("%s in node %s in bt %s is bound to key %s but the key doesn't exist in blackboard %s. Resetting the key to none."), *PropertyName.ToString(), *GetNameSafe(Outer), Outer ? *GetNameSafe(Outer->GetTypedOuter<UBehaviorTree>()) : TEXT("None"), *Key.ToString(), *Blackboard.GetFullName());
 			}
 		}
 		else
 		{
-			UE_LOG(LogBehaviorTree, Warning, TEXT("%s in node %s is bound to key %s but the key doesn't exist in blackboard %s. Resetting the key to none."), *PropertyName.ToString(), *GetNameSafe(Outer), *Key.ToString(), *Blackboard.GetFullName());
+			UE_LOG(LogBehaviorTree, Warning, TEXT("%s in node %s in bt %s is bound to key %s but the key doesn't exist in blackboard %s. Resetting the key to none."), *PropertyName.ToString(), *GetNameSafe(Outer), Outer ? *GetNameSafe(Outer->GetTypedOuter<UBehaviorTree>()) : TEXT("None"), *Key.ToString(), *Blackboard.GetFullName());
 		}
 		Key = NAME_None;
 	}
