[AIDataProvider] fixed PostEditChange to only validate providers when the struct is of the right type

diff --git a/Engine/Source/Runtime/AIModule/Private/EnvironmentQuery/EnvQueryNode.cpp b/Engine/Source/Runtime/AIModule/Private/EnvironmentQuery/EnvQueryNode.cpp
--- a/Engine/Source/Runtime/AIModule/Private/EnvironmentQuery/EnvQueryNode.cpp
+++ b/Engine/Source/Runtime/AIModule/Private/EnvironmentQuery/EnvQueryNode.cpp
@@ -34,27 +34,31 @@ void UEnvQueryNode::PostEditChangeProperty(FPropertyChangedEvent& PropertyChange
 	
 	Super::PostEditChangeProperty(PropertyChangedEvent);
 
-	// Validate that DataField is still a valid value for DataBinding
 	if (const FStructProperty* TestStruct = CastField<FStructProperty>(PropertyChangedEvent.MemberProperty))
 	{
-		if (FAIDataProviderValue* AIDataProviderValue = TestStruct->ContainerPtrToValuePtr<FAIDataProviderValue>(this))
+		// Validate that DataField is still a valid value for DataBinding
+		if (TestStruct->Struct != nullptr
+			&& TestStruct->Struct->IsChildOf(FAIDataProviderValue::StaticStruct()))
 		{
-			TArray<FName> MatchingProperties;
-			AIDataProviderValue->GetMatchingProperties(MatchingProperties);
+			if (FAIDataProviderValue* AIDataProviderValue = TestStruct->ContainerPtrToValuePtr<FAIDataProviderValue>(this))
+			{
+				TArray<FName> MatchingProperties;
+				AIDataProviderValue->GetMatchingProperties(MatchingProperties);
 
-			const FName NameValue = AIDataProviderValue->DataField;
+				const FName NameValue = AIDataProviderValue->DataField;
 
-			// If current field name is not in the list of matching names then default on the first match
-			if (MatchingProperties.Num() && !MatchingProperties.Contains(NameValue))
-			{
-				AIDataProviderValue->DataField = MatchingProperties[0];
-			}
-			// If current field name is set but there are no matching names then reset the binding to the default since it's invalid.
-			// This could happen when copy /pasting FAIDataProviderValue based properties of different types.
-			else if (MatchingProperties.IsEmpty() && !NameValue.IsNone())
-			{
-				AIDataProviderValue->DataField = FName();
-				AIDataProviderValue->DataBinding = nullptr;
+				// If current field name is not in the list of matching names then default on the first match
+				if (MatchingProperties.Num() && !MatchingProperties.Contains(NameValue))
+				{
+					AIDataProviderValue->DataField = MatchingProperties[0];
+				}
+				// If current field name is set but there are no matching names then reset the binding to the default since it's invalid.
+				// This could happen when copy /pasting FAIDataProviderValue based properties of different types.
+				else if (MatchingProperties.IsEmpty() && !NameValue.IsNone())
+				{
+					AIDataProviderValue->DataField = FName();
+					AIDataProviderValue->DataBinding = nullptr;
+				}
 			}
 		}
 
