PR #12436: Synchronize namespace reference in comments in CrowdManager.h to namespace found in CrowdManager.cpp

diff --git a/Engine/Source/Runtime/AIModule/Classes/Navigation/CrowdManager.h b/Engine/Source/Runtime/AIModule/Classes/Navigation/CrowdManager.h
--- a/Engine/Source/Runtime/AIModule/Classes/Navigation/CrowdManager.h
+++ b/Engine/Source/Runtime/AIModule/Classes/Navigation/CrowdManager.h
@@ -65,7 +65,7 @@ typedef TSharedPtr<const FNavigationQueryFilter, ESPMode::ThreadSafe> FSharedCon
  *      CrowdManager->RegisterAgent(this);
  *   }
  *
- *   Check flags in CrowdDebugDrawing namespace (CrowdManager.cpp) for debugging options.
+ *   Check flags in FCrowdDebug namespace (CrowdManager.cpp) for debugging options.
  */
 
 USTRUCT()
