Add run query on selection changed on EQSTesting pawn
Add helper function to EnvQueryItemType_Actor to set editor selected actors inside the context to simplify testing in editor

diff --git a/Engine/Source/Runtime/AIModule/Classes/EnvironmentQuery/EQSTestingPawn.h b/Engine/Source/Runtime/AIModule/Classes/EnvironmentQuery/EQSTestingPawn.h
--- a/Engine/Source/Runtime/AIModule/Classes/EnvironmentQuery/EQSTestingPawn.h
+++ b/Engine/Source/Runtime/AIModule/Classes/EnvironmentQuery/EQSTestingPawn.h
@@ -61,6 +61,9 @@ class AEQSTestingPawn : public ACharacter, public IEQSQueryResultSourceInterface
 	UPROPERTY(Category = EQS, EditAnywhere)
 	uint32 bTickDuringGame : 1;
 
+	UPROPERTY(Category = EQS, EditAnywhere)
+	uint32 bRunQueryOnSelectionChanged : 1;
+
 	UPROPERTY(Category=EQS, EditAnywhere)
 	TEnumAsByte<EEnvQueryRunMode::Type> QueryingMode;
 
diff --git a/Engine/Source/Runtime/AIModule/Classes/EnvironmentQuery/Items/EnvQueryItemType_Actor.h b/Engine/Source/Runtime/AIModule/Classes/EnvironmentQuery/Items/EnvQueryItemType_Actor.h
--- a/Engine/Source/Runtime/AIModule/Classes/EnvironmentQuery/Items/EnvQueryItemType_Actor.h
+++ b/Engine/Source/Runtime/AIModule/Classes/EnvironmentQuery/Items/EnvQueryItemType_Actor.h
@@ -28,6 +28,9 @@ public:
 	static AIMODULE_API void SetContextHelper(FEnvQueryContextData& ContextData, const TArray<const AActor*>& MultipleActors);
 	static AIMODULE_API void SetContextHelper(FEnvQueryContextData& ContextData, const TArray<AActor*>& MultipleActors);
 
+	/** In editor, outside a PIE session, will add all currently selected actors to the context data. */
+	static AIMODULE_API void AddEditorSelectedActors(FEnvQueryContextData& ContextData);
+
 	AIMODULE_API virtual FVector GetItemLocation(const uint8* RawData) const override;
 	AIMODULE_API virtual FRotator GetItemRotation(const uint8* RawData) const override;
 	AIMODULE_API virtual AActor* GetActor(const uint8* RawData) const override;
diff --git a/Engine/Source/Runtime/AIModule/Private/EnvironmentQuery/EQSTestingPawn.cpp b/Engine/Source/Runtime/AIModule/Private/EnvironmentQuery/EQSTestingPawn.cpp
--- a/Engine/Source/Runtime/AIModule/Private/EnvironmentQuery/EQSTestingPawn.cpp
+++ b/Engine/Source/Runtime/AIModule/Private/EnvironmentQuery/EQSTestingPawn.cpp
@@ -24,6 +24,7 @@ extern UNREALED_API UEditorEngine* GEditor;
 
 #include "Engine/Selection.h"
 #include "Components/BillboardComponent.h"
+#include "EngineUtils.h"
 
 #include UE_INLINE_GENERATED_CPP_BY_NAME(EQSTestingPawn)
 
@@ -37,6 +38,7 @@ AEQSTestingPawn::AEQSTestingPawn(const FObjectInitializer& ObjectInitializer)
 	, bDrawLabels(true)
 	, bDrawFailedItems(true)
 	, bReRunQueryOnlyOnFinishedMove(true)
+	, bRunQueryOnSelectionChanged(false)
 	, QueryingMode(EEnvQueryRunMode::AllMatching)
 {
 	static FName CollisionProfileName(TEXT("NoCollision"));
@@ -143,6 +145,18 @@ void AEQSTestingPawn::OnEditorSelectionChanged(UObject* NewSelection)
 			EQSPawn->RunEQSQuery();
 		}
 	}
+
+	if (UWorld* SelectionWorld = NewSelection ? NewSelection->GetWorld() : nullptr)
+	{
+		for (TActorIterator<AEQSTestingPawn> EQSPawnIterator(SelectionWorld); EQSPawnIterator; ++EQSPawnIterator)
+		{
+			if (!SelectedPawns.Contains(*EQSPawnIterator) && EQSPawnIterator->bRunQueryOnSelectionChanged)
+			{
+				EQSPawnIterator->QueryTemplate->CollectQueryParams(**EQSPawnIterator, EQSPawnIterator->QueryConfig);
+				EQSPawnIterator->RunEQSQuery();
+			}
+		}
+	}
 }
 #endif // WITH_EDITOR
 
@@ -322,6 +336,7 @@ void AEQSTestingPawn::OnPropertyChanged(const FName PropName)
 	static const FName NAME_QueryConfig = GET_MEMBER_NAME_CHECKED(AEQSTestingPawn, QueryConfig);
 	static const FName NAME_ShouldBeVisibleInGame = GET_MEMBER_NAME_CHECKED(AEQSTestingPawn, bShouldBeVisibleInGame);
 	static const FName NAME_QueryingMode = GET_MEMBER_NAME_CHECKED(AEQSTestingPawn, QueryingMode);
+	static const FName NAME_bRunQueryOnSelectionChanged = GET_MEMBER_NAME_CHECKED(AEQSTestingPawn, bRunQueryOnSelectionChanged);
 
 	if (PropName == NAME_QueryTemplate || PropName == NAME_QueryConfig)
 	{
@@ -354,6 +369,14 @@ void AEQSTestingPawn::OnPropertyChanged(const FName PropName)
 	{
 		RunEQSQuery();
 	}
+	else if (PropName == NAME_bRunQueryOnSelectionChanged)
+	{
+		if (bRunQueryOnSelectionChanged)
+		{
+			EdRenderComp->bDrawOnlyWhenSelected = false;
+		}
+		UpdateDrawing();
+	}
 }
 
 void AEQSTestingPawn::PostEditChangeProperty( FPropertyChangedEvent& PropertyChangedEvent)
diff --git a/Engine/Source/Runtime/AIModule/Private/EnvironmentQuery/Items/EnvQueryItemType_Actor.cpp b/Engine/Source/Runtime/AIModule/Private/EnvironmentQuery/Items/EnvQueryItemType_Actor.cpp
--- a/Engine/Source/Runtime/AIModule/Private/EnvironmentQuery/Items/EnvQueryItemType_Actor.cpp
+++ b/Engine/Source/Runtime/AIModule/Private/EnvironmentQuery/Items/EnvQueryItemType_Actor.cpp
@@ -6,6 +6,11 @@
 #include "AITypes.h"
 #include "EnvironmentQuery/EnvQueryTypes.h"
 
+#if WITH_EDITOR
+#include "Selection.h"
+#include "Editor.h"
+#endif // WITH_EDITOR
+
 #include UE_INLINE_GENERATED_CPP_BY_NAME(EnvQueryItemType_Actor)
 
 UEnvQueryItemType_Actor::UEnvQueryItemType_Actor(const FObjectInitializer& ObjectInitializer) : Super(ObjectInitializer)
@@ -83,3 +88,22 @@ void UEnvQueryItemType_Actor::SetContextHelper(FEnvQueryContextData& ContextData
 	}
 }
 
+void UEnvQueryItemType_Actor::AddEditorSelectedActors(FEnvQueryContextData& ContextData)
+{
+#if WITH_EDITOR
+	if (GEditor && !GEditor->IsPlayingSessionInEditor())
+	{
+		TArray<AActor*> SelectedActors;
+		SelectedActors.Reserve(GEditor->GetSelectedActorCount());
+		for (FSelectionIterator It = GEditor->GetSelectedActorIterator(); It; ++It)
+		{
+			if (AActor* Actor = Cast<AActor>(*It))
+			{
+				SelectedActors.Add(Actor);
+			}
+		}
+		SetContextHelper(ContextData, SelectedActors);
+	}
+#endif // WITH_EDITOR
+}
+
