Add ForceRebuildProperty on type property in BTTask_SetKeyValueX to fix issue where value picker didn't always refresh withtout reselecting the node.
#rnx

diff --git a/Engine/Source/Runtime/AIModule/Classes/BehaviorTree/Tasks/BTTask_SetKeyValue.h b/Engine/Source/Runtime/AIModule/Classes/BehaviorTree/Tasks/BTTask_SetKeyValue.h
--- a/Engine/Source/Runtime/AIModule/Classes/BehaviorTree/Tasks/BTTask_SetKeyValue.h
+++ b/Engine/Source/Runtime/AIModule/Classes/BehaviorTree/Tasks/BTTask_SetKeyValue.h
@@ -38,7 +38,7 @@ public:
 	virtual void PostEditChangeProperty(FPropertyChangedEvent& PropertyChangedEvent) override;
 #endif // WITH_EDITOR
 private:
-	UPROPERTY(EditAnywhere, Category = Blackboard, meta = (AllowAbstract = "true", NoClear))
+	UPROPERTY(EditAnywhere, Category = Blackboard, meta = (AllowAbstract = "true", NoClear, ForceRebuildProperty="Value"))
 	TSubclassOf<UObject> BaseClass = UObject::StaticClass();
 
 	UPROPERTY(EditAnywhere, Category = Blackboard)
@@ -61,7 +61,7 @@ public:
 	virtual void PostEditChangeProperty(FPropertyChangedEvent& PropertyChangedEvent) override;
 #endif // WITH_EDITOR
 private:
-	UPROPERTY(EditAnywhere, Category = Blackboard, meta = (AllowAbstract = "true", NoClear))
+	UPROPERTY(EditAnywhere, Category = Blackboard, meta = (AllowAbstract = "true", NoClear, ForceRebuildProperty = "Value"))
 	TObjectPtr<UEnum> EnumType;
 
 	UPROPERTY(EditAnywhere, Category = Blackboard)
@@ -148,7 +148,7 @@ public:
 	virtual void PostEditChangeProperty(FPropertyChangedEvent& PropertyChangedEvent) override;
 #endif // WITH_EDITOR
 private:
-	UPROPERTY(EditAnywhere, Category = Blackboard, meta = (AllowAbstract = "true", NoClear))
+	UPROPERTY(EditAnywhere, Category = Blackboard, meta = (AllowAbstract = "true", NoClear, ForceRebuildProperty = "Value"))
 	TSubclassOf<UObject> BaseClass = UObject::StaticClass();
 
 	UPROPERTY(EditAnywhere, Category = Blackboard)
@@ -187,7 +187,7 @@ public:
 	virtual void PostEditChangeProperty(FPropertyChangedEvent& PropertyChangedEvent) override;
 #endif // WITH_EDITOR
 private:
-	UPROPERTY(EditAnywhere, Category = Blackboard)
+	UPROPERTY(EditAnywhere, Category = Blackboard, meta = (ForceRebuildProperty = "Value"))
 	TObjectPtr<UScriptStruct> StructType;
 
 	UPROPERTY(EditAnywhere, Category = Blackboard)
diff --git a/Engine/Source/Runtime/AIModule/Private/BehaviorTree/Tasks/BTTask_SetKeyValue.cpp b/Engine/Source/Runtime/AIModule/Private/BehaviorTree/Tasks/BTTask_SetKeyValue.cpp
--- a/Engine/Source/Runtime/AIModule/Private/BehaviorTree/Tasks/BTTask_SetKeyValue.cpp
+++ b/Engine/Source/Runtime/AIModule/Private/BehaviorTree/Tasks/BTTask_SetKeyValue.cpp
@@ -45,7 +45,7 @@ UBTTask_SetKeyValueClass::UBTTask_SetKeyValueClass(const FObjectInitializer& Obj
 
 UBTTask_SetKeyValueEnum::UBTTask_SetKeyValueEnum(const FObjectInitializer& ObjectInitializer /*= FObjectInitializer::Get()*/)
 {
-	NodeName = "Set Class Key";
+	NodeName = "Set Enum Key";
 	// Since there are no good default value for enum don't add a filter and allow none as a value so the key isn't bound to a random key.
 	INIT_TASK_NODE_NOTIFY_FLAGS();
 }
