Fixed CastChecked between int16 and int32 types

diff --git a/Engine/Source/Runtime/AIModule/Private/GameplayDebugger/GameplayDebuggerCategory_Navmesh.cpp b/Engine/Source/Runtime/AIModule/Private/GameplayDebugger/GameplayDebuggerCategory_Navmesh.cpp
--- a/Engine/Source/Runtime/AIModule/Private/GameplayDebugger/GameplayDebuggerCategory_Navmesh.cpp
+++ b/Engine/Source/Runtime/AIModule/Private/GameplayDebugger/GameplayDebuggerCategory_Navmesh.cpp
@@ -131,8 +131,8 @@ void FGameplayDebuggerCategory_Navmesh::CollectData(APlayerController* OwnerPC,
 		if (NavSys) 
 		{
 			DataPack.NumDirtyAreas = NavSys->GetNumDirtyAreas();
-			DataPack.NumRunningTasks = IntCastChecked<uint16>(NavSys->GetNumRunningBuildTasks());
-			DataPack.NumRemainingTasks = IntCastChecked<uint16>(NavSys->GetNumRemainingBuildTasks());
+			DataPack.NumRunningTasks = NavSys->GetNumRunningBuildTasks();
+			DataPack.NumRemainingTasks = NavSys->GetNumRemainingBuildTasks();
 			DataPack.bIsNavOctreeLocked = NavSys->IsNavigationOctreeLocked();
 			DataPack.bIsNavBuildLocked = NavSys->IsNavigationBuildingLocked();
 			DataPack.NavBuildLockStatusDesc = FString("Unknown");
diff --git a/Engine/Source/Runtime/AIModule/Public/GameplayDebugger/GameplayDebuggerCategory_Navmesh.h b/Engine/Source/Runtime/AIModule/Public/GameplayDebugger/GameplayDebuggerCategory_Navmesh.h
--- a/Engine/Source/Runtime/AIModule/Public/GameplayDebugger/GameplayDebuggerCategory_Navmesh.h
+++ b/Engine/Source/Runtime/AIModule/Public/GameplayDebugger/GameplayDebuggerCategory_Navmesh.h
@@ -46,8 +46,8 @@ protected:
 		FVector LockedReferenceLocation = FNavigationSystem::InvalidLocation;
 		int32 NumDirtyAreas = 0;
 		int32 NumSuspendedDirtyAreas = 0;
-		uint16 NumRunningTasks = 0;
-		uint16 NumRemainingTasks = 0;
+		int32 NumRunningTasks = 0;
+		int32 NumRemainingTasks = 0;
 		bool bCanChangeReference = false;
 		bool bCanCycleNavigationData = false;
 		bool bIsUsingPlayerActor = false;
