Exposed Get/SetValue methods and InvalidValue on the UBlackboardKeyType_Struct

diff --git a/Engine/Source/Runtime/AIModule/Classes/BehaviorTree/Blackboard/BlackboardKeyType_Struct.h b/Engine/Source/Runtime/AIModule/Classes/BehaviorTree/Blackboard/BlackboardKeyType_Struct.h
--- a/Engine/Source/Runtime/AIModule/Classes/BehaviorTree/Blackboard/BlackboardKeyType_Struct.h
+++ b/Engine/Source/Runtime/AIModule/Classes/BehaviorTree/Blackboard/BlackboardKeyType_Struct.h
@@ -18,12 +18,12 @@ class UBlackboardKeyType_Struct : public UBlackboardKeyType
 
 public:
 	using FDataType = FConstStructView;
-	static const FDataType InvalidValue;
+	static AIMODULE_API const FDataType InvalidValue;
 
 	UBlackboardKeyType_Struct(const FObjectInitializer& ObjectInitializer);
 
-	static FConstStructView GetValue(UBlackboardKeyType_Struct* KeyOb, const uint8* RawData);
-	static bool SetValue(UBlackboardKeyType_Struct* KeyOb, uint8* RawData, FConstStructView Value);
+	static AIMODULE_API FConstStructView GetValue(UBlackboardKeyType_Struct* KeyOb, const uint8* RawData);
+	static AIMODULE_API bool SetValue(UBlackboardKeyType_Struct* KeyOb, uint8* RawData, FConstStructView Value);
 
 	virtual EBlackboardCompare::Type CompareValues(const UBlackboardComponent& OwnerComp, const uint8* MemoryBlock, const UBlackboardKeyType* OtherKeyOb, const uint8* OtherMemoryBlock) const override;
 	virtual void CopyValues(UBlackboardComponent& OwnerComp, uint8* MemoryBlock, const UBlackboardKeyType* SourceKeyOb, const uint8* SourceBlock) override;
