In CreateCrowdManager, use warning instead of making an error if no navdata is found.
This is because OnNavMeshUpdate() is called on UnregisterNavData and it might be the last one. We want to avoid errors on closing maps.

diff --git a/Engine/Source/Runtime/AIModule/Private/Navigation/CrowdManager.cpp b/Engine/Source/Runtime/AIModule/Private/Navigation/CrowdManager.cpp
--- a/Engine/Source/Runtime/AIModule/Private/Navigation/CrowdManager.cpp
+++ b/Engine/Source/Runtime/AIModule/Private/Navigation/CrowdManager.cpp
@@ -896,7 +896,7 @@ void UCrowdManager::CreateCrowdManager()
 		// if navigation mesh is found
 		UpdateNavData();
 
-		UE_CLOG(MyNavData == nullptr, LogCrowdFollowing, Error, TEXT("Unable to find RecastNavMesh instance while trying to create UCrowdManager instance"));
+		UE_CLOG(MyNavData == nullptr, LogCrowdFollowing, Warning, TEXT("Unable to find RecastNavMesh instance while trying to create UCrowdManager instance"));
 		return;
 	}
 
