Fix for FGraphAStar::FindPath : make sure we fill OutPath when the StartNode and EndNode are identical. This is important to avoid having to handle a special case in the code that call this function and to be consistent with Recast code

diff --git a/Engine/Source/Runtime/AIModule/Public/GraphAStar.h b/Engine/Source/Runtime/AIModule/Public/GraphAStar.h
--- a/Engine/Source/Runtime/AIModule/Public/GraphAStar.h
+++ b/Engine/Source/Runtime/AIModule/Public/GraphAStar.h
@@ -511,6 +511,16 @@ struct FGraphAStar
 
 		if (StartNode.NodeRef == EndNode.NodeRef)
 		{
+			UE_GRAPH_ASTAR_LOG(Display, TEXT("Storing path result (length=1)..."));
+			OutPath.Reset(1);
+        	OutPath.AddZeroed(1);
+
+			FSearchNode InitialNode = StartNode;
+			InitialNode.TraversalCost = 0;
+			InitialNode.TotalCost = 0;
+
+			UE_GRAPH_ASTAR_LOG(Display, TEXT("  NodeRef %i"), InitialNode.NodeRef);
+			SetPathInfo(OutPath, 0, InitialNode);
 			return SearchSuccess;
 		}
 
