[EQS Update]
- Create a new UEnvQueryContext_NavigationData,
- Add a NavDataOverrideContext in UEnvQueryGenerator_ProjectedPoints to be able to override Navigation Data instead of using the owner's one
#tests Tested in PIE with several EQS on AI

diff --git a/Engine/Source/Runtime/AIModule/Classes/EnvironmentQuery/Contexts/EnvQueryContext_NavigationData.h b/Engine/Source/Runtime/AIModule/Classes/EnvironmentQuery/Contexts/EnvQueryContext_NavigationData.h
new file mode 100644
--- /dev/null
+++ b/Engine/Source/Runtime/AIModule/Classes/EnvironmentQuery/Contexts/EnvQueryContext_NavigationData.h
@@ -0,0 +1,27 @@
+// Copyright Epic Games, Inc. All Rights Reserved.
+
+#pragma once
+
+#include "AI/Navigation/NavigationTypes.h"
+#include "EnvironmentQuery/EnvQueryContext.h"
+#include "UObject/ObjectMacros.h"
+
+#include "EnvQueryContext_NavigationData.generated.h"
+
+struct FEnvQueryContextData;
+struct FEnvQueryInstance;
+
+UCLASS(MinimalAPI, Abstract, Blueprintable)
+class UEnvQueryContext_NavigationData : public UEnvQueryContext
+{
+	GENERATED_UCLASS_BODY()
+
+	/** Return the NavigationData setup in NavAgentProperties */
+	AIMODULE_API virtual void ProvideContext(FEnvQueryInstance& QueryInstance, FEnvQueryContextData& ContextData) const override;
+
+protected:
+
+	/** NavAgentProperties used to find a NavigationData Override */
+	UPROPERTY(EditDefaultsOnly, Category = "NavigationData")
+	FNavAgentProperties NavAgentProperties;
+};
diff --git a/Engine/Source/Runtime/AIModule/Classes/EnvironmentQuery/Generators/EnvQueryGenerator_ProjectedPoints.h b/Engine/Source/Runtime/AIModule/Classes/EnvironmentQuery/Generators/EnvQueryGenerator_ProjectedPoints.h
--- a/Engine/Source/Runtime/AIModule/Classes/EnvironmentQuery/Generators/EnvQueryGenerator_ProjectedPoints.h
+++ b/Engine/Source/Runtime/AIModule/Classes/EnvironmentQuery/Generators/EnvQueryGenerator_ProjectedPoints.h
@@ -2,13 +2,17 @@
 
 #pragma once
 
+#if UE_ENABLE_INCLUDE_ORDER_DEPRECATED_IN_5_6
 #include "CoreMinimal.h"
+#endif // UE_ENABLE_INCLUDE_ORDER_DEPRECATED_IN_5_6
 #include "UObject/ObjectMacros.h"
 #include "AI/Navigation/NavigationTypes.h"
-#include "EnvironmentQuery/EnvQueryTypes.h"
 #include "EnvironmentQuery/EnvQueryGenerator.h"
+
 #include "EnvQueryGenerator_ProjectedPoints.generated.h"
 
+class EnvQueryContext;
+
 UCLASS(Abstract, MinimalAPI)
 class UEnvQueryGenerator_ProjectedPoints : public UEnvQueryGenerator
 {
@@ -37,4 +41,8 @@ class UEnvQueryGenerator_ProjectedPoints : public UEnvQueryGenerator
 	AIMODULE_API virtual void StoreNavPoints(const TArray<FNavLocation>& Points, FEnvQueryInstance& QueryInstance) const;
 
 	AIMODULE_API virtual void PostLoad() override;
+
+	/** Navigation Data Override */
+	UPROPERTY(EditDefaultsOnly, Category = "Generator")
+	TSubclassOf<UEnvQueryContext> NavDataOverrideContext;
 };
diff --git a/Engine/Source/Runtime/AIModule/Private/EnvironmentQuery/Contexts/EnvQueryContext_NavigationData.cpp b/Engine/Source/Runtime/AIModule/Private/EnvironmentQuery/Contexts/EnvQueryContext_NavigationData.cpp
new file mode 100644
--- /dev/null
+++ b/Engine/Source/Runtime/AIModule/Private/EnvironmentQuery/Contexts/EnvQueryContext_NavigationData.cpp
@@ -0,0 +1,25 @@
+// Copyright Epic Games, Inc. All Rights Reserved.
+
+#include "EnvironmentQuery/Contexts/EnvQueryContext_NavigationData.h"
+#include "EnvironmentQuery/EnvQueryManager.h"
+#include "EnvironmentQuery/Items/EnvQueryItemType_Actor.h"
+#include "NavigationSystem.h"
+#include "NavigationSystemTypes.h"
+
+#include UE_INLINE_GENERATED_CPP_BY_NAME(EnvQueryContext_NavigationData)
+
+UEnvQueryContext_NavigationData::UEnvQueryContext_NavigationData(const FObjectInitializer& ObjectInitializer)
+	: Super(ObjectInitializer)
+{
+}
+
+void UEnvQueryContext_NavigationData::ProvideContext(FEnvQueryInstance& QueryInstance, FEnvQueryContextData& ContextData) const
+{
+	const UNavigationSystemV1* NavigationSystem = FNavigationSystem::GetCurrent<UNavigationSystemV1>(GetWorld());
+	const ANavigationData* NavigationData = NavigationSystem ? NavigationSystem->GetNavDataForProps(NavAgentProperties) : nullptr;
+	
+	if (NavigationData)
+	{
+		UEnvQueryItemType_Actor::SetContextHelper(ContextData, NavigationData);
+	}
+}
diff --git a/Engine/Source/Runtime/AIModule/Private/EnvironmentQuery/Generators/EnvQueryGenerator_ProjectedPoints.cpp b/Engine/Source/Runtime/AIModule/Private/EnvironmentQuery/Generators/EnvQueryGenerator_ProjectedPoints.cpp
--- a/Engine/Source/Runtime/AIModule/Private/EnvironmentQuery/Generators/EnvQueryGenerator_ProjectedPoints.cpp
+++ b/Engine/Source/Runtime/AIModule/Private/EnvironmentQuery/Generators/EnvQueryGenerator_ProjectedPoints.cpp
@@ -3,6 +3,10 @@
 #include "EnvironmentQuery/Generators/EnvQueryGenerator_ProjectedPoints.h"
 #include "EnvironmentQuery/Items/EnvQueryItemType_Point.h"
 #include "EnvironmentQuery/EnvQueryTraceHelpers.h"
+#include "GameFramework/Pawn.h"
+#include "NavigationData.h"
+#include "NavigationSystem.h"
+#include "NavigationSystemTypes.h"
 
 #include UE_INLINE_GENERATED_CPP_BY_NAME(EnvQueryGenerator_ProjectedPoints)
 
@@ -21,7 +25,38 @@ void UEnvQueryGenerator_ProjectedPoints::ProjectAndFilterNavPoints(TArray<FNavLo
 	const ANavigationData* NavData = nullptr;
 	if (ProjectionData.TraceMode != EEnvQueryTrace::None)
 	{
-		NavData = FEQSHelpers::FindNavigationDataForQuery(QueryInstance);
+		if (NavDataOverrideContext)
+		{
+			TArray<AActor*> NavDataCandidateActors;
+			if (QueryInstance.PrepareContext(NavDataOverrideContext, NavDataCandidateActors))
+			{
+				const UNavigationSystemV1* NavigationSystem = FNavigationSystem::GetCurrent<UNavigationSystemV1>(GetWorld());
+
+				for (const AActor* NavDataCandidateActor : NavDataCandidateActors)
+				{
+					if (const ANavigationData* NavDataCandidate = Cast<ANavigationData>(NavDataCandidateActor))
+					{
+						// NavigationData found, we can stop here
+						NavData = NavDataCandidate;
+						break;
+					}
+					else if (const APawn* PawnCandidate = Cast<APawn>(NavDataCandidateActor))
+					{
+						// Check if this actor use a NavData that could be used
+						if (NavigationSystem)
+						{
+							NavData = NavigationSystem->GetNavDataForProps(PawnCandidate->GetNavAgentPropertiesRef());
+							break;
+						}
+					}
+				}
+			}
+		}
+
+		if (!NavData)
+		{
+			NavData = FEQSHelpers::FindNavigationDataForQuery(QueryInstance);
+		}
 	}
 
 	const UObject* Querier = QueryInstance.Owner.Get();
