Add missing references for loops in UCrowdManager::PostMovePointUpdate.

diff --git a/Engine/Source/Runtime/AIModule/Private/Navigation/CrowdManager.cpp b/Engine/Source/Runtime/AIModule/Private/Navigation/CrowdManager.cpp
--- a/Engine/Source/Runtime/AIModule/Private/Navigation/CrowdManager.cpp
+++ b/Engine/Source/Runtime/AIModule/Private/Navigation/CrowdManager.cpp
@@ -1425,10 +1425,10 @@ void UCrowdManager::PostMovePointUpdate()
 	const uint8 UpdateDestinationFlag = 1;
 
 	// special case when following last segment of full path to actor: replace end point with actor's location
-	for (auto It : ActiveAgents)
+	for (auto& AgentTuple : ActiveAgents)
 	{
-		UCrowdFollowingComponent* PathComp = Cast<UCrowdFollowingComponent>(It.Key);
-		const FCrowdAgentData& AgentData = It.Value;
+		UCrowdFollowingComponent* PathComp = Cast<UCrowdFollowingComponent>(AgentTuple.Key);
+		const FCrowdAgentData& AgentData = AgentTuple.Value;
 		FVector NewGoalPosition;
 
 		const bool bUpdateTargetPos = PathComp ? PathComp->ShouldTrackMovingGoal(NewGoalPosition) : false;
@@ -1462,9 +1462,9 @@ void UCrowdManager::PostMovePointUpdate()
 		}
 	}
 
-	for (auto It : ActiveAgents)
+	for (auto& AgentTuple : ActiveAgents)
 	{
-		const FCrowdAgentData& AgentData = It.Value;
+		const FCrowdAgentData& AgentData = AgentTuple.Value;
 		if (AgentFlags.IsValidIndex(AgentData.AgentIndex))
 		{
 			AgentFlags[AgentData.AgentIndex] &= ~UpdateDestinationFlag;
