Add option in FGraphAStar to support bRequireNavigableEndLocation from a query

diff --git a/Engine/Source/Runtime/AIModule/Public/GraphAStar.h b/Engine/Source/Runtime/AIModule/Public/GraphAStar.h
--- a/Engine/Source/Runtime/AIModule/Public/GraphAStar.h
+++ b/Engine/Source/Runtime/AIModule/Public/GraphAStar.h
@@ -182,6 +182,7 @@ template <> struct TAllocatorTraits<FRangeChecklessAllocator<false>> : TAllocato
  *		int32 GetNeighbourCount(FNodeRef NodeRef) const;													- returns number of neighbours that the graph node identified with NodeRef has, it is ok if not implemented, the logic will stop calling GetNeighbour once it received an invalid noderef
  *		bool ShouldIgnoreClosedNodes() const;																- whether to revisit closed node or not
  *		bool ShouldIncludeStartNodeInPath() const;															- whether to put the start node in the resulting path
+ *		bool ShouldFailOnInvalidEndNode() const;															- whether to early out if the end node is not valid
  *		int32 GetMaxSearchNodes() const;																	- whether to limit the number of search nodes to a maximum
  *		FVector::FReal GetCostLimit() const																	- whether to limit the search to a maximum cost
  */
@@ -332,6 +333,7 @@ struct FGraphAStar
 	DECLARE_OPTIONALLY_IMPLEMENTED_TEMPLATE_CLASS_FUNCTION_2PARAMS_CONST(FVector::FReal, GetHeuristicCost, const FSearchNode&, const FSearchNode&, Obj.GetHeuristicCost(Param1.NodeRef, Param2.NodeRef))
 	DECLARE_OPTIONALLY_IMPLEMENTED_TEMPLATE_CLASS_FUNCTION_CONST(bool, ShouldIgnoreClosedNodes, false);
 	DECLARE_OPTIONALLY_IMPLEMENTED_TEMPLATE_CLASS_FUNCTION_CONST(bool, ShouldIncludeStartNodeInPath, false);
+	DECLARE_OPTIONALLY_IMPLEMENTED_TEMPLATE_CLASS_FUNCTION_CONST(bool, ShouldFailOnInvalidEndNode, true);
 	DECLARE_OPTIONALLY_IMPLEMENTED_TEMPLATE_CLASS_FUNCTION_CONST(FVector::FReal, GetCostLimit, TNumericLimits<FVector::FReal>::Max());
 	// Custom methods implemented over TQueryFilter optionally implemented methods
 	DECLARE_OPTIONALLY_IMPLEMENTED_TEMPLATE_CLASS_FUNCTION_CUSTOM(GetMaxSearchNodes, bool, HasReachMaxSearchNodes, uint32 NodeCount, false, NodeCount >= Obj.GetMaxSearchNodes());
@@ -504,7 +506,7 @@ struct FGraphAStar
 		UE_GRAPH_ASTAR_LOG(Display, TEXT(""));
 		UE_GRAPH_ASTAR_LOG(Display, TEXT("Starting FindPath request..."));
 
-		if (!(Graph.IsValidRef(StartNode.NodeRef) && Graph.IsValidRef(EndNode.NodeRef)))
+		if (!Graph.IsValidRef(StartNode.NodeRef) || (ShouldFailOnInvalidEndNode(Filter) && !Graph.IsValidRef(EndNode.NodeRef)))
 		{
 			return SearchFail;
 		}
