PR #12397: EnvQuery runtime template object stored in the EnvQueryManager cache should be created with a unique name

diff --git a/Engine/Source/Runtime/AIModule/Private/EnvironmentQuery/EnvQueryManager.cpp b/Engine/Source/Runtime/AIModule/Private/EnvironmentQuery/EnvQueryManager.cpp
--- a/Engine/Source/Runtime/AIModule/Private/EnvironmentQuery/EnvQueryManager.cpp
+++ b/Engine/Source/Runtime/AIModule/Private/EnvironmentQuery/EnvQueryManager.cpp
@@ -805,8 +805,8 @@ TSharedPtr<FEnvQueryInstance> UEnvQueryManager::CreateQueryInstance(const UEnvQu
 
 		// duplicate template in manager's world for BP based nodes
 		const FString NewInstanceName = RunModeEnum
-			? FString::Printf(TEXT("%s_%s"), *Template->GetFName().ToString(), *RunModeEnum->GetNameStringByValue(RunMode))
-			: FString::Printf(TEXT("%s_%d"), *Template->GetFName().ToString(), uint8(RunMode));
+			? FString::Printf(TEXT("%s_%s"), *MakeUniqueObjectName(this, Template->GetClass(), Template->GetFName()).ToString(), *RunModeEnum->GetNameStringByValue(RunMode))
+			: FString::Printf(TEXT("%s_%d"), *MakeUniqueObjectName(this, Template->GetClass(), Template->GetFName()).ToString(), uint8(RunMode));
 		UEnvQuery* LocalTemplate = (UEnvQuery*)StaticDuplicateObject(Template, this, *NewInstanceName);
 
 		{
