Minor change in PathFollowingComponent: Moved the code that extracts the NavigationFilter from the MoveRequest/NavPath into a function.

#rnx

diff --git a/Engine/Source/Runtime/AIModule/Classes/Navigation/PathFollowingComponent.h b/Engine/Source/Runtime/AIModule/Classes/Navigation/PathFollowingComponent.h
--- a/Engine/Source/Runtime/AIModule/Classes/Navigation/PathFollowingComponent.h
+++ b/Engine/Source/Runtime/AIModule/Classes/Navigation/PathFollowingComponent.h
@@ -595,6 +595,9 @@ protected:
 	/** index of path point for starting deceleration */
 	int32 DecelerationSegmentIndex;
 
+	/** extract the navigation filter that should be used based on the provided MoveRequest and NavPath */
+	AIMODULE_API FSharedConstNavQueryFilter ExtractNavigationFilter(const FAIMoveRequest& RequestData, FNavPathSharedPtr InPath) const;
+
 	/** reset path following data */
 	AIMODULE_API virtual void Reset();
 
diff --git a/Engine/Source/Runtime/AIModule/Private/Navigation/PathFollowingComponent.cpp b/Engine/Source/Runtime/AIModule/Private/Navigation/PathFollowingComponent.cpp
--- a/Engine/Source/Runtime/AIModule/Private/Navigation/PathFollowingComponent.cpp
+++ b/Engine/Source/Runtime/AIModule/Private/Navigation/PathFollowingComponent.cpp
@@ -47,34 +47,6 @@ namespace UE::Navigation::Private
 			return GoalActor.GetActorLocation();
 		}
 	}
-
-	FSharedConstNavQueryFilter ExtractNavigationFilterForRequest(UObject* Owner, FNavPathSharedPtr Path, TObjectPtr<ANavigationData> MyNavData, const FAIMoveRequest& RequestData)
-	{
-		if (!Path->CastPath<FAbstractNavigationPath>() && Path->GetFilter() != nullptr)
-		{
-			return Path->GetFilter();
-		}
-
-		if (MyNavData)
-		{
-			if (RequestData.GetNavigationFilter() != nullptr)
-			{
-				return UNavigationQueryFilter::GetQueryFilter(*MyNavData, Owner, RequestData.GetNavigationFilter());
-			}
-
-			if (const AAIController* AIOwner = Cast<AAIController>(Owner))
-			{
-				if (AIOwner->GetDefaultNavigationFilterClass() != nullptr)
-				{
-					return UNavigationQueryFilter::GetQueryFilter(*MyNavData, AIOwner->GetDefaultNavigationFilterClass());
-				}
-			}
-
-			return MyNavData->GetDefaultQueryFilter();
-		}
-
-		return FSharedConstNavQueryFilter();
-	}
 }
 
 FPathFollowingResult::FPathFollowingResult(uint16 InFlags) : Flags(InFlags)
@@ -423,7 +395,7 @@ FAIRequestID UPathFollowingComponent::RequestMove(const FAIMoveRequest& RequestD
 			MetaNavPath->Initialize(CurrentLocation);
 		}
 
-		NavigationFilter = UE::Navigation::Private::ExtractNavigationFilterForRequest(GetOwner(), Path, MyNavData, RequestData);
+		NavigationFilter = ExtractNavigationFilter(RequestData, Path);
 
 		PathTimeWhenPaused = 0.;
 		OnPathUpdated();
@@ -771,6 +743,34 @@ FVector UPathFollowingComponent::GetMoveFocus(bool bAllowStrafe) const
 	return MoveFocus;
 }
 
+FSharedConstNavQueryFilter UPathFollowingComponent::ExtractNavigationFilter(const FAIMoveRequest& RequestData, FNavPathSharedPtr InPath) const
+{
+	if (InPath.IsValid() && !InPath->CastPath<FAbstractNavigationPath>() && InPath->GetFilter() != nullptr)
+	{
+		return InPath->GetFilter();
+	}
+
+	if (MyNavData)
+	{
+		if (RequestData.GetNavigationFilter() != nullptr)
+		{
+			return UNavigationQueryFilter::GetQueryFilter(*MyNavData, GetOwner(), RequestData.GetNavigationFilter());
+		}
+
+		if (const AAIController* AIOwner = Cast<AAIController>(GetOwner()))
+		{
+			if (AIOwner->GetDefaultNavigationFilterClass() != nullptr)
+			{
+				return UNavigationQueryFilter::GetQueryFilter(*MyNavData, AIOwner->GetDefaultNavigationFilterClass());
+			}
+		}
+
+		return MyNavData->GetDefaultQueryFilter();
+	}
+
+	return FSharedConstNavQueryFilter();
+}
+
 void UPathFollowingComponent::Reset()
 {
 	MoveSegmentStartIndex = 0;
