FAIGenericID: Remove implicit conversions to/from uint32 to help prevent accidental conversions

diff --git a/Engine/Source/Runtime/AIModule/Classes/AITypes.h b/Engine/Source/Runtime/AIModule/Classes/AITypes.h
--- a/Engine/Source/Runtime/AIModule/Classes/AITypes.h
+++ b/Engine/Source/Runtime/AIModule/Classes/AITypes.h
@@ -278,11 +278,12 @@ protected:
 		return Counter;
 	}
 
-	FAIGenericID(typename TCounter::Type InIndex)
+	explicit FAIGenericID(typename TCounter::Type InIndex)
 		: Index(InIndex)
 	{}
 
 public:
+
 	FAIGenericID(const FAIGenericID& Other)
 		: Index(Other.Index)
 	{}
@@ -299,9 +300,14 @@ public:
 
 	static FAIGenericID GetNextID() { return FAIGenericID(GetCounter().GetNextAvailableID()); }
 
+	UE_DEPRECATED(5.6, "AIGenericId is not convertible to its base type anymore.")
 	operator typename TCounter::Type() const { return Index; }
+
 	bool IsValid() const { return Index != InvalidID().Index; }
 
+	bool operator==(const FAIGenericID Other) const { return Index == Other.Index; }
+	bool operator!=(const FAIGenericID Other) const { return Index != Other.Index; }
+
 	static uint32 GetSize() { return GetCounter().GetSize(); }
 
 	static FAIGenericID<TCounter> InvalidID()
diff --git a/Engine/Source/Runtime/AIModule/Private/Perception/AISense_Sight.cpp b/Engine/Source/Runtime/AIModule/Private/Perception/AISense_Sight.cpp
--- a/Engine/Source/Runtime/AIModule/Private/Perception/AISense_Sight.cpp
+++ b/Engine/Source/Runtime/AIModule/Private/Perception/AISense_Sight.cpp
@@ -1013,7 +1013,7 @@ void UAISense_Sight::RemoveAllQueriesByListener(const FPerceptionListener& Liste
 	SCOPE_CYCLE_COUNTER(STAT_AI_Sense_Sight_RemoveByListener);
 	UE_MT_SCOPED_WRITE_ACCESS(QueriesListAccessDetector);
 
-	const uint32 ListenerId = Listener.GetListenerID();
+	const FPerceptionListenerID ListenerId = Listener.GetListenerID();
 
 	auto RemoveQuery = [&ListenerId, &OnRemoveFunc](TArray<FAISightQuery>& SightQueries, const int32 QueryIndex)->EReverseForEachResult
 	{
@@ -1073,7 +1073,7 @@ void UAISense_Sight::RemoveAllQueriesToTarget_Internal(const FAISightTarget::FTa
 
 void UAISense_Sight::OnListenerForgetsActor(const FPerceptionListener& Listener, AActor& ActorToForget)
 {
-	const uint32 ListenerId = Listener.GetListenerID();
+	const FPerceptionListenerID ListenerId = Listener.GetListenerID();
 	const uint32 TargetId = ActorToForget.GetUniqueID();
 	
 	auto ForgetPreviousResult = [&ListenerId, &TargetId](FAISightQuery& SightQuery)->EForEachResult
@@ -1100,7 +1100,7 @@ void UAISense_Sight::OnListenerForgetsAll(const FPerceptionListener& Listener)
 {
 	UE_MT_SCOPED_WRITE_ACCESS(QueriesListAccessDetector);
 
-	const uint32 ListenerId = Listener.GetListenerID();
+	const FPerceptionListenerID ListenerId = Listener.GetListenerID();
 
 	auto ForgetPreviousResult = [&ListenerId](FAISightQuery& SightQuery)->EForEachResult
 	{
