[PawnActions] Removing deprecated code

#rn PawnActions removed from the engine codebase.

diff --git a/Engine/Source/Developer/AITestSuite/Private/MockAI/MockAI.cpp b/Engine/Source/Developer/AITestSuite/Private/MockAI/MockAI.cpp
--- a/Engine/Source/Developer/AITestSuite/Private/MockAI/MockAI.cpp
+++ b/Engine/Source/Developer/AITestSuite/Private/MockAI/MockAI.cpp
@@ -3,7 +3,6 @@
 #include "MockAI.h"
 #include "BehaviorTree/BlackboardComponent.h"
 #include "Perception/AIPerceptionComponent.h"
-#include "Actions/PawnActionsComponent.h"
 #include "BrainComponent.h"
 
 #include UE_INLINE_GENERATED_CPP_BY_NAME(MockAI)
diff --git a/Engine/Source/Runtime/AIModule/Classes/AIController.h b/Engine/Source/Runtime/AIModule/Classes/AIController.h
--- a/Engine/Source/Runtime/AIModule/Classes/AIController.h
+++ b/Engine/Source/Runtime/AIModule/Classes/AIController.h
@@ -30,8 +30,6 @@ class UCanvas;
 class UGameplayTaskResource;
 class UGameplayTasksComponent;
 class UPathFollowingComponent;
-class UDEPRECATED_PawnAction;
-class UDEPRECATED_PawnActionsComponent;
 
 namespace EPathFollowingRequestResult {	enum Type : int; }
 namespace EPathFollowingResult { enum Type : int; }
@@ -143,10 +141,6 @@ public:
 	UPROPERTY(VisibleDefaultsOnly, Category = AI)
 	TObjectPtr<UAIPerceptionComponent> PerceptionComponent;
 
-private:
-	UPROPERTY(Category = AI, BlueprintGetter = GetDeprecatedActionsComponent)
-	TObjectPtr<UDEPRECATED_PawnActionsComponent> ActionsComp_DEPRECATED;
-
 protected:
 	/** blackboard */
 	UPROPERTY(BlueprintReadOnly, Category = AI, meta = (AllowPrivateAccess = "true"))
@@ -456,18 +450,6 @@ public:
 	 *	Intended to be called as part of initialization/setup process */
 	UFUNCTION(BlueprintCallable, Category = "AI|Navigation")
 	AIMODULE_API void SetPathFollowingComponent(UPathFollowingComponent* NewPFComponent);
-
-	//----------------------------------------------------------------------//
-	// DEPRECATED
-	//----------------------------------------------------------------------//
-	UE_DEPRECATED(5.2, "PawnActions have been deprecated and are no longer being supported. It will get removed in following UE5 releases. Use GameplayTasks or AITasks instead.")
-	UDEPRECATED_PawnActionsComponent* GetActionsComp() const { return ActionsComp_DEPRECATED; }
-	
-	UE_DEPRECATED(5.2, "PawnActions have been deprecated and are no longer being supported. It will get removed in following UE5 releases. Use GameplayTasks or AITasks instead.")
-	AIMODULE_API bool PerformAction(UDEPRECATED_PawnAction& Action, EAIRequestPriority::Type Priority, UObject* const Instigator = NULL);
-
-	UFUNCTION(BlueprintGetter, meta = (DeprecatedFunction, DeprecationMessage = "PawnActions have been deprecated and are no longer being supported. It will get removed in following UE5 releases. Use GameplayTasks or AITasks instead."))
-	UDEPRECATED_PawnActionsComponent* GetDeprecatedActionsComponent() const { return ActionsComp_DEPRECATED; }
 };
 
 //----------------------------------------------------------------------//
diff --git a/Engine/Source/Runtime/AIModule/Classes/AITypes.h b/Engine/Source/Runtime/AIModule/Classes/AITypes.h
--- a/Engine/Source/Runtime/AIModule/Classes/AITypes.h
+++ b/Engine/Source/Runtime/AIModule/Classes/AITypes.h
@@ -119,49 +119,6 @@ namespace EAILogicResuming
 	};
 }
 
-UENUM()
-namespace EPawnActionAbortState
-{
-	enum Type : int
-	{
-		NeverStarted,
-		NotBeingAborted,
-		/** This means waiting for child to abort before aborting self. */
-		MarkPendingAbort,
-		LatentAbortInProgress,
-		AbortDone,
-
-		MAX UMETA(Hidden)
-	};
-}
-
-UENUM()
-namespace EPawnActionResult
-{
-	enum Type : int
-	{
-		NotStarted,
-		InProgress,
-		Success,
-		Failed,
-		Aborted
-	};
-}
-
-UENUM()
-namespace EPawnActionEventType
-{
-	enum Type : int
-	{
-		Invalid,
-		FailedToStart,
-		InstantAbort,
-		FinishedAborting,
-		FinishedExecution,
-		Push,
-	};
-}
-
 UENUM()
 namespace EAIRequestPriority
 {
diff --git a/Engine/Source/Runtime/AIModule/Classes/Actions/PawnAction.h b/Engine/Source/Runtime/AIModule/Classes/Actions/PawnAction.h
deleted file mode 100644
diff --git a/Engine/Source/Runtime/AIModule/Classes/Actions/PawnAction_BlueprintBase.h b/Engine/Source/Runtime/AIModule/Classes/Actions/PawnAction_BlueprintBase.h
deleted file mode 100644
diff --git a/Engine/Source/Runtime/AIModule/Classes/Actions/PawnAction_Move.h b/Engine/Source/Runtime/AIModule/Classes/Actions/PawnAction_Move.h
--- a/Engine/Source/Runtime/AIModule/Classes/Actions/PawnAction_Move.h
+++ b/Engine/Source/Runtime/AIModule/Classes/Actions/PawnAction_Move.h
@@ -2,19 +2,8 @@
 
 #pragma once
 
-#include "CoreMinimal.h"
-#include "UObject/ObjectMacros.h"
-#include "Templates/SubclassOf.h"
-#include "Engine/EngineTypes.h"
-#if UE_ENABLE_INCLUDE_ORDER_DEPRECATED_IN_5_4
-#include "NavFilters/NavigationQueryFilter.h"
-#endif
-#include "Actions/PawnAction.h"
-#include "Navigation/PathFollowingComponent.h"
 #include "PawnAction_Move.generated.h"
 
-class AAIController;
-
 UENUM()
 namespace EPawnActionMoveMode
 {
@@ -24,104 +13,3 @@ namespace EPawnActionMoveMode
 		StraightLine,
 	};
 }
-
-UCLASS(MinimalAPI)
-class UDEPRECATED_PawnAction_Move : public UDEPRECATED_PawnAction
-{
-	GENERATED_UCLASS_BODY()
-protected:
-	UPROPERTY(Category = PawnAction, EditAnywhere, BlueprintReadWrite)
-	TObjectPtr<AActor> GoalActor;
-
-	UPROPERTY(Category = PawnAction, EditAnywhere, BlueprintReadWrite)
-	FVector GoalLocation;
-
-	UPROPERTY(Category = PawnAction, EditAnywhere, meta = (ClampMin = "0.01"), BlueprintReadWrite)
-	float AcceptableRadius;
-
-	/** "None" will result in default filter being used */
-	UPROPERTY(Category = PawnAction, EditAnywhere, BlueprintReadWrite)
-	TSubclassOf<UNavigationQueryFilter> FilterClass;
-
-	UPROPERTY(Category = PawnAction, EditAnywhere, BlueprintReadWrite)
-	uint32 bAllowStrafe : 1;
-	
-	/** if set to true (default) will make action succeed when the pawn's collision component overlaps with goal's collision component */
-	UPROPERTY()
-	uint32 bFinishOnOverlap : 1;
-
-	/** if set, movement will use path finding */
-	UPROPERTY()
-	uint32 bUsePathfinding : 1;
-
-	/** if set, use incomplete path when goal can't be reached */
-	UPROPERTY()
-	uint32 bAllowPartialPath : 1;
-
-	/** if set, GoalLocation will be projected on navigation before using  */
-	UPROPERTY()
-	uint32 bProjectGoalToNavigation : 1;
-
-	/** if set, path to GoalActor will be updated with goal's movement */
-	UPROPERTY()
-	uint32 bUpdatePathToGoal : 1;
-
-	/** if set, other actions with the same priority will be aborted when path is changed */
-	UPROPERTY()
-	uint32 bAbortSubActionOnPathChange : 1;
-
-public:
-	AIMODULE_API virtual void BeginDestroy() override;
-
-	static AIMODULE_API UDEPRECATED_PawnAction_Move* CreateAction(UWorld& World, AActor* GoalActor, EPawnActionMoveMode::Type Mode);
-	static AIMODULE_API UDEPRECATED_PawnAction_Move* CreateAction(UWorld& World, const FVector& GoalLocation, EPawnActionMoveMode::Type Mode);
-
-	static AIMODULE_API bool CheckAlreadyAtGoal(AAIController& Controller, const FVector& TestLocation, float Radius);
-	static AIMODULE_API bool CheckAlreadyAtGoal(AAIController& Controller, const AActor& TestGoal, float Radius);
-
-	AIMODULE_API virtual void HandleAIMessage(UBrainComponent*, const FAIMessage&) override;
-
-	AIMODULE_API void SetPath(FNavPathSharedRef InPath);
-	AIMODULE_API virtual void OnPathUpdated(FNavigationPath* UpdatedPath, ENavPathEvent::Type Event);
-
-	void SetAcceptableRadius(float NewAcceptableRadius) { AcceptableRadius = NewAcceptableRadius; }
-	void SetFinishOnOverlap(bool bNewFinishOnOverlap) { bFinishOnOverlap = bNewFinishOnOverlap; }
-	void EnableStrafing(bool bNewStrafing) { bAllowStrafe = bNewStrafing; }
-	void EnablePathUpdateOnMoveGoalLocationChange(bool bEnable) { bUpdatePathToGoal = bEnable; }
-	void EnableGoalLocationProjectionToNavigation(bool bEnable) { bProjectGoalToNavigation = bEnable; }
-	void SetAbortSubActionOnPathUpdate(bool bEnable) { bAbortSubActionOnPathChange = bEnable; }
-	void SetFilterClass(TSubclassOf<UNavigationQueryFilter> NewFilterClass) { FilterClass = NewFilterClass; }
-	void SetAllowPartialPath(bool bEnable) { bAllowPartialPath = bEnable; }
-
-	UE_DEPRECATED(5.1, "Use SetAbortSubActionOnPathUpdate instead.")
-	void EnableChildAbortionOnPathUpdate(bool bEnable) { SetAbortSubActionOnPathUpdate(bEnable); }
-
-protected:
-	/** currently followed path */
-	FNavPathSharedPtr Path;
-
-	FDelegateHandle PathObserverDelegateHandle;
-	
-	/** Handle for efficient management of DeferredPerformMoveAction timer */
-	FTimerHandle TimerHandle_DeferredPerformMoveAction;
-
-	/** Handle for efficient management of TryToRepath timer */
-	FTimerHandle TimerHandle_TryToRepath;
-
-	AIMODULE_API void ClearPath();
-	AIMODULE_API virtual bool Start() override;
-	AIMODULE_API virtual bool Pause(const UDEPRECATED_PawnAction* PausedBy) override;
-	AIMODULE_API virtual bool Resume() override;
-	AIMODULE_API virtual void OnFinished(EPawnActionResult::Type WithResult) override;
-	AIMODULE_API virtual EPawnActionAbortState::Type PerformAbort(EAIForceParam::Type ShouldForce) override;
-	AIMODULE_API virtual bool IsPartialPathAllowed() const;
-
-	AIMODULE_API virtual EPathFollowingRequestResult::Type RequestMove(AAIController& Controller);
-	
-	AIMODULE_API bool PerformMoveAction();
-	AIMODULE_API void DeferredPerformMoveAction();
-
-	AIMODULE_API void TryToRepath();
-	AIMODULE_API void ClearPendingRepath();
-	AIMODULE_API void ClearTimers();
-};
diff --git a/Engine/Source/Runtime/AIModule/Classes/Actions/PawnAction_Repeat.h b/Engine/Source/Runtime/AIModule/Classes/Actions/PawnAction_Repeat.h
deleted file mode 100644
diff --git a/Engine/Source/Runtime/AIModule/Classes/Actions/PawnAction_Sequence.h b/Engine/Source/Runtime/AIModule/Classes/Actions/PawnAction_Sequence.h
deleted file mode 100644
diff --git a/Engine/Source/Runtime/AIModule/Classes/Actions/PawnAction_Wait.h b/Engine/Source/Runtime/AIModule/Classes/Actions/PawnAction_Wait.h
deleted file mode 100644
diff --git a/Engine/Source/Runtime/AIModule/Classes/Actions/PawnActionsComponent.h b/Engine/Source/Runtime/AIModule/Classes/Actions/PawnActionsComponent.h
deleted file mode 100644
diff --git a/Engine/Source/Runtime/AIModule/Classes/BehaviorTree/Tasks/BTTask_PawnActionBase.h b/Engine/Source/Runtime/AIModule/Classes/BehaviorTree/Tasks/BTTask_PawnActionBase.h
deleted file mode 100644
diff --git a/Engine/Source/Runtime/AIModule/Classes/BehaviorTree/Tasks/BTTask_PushPawnAction.h b/Engine/Source/Runtime/AIModule/Classes/BehaviorTree/Tasks/BTTask_PushPawnAction.h
deleted file mode 100644
diff --git a/Engine/Source/Runtime/AIModule/Private/AIController.cpp b/Engine/Source/Runtime/AIModule/Private/AIController.cpp
--- a/Engine/Source/Runtime/AIModule/Private/AIController.cpp
+++ b/Engine/Source/Runtime/AIModule/Private/AIController.cpp
@@ -5,7 +5,6 @@
 #include "NavigationSystem.h"
 #include "Components/CapsuleComponent.h"
 #include "GameFramework/PhysicsVolume.h"
-#include "Actions/PawnActionsComponent.h"
 #include "Engine/Canvas.h"
 #include "BehaviorTree/BehaviorTreeComponent.h"
 #include "Tasks/AITask.h"
@@ -48,8 +47,6 @@ AAIController::AAIController(const FObjectInitializer& ObjectInitializer)
 		PathFollowingComponent->OnRequestFinished.AddUObject(this, &AAIController::OnMoveCompleted);
 	}
 
-	ActionsComp_DEPRECATED = CreateOptionalDefaultSubobject<UDEPRECATED_PawnActionsComponent>(TEXT("ActionsComp"));
-
 	bSkipExtraLOSChecks = true;
 	bWantsPlayerState = false;
 	TeamID = FGenericTeamId::NoTeam;
@@ -1148,12 +1145,6 @@ bool AAIController::SuggestTossVelocity(FVector& OutTossVelocity, FVector Start,
 
 	return UGameplayStatics::SuggestProjectileVelocity(VelocityParams, OutTossVelocity);
 }
-bool AAIController::PerformAction(UDEPRECATED_PawnAction& Action, EAIRequestPriority::Type Priority, UObject* const InInstigator /*= NULL*/)
-{
-PRAGMA_DISABLE_DEPRECATION_WARNINGS
-	return ActionsComp_DEPRECATED != NULL && ActionsComp_DEPRECATED->PushAction(Action, Priority, InInstigator);
-PRAGMA_ENABLE_DEPRECATION_WARNINGS
-}
 
 FString AAIController::GetDebugIcon() const
 {
diff --git a/Engine/Source/Runtime/AIModule/Private/AISystem.cpp b/Engine/Source/Runtime/AIModule/Private/AISystem.cpp
--- a/Engine/Source/Runtime/AIModule/Private/AISystem.cpp
+++ b/Engine/Source/Runtime/AIModule/Private/AISystem.cpp
@@ -175,7 +175,7 @@ void UAISystem::AILoggingVerbose()
 		APlayerController* PC = OuterWorld->GetGameInstance()->GetFirstLocalPlayerController();
 		if (PC)
 		{
-			PC->ConsoleCommand(TEXT("log lognavigation verbose | log logpathfollowing verbose | log LogCharacter verbose | log LogBehaviorTree verbose | log LogPawnAction verbose|"));
+			PC->ConsoleCommand(TEXT("log lognavigation verbose | log logpathfollowing verbose | log LogCharacter verbose | log LogBehaviorTree verbose"));
 		}
 	}
 }
diff --git a/Engine/Source/Runtime/AIModule/Private/Actions/PawnAction.cpp b/Engine/Source/Runtime/AIModule/Private/Actions/PawnAction.cpp
deleted file mode 100644
diff --git a/Engine/Source/Runtime/AIModule/Private/Actions/PawnAction_BlueprintBase.cpp b/Engine/Source/Runtime/AIModule/Private/Actions/PawnAction_BlueprintBase.cpp
deleted file mode 100644
diff --git a/Engine/Source/Runtime/AIModule/Private/Actions/PawnAction_Move.cpp b/Engine/Source/Runtime/AIModule/Private/Actions/PawnAction_Move.cpp
deleted file mode 100644
diff --git a/Engine/Source/Runtime/AIModule/Private/Actions/PawnAction_Repeat.cpp b/Engine/Source/Runtime/AIModule/Private/Actions/PawnAction_Repeat.cpp
deleted file mode 100644
diff --git a/Engine/Source/Runtime/AIModule/Private/Actions/PawnAction_Sequence.cpp b/Engine/Source/Runtime/AIModule/Private/Actions/PawnAction_Sequence.cpp
deleted file mode 100644
diff --git a/Engine/Source/Runtime/AIModule/Private/Actions/PawnAction_Wait.cpp b/Engine/Source/Runtime/AIModule/Private/Actions/PawnAction_Wait.cpp
deleted file mode 100644
diff --git a/Engine/Source/Runtime/AIModule/Private/Actions/PawnActionsComponent.cpp b/Engine/Source/Runtime/AIModule/Private/Actions/PawnActionsComponent.cpp
deleted file mode 100644
diff --git a/Engine/Source/Runtime/AIModule/Private/BehaviorTree/Tasks/BTTask_PawnActionBase.cpp b/Engine/Source/Runtime/AIModule/Private/BehaviorTree/Tasks/BTTask_PawnActionBase.cpp
deleted file mode 100644
diff --git a/Engine/Source/Runtime/AIModule/Private/BehaviorTree/Tasks/BTTask_PushPawnAction.cpp b/Engine/Source/Runtime/AIModule/Private/BehaviorTree/Tasks/BTTask_PushPawnAction.cpp
deleted file mode 100644
