[PathFollowingComponent] activate component ticking only when necessary

diff --git a/Engine/Source/Runtime/AIModule/Classes/Navigation/PathFollowingComponent.h b/Engine/Source/Runtime/AIModule/Classes/Navigation/PathFollowingComponent.h
--- a/Engine/Source/Runtime/AIModule/Classes/Navigation/PathFollowingComponent.h
+++ b/Engine/Source/Runtime/AIModule/Classes/Navigation/PathFollowingComponent.h
@@ -243,8 +243,8 @@ class UPathFollowingComponent : public UActorComponent, public IAIResourceInterf
 	AIMODULE_API virtual void UpdateCachedComponents();
 
 	/** start movement along path
-	  * @return MoveId of requested move
-	  */
+	 * @return MoveId of requested move
+	 */
 	AIMODULE_API virtual FAIRequestID RequestMove(const FAIMoveRequest& RequestData, FNavPathSharedPtr InPath);
 
 	/** aborts following path */
@@ -256,11 +256,11 @@ class UPathFollowingComponent : public UActorComponent, public IAIResourceInterf
 	AIMODULE_API FAIRequestID RequestMoveWithImmediateFinish(EPathFollowingResult::Type Result, EPathFollowingVelocityMode VelocityMode = EPathFollowingVelocityMode::Reset);
 
 	/** pause path following
-	*  @param RequestID - request to pause, FAIRequestID::CurrentRequest means pause current request, regardless of its ID */
+	 *  @param RequestID - request to pause, FAIRequestID::CurrentRequest means pause current request, regardless of its ID */
 	AIMODULE_API virtual void PauseMove(FAIRequestID RequestID = FAIRequestID::CurrentRequest, EPathFollowingVelocityMode VelocityMode = EPathFollowingVelocityMode::Reset);
 
 	/** resume path following
-	*  @param RequestID - request to resume, FAIRequestID::CurrentRequest means restor current request, regardless of its ID */
+	 *  @param RequestID - request to resume, FAIRequestID::CurrentRequest means restor current request, regardless of its ID */
 	AIMODULE_API virtual void ResumeMove(FAIRequestID RequestID = FAIRequestID::CurrentRequest);
 
 	/** notify about finished movement */
@@ -337,13 +337,23 @@ class UPathFollowingComponent : public UActorComponent, public IAIResourceInterf
 	/** Returns current location on navigation data */
 	AIMODULE_API FNavLocation GetCurrentNavLocation() const;
 
-	FORCEINLINE EPathFollowingStatus::Type GetStatus() const { return Status; }
+	FORCEINLINE EPathFollowingStatus::Type GetStatus() const
+	{
+		PRAGMA_DISABLE_DEPRECATION_WARNINGS
+		return Status;
+		PRAGMA_ENABLE_DEPRECATION_WARNINGS
+	}
 	FORCEINLINE float GetAcceptanceRadius() const { return AcceptanceRadius; }
 	FORCEINLINE float GetDefaultAcceptanceRadius() const { return MyDefaultAcceptanceRadius; }
 	AIMODULE_API void SetAcceptanceRadius(const float InAcceptanceRadius);
 	FORCEINLINE AActor* GetMoveGoal() const { return DestinationActor.Get(); }
 	FORCEINLINE bool HasPartialPath() const { return Path.IsValid() && Path->IsPartial(); }
-	FORCEINLINE bool DidMoveReachGoal() const { return bLastMoveReachedGoal && (Status == EPathFollowingStatus::Idle); }
+	FORCEINLINE bool DidMoveReachGoal() const
+	{
+		PRAGMA_DISABLE_DEPRECATION_WARNINGS
+		return bLastMoveReachedGoal && (Status == EPathFollowingStatus::Idle);
+		PRAGMA_ENABLE_DEPRECATION_WARNINGS
+	}
 
 	FORCEINLINE FAIRequestID GetCurrentRequestId() const { return CurrentRequestId; }
 	FORCEINLINE uint32 GetCurrentPathIndex() const { return MoveSegmentStartIndex; }
@@ -388,12 +398,11 @@ class UPathFollowingComponent : public UActorComponent, public IAIResourceInterf
 	UFUNCTION()
 	AIMODULE_API virtual void OnActorBump(AActor* SelfActor, AActor* OtherActor, FVector NormalImpulse, const FHitResult& Hit);
 
-	// IPathFollowingAgentInterface begin
-	AIMODULE_API virtual void OnUnableToMove(const UObject& Instigator) override;		
-	//virtual void OnMoveBlockedBy(const FHitResult& BlockingImpact) {}
+	//~ IPathFollowingAgentInterface begin
+	AIMODULE_API virtual void OnUnableToMove(const UObject& Instigator) override;
 	AIMODULE_API virtual void OnStartedFalling() override;
 	virtual void OnLanded() override {}
-	// IPathFollowingAgentInterface end
+	//~ IPathFollowingAgentInterface end
 
 	/** Check if path following can be activated */
 	AIMODULE_API virtual bool IsPathFollowingAllowed() const;
@@ -404,12 +413,12 @@ class UPathFollowingComponent : public UActorComponent, public IAIResourceInterf
 	/** called when owner is preparing new pathfinding request */
 	virtual void OnPathfindingQuery(FPathFindingQuery& Query) {}
 
-	// IAIResourceInterface begin
+	//~ IAIResourceInterface begin
 	AIMODULE_API virtual void LockResource(EAIRequestPriority::Type LockSource) override;
 	AIMODULE_API virtual void ClearResourceLock(EAIRequestPriority::Type LockSource) override;
 	AIMODULE_API virtual void ForceUnlockResource() override;
 	AIMODULE_API virtual bool IsResourceLocked() const override;
-	// IAIResourceInterface end
+	//~ IAIResourceInterface end
 
 	/** path observer */
 	AIMODULE_API void OnPathEvent(FNavigationPath* InPath, ENavPathEvent::Type Event);
@@ -430,6 +439,8 @@ class UPathFollowingComponent : public UActorComponent, public IAIResourceInterf
 #endif
 
 protected:
+	/** Sets the current status and activates/deactivates the component if necessary */
+	AIMODULE_API void SetStatus(EPathFollowingStatus::Type InStatus);
 
 	/** deprecated associated movement component */
 	UE_DEPRECATED(5.5, "MovementComp is deprecated, please use NavMovementComp and the INavMoveInterface instead.")
@@ -512,6 +523,7 @@ protected:
 	 *	to acceptance radius */
 	int32 PreciseAcceptanceRadiusCheckStartNodeIndex;
 
+	UE_DEPRECATED_FORGAME(5.6, "This member should no longer be accessed directly, please use GetStatus/SetStatus instead.")
 	/** current status */
 	TEnumAsByte<EPathFollowingStatus::Type> Status;
 
@@ -553,6 +565,12 @@ protected:
 	/** True if the next segment is a custom link that has its own reach conditions. */
 	uint8 bMoveSegmentIsUsingCustomLinkReachCondition : 1;
 
+	/**
+	 * Indicates if the component activation/ticking is controlled by setting the Status (Moving --> Activated, other --> Deactivated)
+	 * Note that this optimization is not used for any derived classes since those might need to tick in other cases.
+	 */
+	uint8 bTickComponentOnlyWhenMoving : 1;
+
 	/** detect blocked movement when distance between center of location samples and furthest one (centroid radius) is below threshold */
 	float BlockDetectionDistance;
 
diff --git a/Engine/Source/Runtime/AIModule/Private/Navigation/CrowdFollowingComponent.cpp b/Engine/Source/Runtime/AIModule/Private/Navigation/CrowdFollowingComponent.cpp
--- a/Engine/Source/Runtime/AIModule/Private/Navigation/CrowdFollowingComponent.cpp
+++ b/Engine/Source/Runtime/AIModule/Private/Navigation/CrowdFollowingComponent.cpp
@@ -54,7 +54,7 @@ FVector UCrowdFollowingComponent::GetCrowdAgentLocation() const
 FVector UCrowdFollowingComponent::GetCrowdAgentVelocity() const
 {
 	FVector Velocity(NavMovementInterface.IsValid() ? NavMovementInterface->GetVelocityForNavMovement() : FVector::ZeroVector);
-	Velocity *= (Status == EPathFollowingStatus::Moving) ? 1.0f : 0.25f;
+	Velocity *= (GetStatus() == EPathFollowingStatus::Moving) ? 1.0f : 0.25f;
 	return Velocity;
 }
 
@@ -383,7 +383,7 @@ void UCrowdFollowingComponent::UpdateDestinationForMovingGoal(const FVector& New
 void UCrowdFollowingComponent::ApplyCrowdAgentVelocity(const FVector& NewVelocity, const FVector& DestPathCorner, bool bTraversingLink, bool bIsNearEndOfPath)
 {
 	bCanCheckMovingTooFar = !bTraversingLink && bIsNearEndOfPath;
-	if (IsCrowdSimulationEnabled() && Status == EPathFollowingStatus::Moving && NavMovementInterface.IsValid())
+	if (IsCrowdSimulationEnabled() && GetStatus() == EPathFollowingStatus::Moving && NavMovementInterface.IsValid())
 	{
 		const bool bIsNotFalling = (NavMovementInterface == nullptr || !NavMovementInterface->IsFalling());
 		if (bAffectFallingVelocity || bIsNotFalling)
@@ -513,7 +513,7 @@ void UCrowdFollowingComponent::BeginDestroy()
 
 void UCrowdFollowingComponent::AbortMove(const UObject& Instigator, FPathFollowingResultFlags::Type AbortFlags, FAIRequestID RequestID, EPathFollowingVelocityMode VelocityMode)
 {
-	if (IsCrowdSimulationEnabled() && (Status != EPathFollowingStatus::Idle) && RequestID.IsEquivalent(GetCurrentRequestId()))
+	if (IsCrowdSimulationEnabled() && (GetStatus() != EPathFollowingStatus::Idle) && RequestID.IsEquivalent(GetCurrentRequestId()))
 	{
 		UCrowdManager* CrowdManager = UCrowdManager::GetCurrent(GetWorld());
 		if (CrowdManager)
@@ -527,7 +527,7 @@ void UCrowdFollowingComponent::AbortMove(const UObject& Instigator, FPathFollowi
 
 void UCrowdFollowingComponent::PauseMove(FAIRequestID RequestID, EPathFollowingVelocityMode VelocityMode)
 {
-	if (IsCrowdSimulationEnabled() && (Status != EPathFollowingStatus::Paused) && RequestID.IsEquivalent(GetCurrentRequestId()))
+	if (IsCrowdSimulationEnabled() && (GetStatus() != EPathFollowingStatus::Paused) && RequestID.IsEquivalent(GetCurrentRequestId()))
 	{
 		UCrowdManager* CrowdManager = UCrowdManager::GetCurrent(GetWorld());
 		if (CrowdManager)
@@ -541,7 +541,7 @@ void UCrowdFollowingComponent::PauseMove(FAIRequestID RequestID, EPathFollowingV
 
 void UCrowdFollowingComponent::ResumeMove(FAIRequestID RequestID)
 {
-	if (IsCrowdSimulationEnabled() && (Status == EPathFollowingStatus::Paused) && RequestID.IsEquivalent(GetCurrentRequestId()))
+	if (IsCrowdSimulationEnabled() && (GetStatus() == EPathFollowingStatus::Paused) && RequestID.IsEquivalent(GetCurrentRequestId()))
 	{
 		UCrowdManager* CrowdManager = UCrowdManager::GetCurrent(GetWorld());
 		if (CrowdManager)
@@ -935,7 +935,7 @@ void UCrowdFollowingComponent::UpdatePathSegment()
 	// if agent has control over its movement, check finish conditions
 	const FVector CurrentLocation = NavMovementInterface->GetFeetLocation();
 	const bool bCanReachTarget = NavMovementInterface->CanStopPathFollowing();
-	if (bCanReachTarget && Status == EPathFollowingStatus::Moving)
+	if (bCanReachTarget && GetStatus() == EPathFollowingStatus::Moving)
 	{
 		const FVector GoalLocation = GetCurrentTargetLocation();
 
@@ -987,11 +987,11 @@ void UCrowdFollowingComponent::UpdatePathSegment()
 		}
 	}
 
-	if (bCanReachTarget && Status == EPathFollowingStatus::Moving)
+	if (bCanReachTarget && GetStatus() == EPathFollowingStatus::Moving)
 	{
 		// check waypoint switch condition in meta paths
 		FMetaNavMeshPath* MetaNavPath = bIsUsingMetaPath ? Path->CastPath<FMetaNavMeshPath>() : nullptr;
-		if (MetaNavPath && Status == EPathFollowingStatus::Moving)
+		if (MetaNavPath && GetStatus() == EPathFollowingStatus::Moving)
 		{
 			MetaNavPath->ConditionalMoveToNextSection(CurrentLocation, EMetaPathUpdateReason::MoveTick);
 		}
@@ -1045,7 +1045,7 @@ FVector UCrowdFollowingComponent::GetMoveFocus(bool bAllowStrafe) const
 
 		// if we're not moving, falling, or don't have a crowd agent move direction, set our focus to ahead of the rotation of our owner to keep the same rotation,
 		// otherwise use the Crowd Agent Move Direction to move in the direction we're supposed to be going
-		const FVector ForwardDir = NavMovementInterface->GetOwnerAsObject() && ((Status != EPathFollowingStatus::Moving) || (NavMovementInterface->IsFalling()) || CrowdAgentMoveDirection.IsNearlyZero()) ?
+		const FVector ForwardDir = NavMovementInterface->GetOwnerAsObject() && ((GetStatus() != EPathFollowingStatus::Moving) || (NavMovementInterface->IsFalling()) || CrowdAgentMoveDirection.IsNearlyZero()) ?
 			NavMovementInterface->GetForwardVector() :
 			CrowdAgentMoveDirection.GetSafeNormal2D();
 
@@ -1057,7 +1057,7 @@ FVector UCrowdFollowingComponent::GetMoveFocus(bool bAllowStrafe) const
 
 void UCrowdFollowingComponent::OnNavNodeChanged(NavNodeRef NewPolyRef, NavNodeRef PrevPolyRef, int32 CorridorSize)
 {
-	if (IsCrowdSimulationEnabled() && Status != EPathFollowingStatus::Idle)
+	if (IsCrowdSimulationEnabled() && GetStatus() != EPathFollowingStatus::Idle)
 	{
 		// update last visited path poly
 		FNavMeshPath* NavPath = Path.IsValid() ? Path->CastPath<FNavMeshPath>() : NULL;
@@ -1144,7 +1144,7 @@ void UCrowdFollowingComponent::GetDebugStringTokens(TArray<FString>& Tokens, TAr
 		Flags.Add(EPathFollowingDebugTokens::FailedValue);
 	}
 
-	if (Status != EPathFollowingStatus::Moving)
+	if (GetStatus() != EPathFollowingStatus::Moving)
 	{
 		return;
 	}
@@ -1226,7 +1226,7 @@ void UCrowdFollowingComponent::DescribeSelfToVisLog(FVisualLogEntry* Snapshot) c
 	FNavMeshPath* NavMeshPath = Path.IsValid() ? Path->CastPath<FNavMeshPath>() : NULL;
 	FAbstractNavigationPath* DirectPath = Path.IsValid() ? Path->CastPath<FAbstractNavigationPath>() : NULL;
 
-	if (Status == EPathFollowingStatus::Moving)
+	if (GetStatus() == EPathFollowingStatus::Moving)
 	{
 		StatusDesc += FString::Printf(TEXT(" [path:%d, visited:%d]"), PathStartIndex, LastPathPolyIndex);
 	}
diff --git a/Engine/Source/Runtime/AIModule/Private/Navigation/GridPathFollowingComponent.cpp b/Engine/Source/Runtime/AIModule/Private/Navigation/GridPathFollowingComponent.cpp
--- a/Engine/Source/Runtime/AIModule/Private/Navigation/GridPathFollowingComponent.cpp
+++ b/Engine/Source/Runtime/AIModule/Private/Navigation/GridPathFollowingComponent.cpp
@@ -179,7 +179,7 @@ void UGridPathFollowingComponent::ResumeMove(FAIRequestID RequestID)
 		if (HasActiveGrid())
 		{
 			UE_VLOG(GetOwner(), LogPathFollowing, Log, TEXT("ResumeMove: RequestID(%u) is on grid"), RequestID.GetID());
-			Status = EPathFollowingStatus::Moving;
+			SetStatus(EPathFollowingStatus::Moving);
 		}
 		else
 		{
diff --git a/Engine/Source/Runtime/AIModule/Private/Navigation/PathFollowingComponent.cpp b/Engine/Source/Runtime/AIModule/Private/Navigation/PathFollowingComponent.cpp
--- a/Engine/Source/Runtime/AIModule/Private/Navigation/PathFollowingComponent.cpp
+++ b/Engine/Source/Runtime/AIModule/Private/Navigation/PathFollowingComponent.cpp
@@ -127,7 +127,11 @@ const float UPathFollowingComponent::DefaultAcceptanceRadius = -1.f;
 UPathFollowingComponent::UPathFollowingComponent(const FObjectInitializer& ObjectInitializer) : Super(ObjectInitializer)
 {
 	PrimaryComponentTick.bCanEverTick = true;
-	bAutoActivate = true;
+
+	// Activation/Ticking is controlled by Status if this instance is of type UPathFollowingComponent (not a derived class)
+	bTickComponentOnlyWhenMoving = StaticClass() == GetClass();
+	PrimaryComponentTick.bStartWithTickEnabled = !bTickComponentOnlyWhenMoving;
+	bAutoActivate = !bTickComponentOnlyWhenMoving;
 
 	MinAgentRadiusPct = 1.1f;
 	MinAgentHalfHeightPct = 1.05f;
@@ -165,6 +169,29 @@ UPathFollowingComponent::UPathFollowingComponent(const FObjectInitializer& Objec
 #endif // !UE_BUILD_SHIPPING
 }
 
+void UPathFollowingComponent::SetStatus(const EPathFollowingStatus::Type InStatus)
+{
+	Status = InStatus;
+
+	if (bTickComponentOnlyWhenMoving)
+	{
+		if (Status == EPathFollowingStatus::Moving)
+		{
+			if (!IsActive())
+			{
+				Activate();
+			}
+		}
+		else
+		{
+			if (IsActive())
+			{
+				Deactivate();
+			}
+		}
+	}
+}
+
 void UPathFollowingComponent::LogPathHelper(const AActor* LogOwner, FNavigationPath* InLogPath, const AActor* LogGoalActor)
 {
 #if ENABLE_VISUAL_LOG
@@ -302,7 +329,7 @@ bool UPathFollowingComponent::HandlePathUpdateEvent()
 
 	if (Status == EPathFollowingStatus::Waiting || Status == EPathFollowingStatus::Moving)
 	{
-		Status = EPathFollowingStatus::Moving;
+		SetStatus(EPathFollowingStatus::Moving);
 
 		const int32 CurrentSegment = DetermineStartingPathPoint(Path.Get());
 		SetMoveSegment(CurrentSegment);
@@ -422,7 +449,7 @@ FAIRequestID UPathFollowingComponent::RequestMove(const FAIMoveRequest& RequestD
 			// with async pathfinding paths can be incomplete, movement will start after receiving path event 
 			if (!bIsUsingMetaPath && Path.IsValid() && Path->IsValid())
 			{
-				Status = EPathFollowingStatus::Moving;
+				SetStatus(EPathFollowingStatus::Moving);
 
 				// determine with path segment should be followed
 				const uint32 CurrentSegment = DetermineStartingPathPoint(InPath.Get());
@@ -430,7 +457,7 @@ FAIRequestID UPathFollowingComponent::RequestMove(const FAIMoveRequest& RequestD
 			}
 			else
 			{
-				Status = EPathFollowingStatus::Waiting;
+				SetStatus(EPathFollowingStatus::Waiting);
 				GetWorld()->GetTimerManager().SetTimer(WaitingForPathTimer, this, &UPathFollowingComponent::OnWaitingPathTimeout, WaitingTimeout);
 			}
 		}
@@ -498,7 +525,7 @@ void UPathFollowingComponent::ResumeMove(FAIRequestID RequestID)
 		const bool bIsOnPath = IsOnPath();
 		if (bIsOnPath)
 		{
-			Status = EPathFollowingStatus::Moving;
+			SetStatus(EPathFollowingStatus::Moving);
 
 			const bool bWasPathUpdatedRecently = Path.IsValid() ? (Path->GetTimeStamp() > PathTimeWhenPaused) : false;
 			if (bMovedDuringPause || bWasPathUpdatedRecently)
@@ -514,7 +541,7 @@ void UPathFollowingComponent::ResumeMove(FAIRequestID RequestID)
 		else if (Path.IsValid() && Path->IsValid() && Path->GetNavigationDataUsed() == NULL)
 		{
 			// this means it's a scripted path, just resume
-			Status = EPathFollowingStatus::Moving;
+			SetStatus(EPathFollowingStatus::Moving);
 			UpdateMoveFocus();
 		}
 		else
@@ -797,7 +824,7 @@ void UPathFollowingComponent::Reset()
 	OriginalMoveRequestGoalLocation = FAISystem::InvalidLocation;
 
 	CurrentRequestId = FAIRequestID::InvalidRequest;
-	Status = EPathFollowingStatus::Idle;
+	SetStatus(EPathFollowingStatus::Idle);
 
 	if (WaitingForPathTimer.IsValid())
 	{
