[BehaviorTree] replaced some ensure on blackboard asset by error reporting since this is a possible scenario in Editor
Bonus: some QOL improvements to refresh all error messages when assigning a Blackboard and updating error message on newly added nodes

diff --git a/Engine/Plugins/Experimental/GameplayBehaviors/Source/GameplayBehaviorsModule/Private/AI/BTDecorator_GameplayTagQuery.cpp b/Engine/Plugins/Experimental/GameplayBehaviors/Source/GameplayBehaviorsModule/Private/AI/BTDecorator_GameplayTagQuery.cpp
--- a/Engine/Plugins/Experimental/GameplayBehaviors/Source/GameplayBehaviorsModule/Private/AI/BTDecorator_GameplayTagQuery.cpp
+++ b/Engine/Plugins/Experimental/GameplayBehaviors/Source/GameplayBehaviorsModule/Private/AI/BTDecorator_GameplayTagQuery.cpp
@@ -145,16 +145,27 @@ void UBTDecorator_GameplayTagQuery::PostEditChangeProperty(FPropertyChangedEvent
 
 	CacheGameplayTagsInsideQuery();
 }
+
+FString UBTDecorator_GameplayTagQuery::GetErrorMessage() const
+{
+	if (GetBlackboardAsset() == nullptr)
+	{
+		return UE::BehaviorTree::Messages::BlackboardNotSet.ToString();
+	}
+	return Super::GetErrorMessage();
+}
 #endif	// WITH_EDITOR
 
 void UBTDecorator_GameplayTagQuery::InitializeFromAsset(UBehaviorTree& Asset)
 {
 	Super::InitializeFromAsset(Asset);
 
-	UBlackboardData* BBAsset = GetBlackboardAsset();
-	if (ensure(BBAsset))
+	if (const UBlackboardData* BBAsset = GetBlackboardAsset())
 	{
 		ActorForGameplayTagQuery.ResolveSelectedKey(*BBAsset);
 	}
+	else
+	{
+		ActorForGameplayTagQuery.InvalidateResolvedKey();
+	}
 }
-
diff --git a/Engine/Plugins/Experimental/GameplayBehaviors/Source/GameplayBehaviorsModule/Public/AI/BTDecorator_GameplayTagQuery.h b/Engine/Plugins/Experimental/GameplayBehaviors/Source/GameplayBehaviorsModule/Public/AI/BTDecorator_GameplayTagQuery.h
--- a/Engine/Plugins/Experimental/GameplayBehaviors/Source/GameplayBehaviorsModule/Public/AI/BTDecorator_GameplayTagQuery.h
+++ b/Engine/Plugins/Experimental/GameplayBehaviors/Source/GameplayBehaviorsModule/Public/AI/BTDecorator_GameplayTagQuery.h
@@ -58,7 +58,7 @@ protected:
 		do this only once and cache it.
 	*/
 	virtual void CacheGameplayTagsInsideQuery();
-
+	virtual FString GetErrorMessage() const override;
 	virtual void PostEditChangeProperty(FPropertyChangedEvent& PropertyChangedEvent) override;
 #endif
 
diff --git a/Engine/Source/Editor/AIGraph/Classes/AIGraph.h b/Engine/Source/Editor/AIGraph/Classes/AIGraph.h
--- a/Engine/Source/Editor/AIGraph/Classes/AIGraph.h
+++ b/Engine/Source/Editor/AIGraph/Classes/AIGraph.h
@@ -35,7 +35,12 @@ class AIGRAPH_API UAIGraph : public UEdGraph
 	virtual void OnNodesPasted(const FString& ImportStr);
 
 	bool UpdateUnknownNodeClasses();
+
+	UE_DEPRECATED(5.6, "This method was only used to fetch error messages and log them. It has been renamed to better fit its purpose. Use UpdateErrorMessages instead.")
 	void UpdateDeprecatedClasses();
+
+	/** Sets error messages and logs errors about nodes. */
+	void UpdateErrorMessages();
 	void RemoveOrphanedNodes();
 	void UpdateClassData();
 
diff --git a/Engine/Source/Editor/AIGraph/Private/AIGraph.cpp b/Engine/Source/Editor/AIGraph/Private/AIGraph.cpp
--- a/Engine/Source/Editor/AIGraph/Private/AIGraph.cpp
+++ b/Engine/Source/Editor/AIGraph/Private/AIGraph.cpp
@@ -37,7 +37,7 @@ void UAIGraph::OnCreated()
 
 void UAIGraph::OnLoaded()
 {
-	UpdateDeprecatedClasses();
+	UpdateErrorMessages();
 	UpdateUnknownNodeClasses();
 }
 
@@ -123,10 +123,14 @@ void UpdateAIGraphNodeErrorMessage(UAIGraphNode& Node)
 	}
 }
 
+// deprecated
 void UAIGraph::UpdateDeprecatedClasses()
 {
-	// This function sets error messages and logs errors about nodes.
+	UpdateErrorMessages();
+}
 
+void UAIGraph::UpdateErrorMessages()
+{
 	for (int32 Idx = 0, IdxNum = Nodes.Num(); Idx < IdxNum; ++Idx)
 	{
 		UAIGraphNode* Node = Cast<UAIGraphNode>(Nodes[Idx]);
@@ -150,12 +154,12 @@ void UAIGraph::Serialize(FArchive& Ar)
 	// Overridden to flags up errors in the behavior tree while cooking.
 	Super::Serialize(Ar);
 
-	// Execute UpdateDeprecatedClasses only when saving to persistent storage,
+	// Execute UpdateErrorMessages only when saving to persistent storage,
 	// otherwise node instances might not be fully created (i.e. transaction buffer while loading the asset).
 	if ((Ar.IsSaving() && Ar.IsPersistent()) || Ar.IsCooking())
 	{
-		// Logging of errors happens in UpdateDeprecatedClasses
-		UpdateDeprecatedClasses();
+		// Logging of errors happens in UpdateErrorMessages
+		UpdateErrorMessages();
 	}
 }
 
diff --git a/Engine/Source/Editor/AIGraph/Private/AIGraphNode.cpp b/Engine/Source/Editor/AIGraph/Private/AIGraphNode.cpp
--- a/Engine/Source/Editor/AIGraph/Private/AIGraphNode.cpp
+++ b/Engine/Source/Editor/AIGraph/Private/AIGraphNode.cpp
@@ -41,6 +41,7 @@ void UAIGraphNode::PostPlacedNewNode()
 			NodeInstance = NewObject<UObject>(GraphOwner, NodeClass);
 			NodeInstance->SetFlags(RF_Transactional);
 			InitializeInstance();
+			UpdateErrorMessage();
 		}
 	}
 }
diff --git a/Engine/Source/Editor/BehaviorTreeEditor/Private/BehaviorTreeGraph.cpp b/Engine/Source/Editor/BehaviorTreeEditor/Private/BehaviorTreeGraph.cpp
--- a/Engine/Source/Editor/BehaviorTreeEditor/Private/BehaviorTreeGraph.cpp
+++ b/Engine/Source/Editor/BehaviorTreeEditor/Private/BehaviorTreeGraph.cpp
@@ -93,6 +93,8 @@ void UBehaviorTreeGraph::UpdateBlackboardChange()
 			}
 		}
 	}
+
+	UpdateErrorMessages();
 }
 
 void UBehaviorTreeGraph::UpdateAsset(int32 UpdateFlags)
diff --git a/Engine/Source/Runtime/AIModule/Classes/BehaviorTree/BehaviorTreeTypes.h b/Engine/Source/Runtime/AIModule/Classes/BehaviorTree/BehaviorTreeTypes.h
--- a/Engine/Source/Runtime/AIModule/Classes/BehaviorTree/BehaviorTreeTypes.h
+++ b/Engine/Source/Runtime/AIModule/Classes/BehaviorTree/BehaviorTreeTypes.h
@@ -51,6 +51,13 @@ template <> struct TIsValidVariadicFunctionArg<FBlackboard::FKey>
 };
 #endif
 
+#if WITH_EDITORONLY_DATA
+namespace UE::BehaviorTree::Messages
+{
+	extern AIMODULE_API FText BlackboardNotSet;
+}
+#endif // WITH_EDITORONLY_DATA
+
 enum class EBlackboardNotificationResult : uint8
 {
 	RemoveObserver,
diff --git a/Engine/Source/Runtime/AIModule/Classes/BehaviorTree/Decorators/BTDecorator_BlackboardBase.h b/Engine/Source/Runtime/AIModule/Classes/BehaviorTree/Decorators/BTDecorator_BlackboardBase.h
--- a/Engine/Source/Runtime/AIModule/Classes/BehaviorTree/Decorators/BTDecorator_BlackboardBase.h
+++ b/Engine/Source/Runtime/AIModule/Classes/BehaviorTree/Decorators/BTDecorator_BlackboardBase.h
@@ -24,6 +24,7 @@ class UBTDecorator_BlackboardBase : public UBTDecorator
 
 #if WITH_EDITOR
 	AIMODULE_API virtual FName GetNodeIconName() const override;
+	AIMODULE_API virtual FString GetErrorMessage() const override;
 #endif
 
 public:
diff --git a/Engine/Source/Runtime/AIModule/Classes/BehaviorTree/Decorators/BTDecorator_CheckGameplayTagsOnActor.h b/Engine/Source/Runtime/AIModule/Classes/BehaviorTree/Decorators/BTDecorator_CheckGameplayTagsOnActor.h
--- a/Engine/Source/Runtime/AIModule/Classes/BehaviorTree/Decorators/BTDecorator_CheckGameplayTagsOnActor.h
+++ b/Engine/Source/Runtime/AIModule/Classes/BehaviorTree/Decorators/BTDecorator_CheckGameplayTagsOnActor.h
@@ -43,7 +43,7 @@ protected:
 #if WITH_EDITOR
 	/** describe decorator and cache it */
 	AIMODULE_API virtual void BuildDescription();
-
+	AIMODULE_API virtual FString GetErrorMessage() const override;
 	AIMODULE_API virtual void PostEditChangeProperty(FPropertyChangedEvent& PropertyChangedEvent) override;
 #endif
 
diff --git a/Engine/Source/Runtime/AIModule/Classes/BehaviorTree/Decorators/BTDecorator_CompareBBEntries.h b/Engine/Source/Runtime/AIModule/Classes/BehaviorTree/Decorators/BTDecorator_CompareBBEntries.h
--- a/Engine/Source/Runtime/AIModule/Classes/BehaviorTree/Decorators/BTDecorator_CompareBBEntries.h
+++ b/Engine/Source/Runtime/AIModule/Classes/BehaviorTree/Decorators/BTDecorator_CompareBBEntries.h
@@ -58,5 +58,6 @@ public:
 
 #if WITH_EDITOR
 	AIMODULE_API virtual FName GetNodeIconName() const override;
+	AIMODULE_API virtual FString GetErrorMessage() const override;
 #endif // WITH_EDITOR
 };
diff --git a/Engine/Source/Runtime/AIModule/Classes/BehaviorTree/Decorators/BTDecorator_ConeCheck.h b/Engine/Source/Runtime/AIModule/Classes/BehaviorTree/Decorators/BTDecorator_ConeCheck.h
--- a/Engine/Source/Runtime/AIModule/Classes/BehaviorTree/Decorators/BTDecorator_ConeCheck.h
+++ b/Engine/Source/Runtime/AIModule/Classes/BehaviorTree/Decorators/BTDecorator_ConeCheck.h
@@ -54,6 +54,7 @@ class UBTDecorator_ConeCheck : public UBTDecorator
 
 #if WITH_EDITOR
 	AIMODULE_API virtual FName GetNodeIconName() const override;
+	AIMODULE_API virtual FString GetErrorMessage() const override;
 #endif // WITH_EDITOR
 
 protected:
diff --git a/Engine/Source/Runtime/AIModule/Classes/BehaviorTree/Decorators/BTDecorator_DoesPathExist.h b/Engine/Source/Runtime/AIModule/Classes/BehaviorTree/Decorators/BTDecorator_DoesPathExist.h
--- a/Engine/Source/Runtime/AIModule/Classes/BehaviorTree/Decorators/BTDecorator_DoesPathExist.h
+++ b/Engine/Source/Runtime/AIModule/Classes/BehaviorTree/Decorators/BTDecorator_DoesPathExist.h
@@ -65,5 +65,6 @@ public:
 
 #if WITH_EDITOR
 	AIMODULE_API virtual FName GetNodeIconName() const override;
+	AIMODULE_API virtual FString GetErrorMessage() const override;
 #endif // WITH_EDITOR
 };
diff --git a/Engine/Source/Runtime/AIModule/Classes/BehaviorTree/Decorators/BTDecorator_KeepInCone.h b/Engine/Source/Runtime/AIModule/Classes/BehaviorTree/Decorators/BTDecorator_KeepInCone.h
--- a/Engine/Source/Runtime/AIModule/Classes/BehaviorTree/Decorators/BTDecorator_KeepInCone.h
+++ b/Engine/Source/Runtime/AIModule/Classes/BehaviorTree/Decorators/BTDecorator_KeepInCone.h
@@ -58,6 +58,7 @@ class UBTDecorator_KeepInCone : public UBTDecorator
 
 #if WITH_EDITOR
 	AIMODULE_API virtual FName GetNodeIconName() const override;
+	AIMODULE_API virtual FString GetErrorMessage() const override;
 #endif // WITH_EDITOR
 
 protected:
diff --git a/Engine/Source/Runtime/AIModule/Classes/BehaviorTree/Services/BTService_BlackboardBase.h b/Engine/Source/Runtime/AIModule/Classes/BehaviorTree/Services/BTService_BlackboardBase.h
--- a/Engine/Source/Runtime/AIModule/Classes/BehaviorTree/Services/BTService_BlackboardBase.h
+++ b/Engine/Source/Runtime/AIModule/Classes/BehaviorTree/Services/BTService_BlackboardBase.h
@@ -21,6 +21,10 @@ class UBTService_BlackboardBase : public UBTService
 
 protected:
 
+#if WITH_EDITOR
+	AIMODULE_API virtual FString GetErrorMessage() const override;
+#endif
+
 	/** blackboard key selector */
 	UPROPERTY(EditAnywhere, Category=Blackboard)
 	struct FBlackboardKeySelector BlackboardKey;
diff --git a/Engine/Source/Runtime/AIModule/Classes/BehaviorTree/Tasks/BTTask_BlackboardBase.h b/Engine/Source/Runtime/AIModule/Classes/BehaviorTree/Tasks/BTTask_BlackboardBase.h
--- a/Engine/Source/Runtime/AIModule/Classes/BehaviorTree/Tasks/BTTask_BlackboardBase.h
+++ b/Engine/Source/Runtime/AIModule/Classes/BehaviorTree/Tasks/BTTask_BlackboardBase.h
@@ -22,6 +22,10 @@ class UBTTask_BlackboardBase : public UBTTaskNode
 
 protected:
 
+#if WITH_EDITOR
+	AIMODULE_API virtual FString GetErrorMessage() const override;
+#endif
+
 	/** blackboard key selector */
 	UPROPERTY(EditAnywhere, Category=Blackboard)
 	struct FBlackboardKeySelector BlackboardKey;
diff --git a/Engine/Source/Runtime/AIModule/Classes/BehaviorTree/Tasks/BTTask_WaitBlackboardTime.h b/Engine/Source/Runtime/AIModule/Classes/BehaviorTree/Tasks/BTTask_WaitBlackboardTime.h
--- a/Engine/Source/Runtime/AIModule/Classes/BehaviorTree/Tasks/BTTask_WaitBlackboardTime.h
+++ b/Engine/Source/Runtime/AIModule/Classes/BehaviorTree/Tasks/BTTask_WaitBlackboardTime.h
@@ -26,6 +26,10 @@ class UBTTask_WaitBlackboardTime : public UBTTask_Wait
 
 protected:
 
+#if WITH_EDITOR
+	AIMODULE_API virtual FString GetErrorMessage() const override;
+#endif
+
 	/** blackboard key selector */
 	UPROPERTY(EditAnywhere, Category=Blackboard)
 	struct FBlackboardKeySelector BlackboardKey;
diff --git a/Engine/Source/Runtime/AIModule/Private/BehaviorTree/BehaviorTreeTypes.cpp b/Engine/Source/Runtime/AIModule/Private/BehaviorTree/BehaviorTreeTypes.cpp
--- a/Engine/Source/Runtime/AIModule/Private/BehaviorTree/BehaviorTreeTypes.cpp
+++ b/Engine/Source/Runtime/AIModule/Private/BehaviorTree/BehaviorTreeTypes.cpp
@@ -24,6 +24,10 @@
 
 #include UE_INLINE_GENERATED_CPP_BY_NAME(BehaviorTreeTypes)
 
+#if WITH_EDITORONLY_DATA
+FText UE::BehaviorTree::Messages::BlackboardNotSet(NSLOCTEXT("BehaviorTree", "BlackboardNotSet", "Blackboard not set"));
+#endif
+
 //----------------------------------------------------------------------//
 // FBehaviorTreeInstance
 //----------------------------------------------------------------------//
diff --git a/Engine/Source/Runtime/AIModule/Private/BehaviorTree/Decorators/BTDecorator_BlackboardBase.cpp b/Engine/Source/Runtime/AIModule/Private/BehaviorTree/Decorators/BTDecorator_BlackboardBase.cpp
--- a/Engine/Source/Runtime/AIModule/Private/BehaviorTree/Decorators/BTDecorator_BlackboardBase.cpp
+++ b/Engine/Source/Runtime/AIModule/Private/BehaviorTree/Decorators/BTDecorator_BlackboardBase.cpp
@@ -21,14 +21,12 @@ void UBTDecorator_BlackboardBase::InitializeFromAsset(UBehaviorTree& Asset)
 {
 	Super::InitializeFromAsset(Asset);
 
-	UBlackboardData* BBAsset = GetBlackboardAsset();
-	if (BBAsset)
+	if (const UBlackboardData* BBAsset = GetBlackboardAsset())
 	{
 		BlackboardKey.ResolveSelectedKey(*BBAsset);
 	}
 	else
 	{
-		UE_LOG(LogBehaviorTree, Warning, TEXT("Can't initialize %s due to missing blackboard data."), *GetName());
 		BlackboardKey.InvalidateResolvedKey();
 	}
 }
@@ -74,5 +72,14 @@ FName UBTDecorator_BlackboardBase::GetNodeIconName() const
 	return FName("BTEditor.Graph.BTNode.Decorator.Blackboard.Icon");
 }
 
+FString UBTDecorator_BlackboardBase::GetErrorMessage() const
+{
+	if (GetBlackboardAsset() == nullptr)
+	{
+		return UE::BehaviorTree::Messages::BlackboardNotSet.ToString();
+	}
+	return Super::GetErrorMessage();
+}
+
 #endif	// WITH_EDITOR
 
diff --git a/Engine/Source/Runtime/AIModule/Private/BehaviorTree/Decorators/BTDecorator_CheckGameplayTagsOnActor.cpp b/Engine/Source/Runtime/AIModule/Private/BehaviorTree/Decorators/BTDecorator_CheckGameplayTagsOnActor.cpp
--- a/Engine/Source/Runtime/AIModule/Private/BehaviorTree/Decorators/BTDecorator_CheckGameplayTagsOnActor.cpp
+++ b/Engine/Source/Runtime/AIModule/Private/BehaviorTree/Decorators/BTDecorator_CheckGameplayTagsOnActor.cpp
@@ -82,17 +82,29 @@ void UBTDecorator_CheckGameplayTagsOnActor::PostEditChangeProperty(FPropertyChan
 	BuildDescription();
 }
 
+FString UBTDecorator_CheckGameplayTagsOnActor::GetErrorMessage() const
+{
+	if (GetBlackboardAsset() == nullptr)
+	{
+		return UE::BehaviorTree::Messages::BlackboardNotSet.ToString();
+	}
+	return Super::GetErrorMessage();
+}
+
 #endif	// WITH_EDITOR
 
 void UBTDecorator_CheckGameplayTagsOnActor::InitializeFromAsset(UBehaviorTree& Asset)
 {
 	Super::InitializeFromAsset(Asset);
 
-	UBlackboardData* BBAsset = GetBlackboardAsset();
-	if (ensure(BBAsset))
+	if (const UBlackboardData* BBAsset = GetBlackboardAsset())
 	{
 		ActorToCheck.ResolveSelectedKey(*BBAsset);
 	}
+	else
+	{
+		ActorToCheck.InvalidateResolvedKey();
+	}
 
 #if WITH_EDITOR
 	BuildDescription();
diff --git a/Engine/Source/Runtime/AIModule/Private/BehaviorTree/Decorators/BTDecorator_CompareBBEntries.cpp b/Engine/Source/Runtime/AIModule/Private/BehaviorTree/Decorators/BTDecorator_CompareBBEntries.cpp
--- a/Engine/Source/Runtime/AIModule/Private/BehaviorTree/Decorators/BTDecorator_CompareBBEntries.cpp
+++ b/Engine/Source/Runtime/AIModule/Private/BehaviorTree/Decorators/BTDecorator_CompareBBEntries.cpp
@@ -17,12 +17,16 @@ void UBTDecorator_CompareBBEntries::InitializeFromAsset(UBehaviorTree& Asset)
 {
 	Super::InitializeFromAsset(Asset);
 
-	UBlackboardData* BBAsset = GetBlackboardAsset();
-	if (ensure(BBAsset))
+	if (const UBlackboardData* BBAsset = GetBlackboardAsset())
 	{
 		BlackboardKeyA.ResolveSelectedKey(*BBAsset);
 		BlackboardKeyB.ResolveSelectedKey(*BBAsset);
 	}
+	else
+	{
+		BlackboardKeyA.InvalidateResolvedKey();
+		BlackboardKeyB.InvalidateResolvedKey();
+	}
 }
 
 // @note I know it's ugly to have "return" statements in many places inside a function, but the way 
@@ -97,5 +101,14 @@ FName UBTDecorator_CompareBBEntries::GetNodeIconName() const
 	return FName("BTEditor.Graph.BTNode.Decorator.CompareBlackboardEntries.Icon");
 }
 
+FString UBTDecorator_CompareBBEntries::GetErrorMessage() const
+{
+	if (GetBlackboardAsset() == nullptr)
+	{
+		return UE::BehaviorTree::Messages::BlackboardNotSet.ToString();
+	}
+	return Super::GetErrorMessage();
+}
+
 #endif	// WITH_EDITOR
 
diff --git a/Engine/Source/Runtime/AIModule/Private/BehaviorTree/Decorators/BTDecorator_ConeCheck.cpp b/Engine/Source/Runtime/AIModule/Private/BehaviorTree/Decorators/BTDecorator_ConeCheck.cpp
--- a/Engine/Source/Runtime/AIModule/Private/BehaviorTree/Decorators/BTDecorator_ConeCheck.cpp
+++ b/Engine/Source/Runtime/AIModule/Private/BehaviorTree/Decorators/BTDecorator_ConeCheck.cpp
@@ -31,13 +31,18 @@ void UBTDecorator_ConeCheck::InitializeFromAsset(UBehaviorTree& Asset)
 {
 	Super::InitializeFromAsset(Asset);
 
-	UBlackboardData* BBAsset = GetBlackboardAsset();
-	if (ensure(BBAsset))
+	if (const UBlackboardData* BBAsset = GetBlackboardAsset())
 	{
 		ConeOrigin.ResolveSelectedKey(*BBAsset);
 		ConeDirection.ResolveSelectedKey(*BBAsset);
 		Observed.ResolveSelectedKey(*BBAsset);
 	}
+	else
+	{
+		ConeOrigin.InvalidateResolvedKey();
+		ConeDirection.InvalidateResolvedKey();
+		Observed.InvalidateResolvedKey();
+	}
 }
 
 bool UBTDecorator_ConeCheck::CalculateDirection(const UBlackboardComponent* BlackboardComp, const FBlackboardKeySelector& Origin, const FBlackboardKeySelector& End, FVector& Direction) const
@@ -161,5 +166,14 @@ FName UBTDecorator_ConeCheck::GetNodeIconName() const
 	return FName("BTEditor.Graph.BTNode.Decorator.ConeCheck.Icon");
 }
 
+FString UBTDecorator_ConeCheck::GetErrorMessage() const
+{
+	if (GetBlackboardAsset() == nullptr)
+	{
+		return UE::BehaviorTree::Messages::BlackboardNotSet.ToString();
+	}
+	return Super::GetErrorMessage();
+}
+
 #endif	// WITH_EDITOR
 
diff --git a/Engine/Source/Runtime/AIModule/Private/BehaviorTree/Decorators/BTDecorator_DoesPathExist.cpp b/Engine/Source/Runtime/AIModule/Private/BehaviorTree/Decorators/BTDecorator_DoesPathExist.cpp
--- a/Engine/Source/Runtime/AIModule/Private/BehaviorTree/Decorators/BTDecorator_DoesPathExist.cpp
+++ b/Engine/Source/Runtime/AIModule/Private/BehaviorTree/Decorators/BTDecorator_DoesPathExist.cpp
@@ -39,12 +39,16 @@ void UBTDecorator_DoesPathExist::InitializeFromAsset(UBehaviorTree& Asset)
 		bUseSelf = false;
 	}
 
-	UBlackboardData* BBAsset = GetBlackboardAsset();
-	if (ensure(BBAsset))
+	if (const UBlackboardData* BBAsset = GetBlackboardAsset())
 	{
 		BlackboardKeyA.ResolveSelectedKey(*BBAsset);
 		BlackboardKeyB.ResolveSelectedKey(*BBAsset);
 	}
+	else
+	{
+		BlackboardKeyA.InvalidateResolvedKey();
+		BlackboardKeyB.InvalidateResolvedKey();
+	}
 }
 
 bool UBTDecorator_DoesPathExist::CalculateRawConditionValue(UBehaviorTreeComponent& OwnerComp, uint8* NodeMemory) const
@@ -106,5 +110,14 @@ FName UBTDecorator_DoesPathExist::GetNodeIconName() const
 	return FName("BTEditor.Graph.BTNode.Decorator.DoesPathExist.Icon");
 }
 
+FString UBTDecorator_DoesPathExist::GetErrorMessage() const
+{
+	if (GetBlackboardAsset() == nullptr)
+	{
+		return UE::BehaviorTree::Messages::BlackboardNotSet.ToString();
+	}
+	return Super::GetErrorMessage();
+}
+
 #endif	// WITH_EDITOR
 
diff --git a/Engine/Source/Runtime/AIModule/Private/BehaviorTree/Decorators/BTDecorator_KeepInCone.cpp b/Engine/Source/Runtime/AIModule/Private/BehaviorTree/Decorators/BTDecorator_KeepInCone.cpp
--- a/Engine/Source/Runtime/AIModule/Private/BehaviorTree/Decorators/BTDecorator_KeepInCone.cpp
+++ b/Engine/Source/Runtime/AIModule/Private/BehaviorTree/Decorators/BTDecorator_KeepInCone.cpp
@@ -48,12 +48,16 @@ void UBTDecorator_KeepInCone::InitializeFromAsset(UBehaviorTree& Asset)
 		bUseSelfAsObserved = false;
 	}
 
-	UBlackboardData* BBAsset = GetBlackboardAsset();
-	if (ensure(BBAsset))
+	if (const UBlackboardData* BBAsset = GetBlackboardAsset())
 	{
 		ConeOrigin.ResolveSelectedKey(*BBAsset);
 		Observed.ResolveSelectedKey(*BBAsset);
 	}
+	else
+	{
+		ConeOrigin.InvalidateResolvedKey();
+		Observed.InvalidateResolvedKey();
+	}
 }
 
 bool UBTDecorator_KeepInCone::CalculateCurrentDirection(const UBehaviorTreeComponent& OwnerComp, FVector& Direction) const
@@ -153,5 +157,14 @@ FName UBTDecorator_KeepInCone::GetNodeIconName() const
 	return FName("BTEditor.Graph.BTNode.Decorator.KeepInCone.Icon");
 }
 
+FString UBTDecorator_KeepInCone::GetErrorMessage() const
+{
+	if (GetBlackboardAsset() == nullptr)
+	{
+		return UE::BehaviorTree::Messages::BlackboardNotSet.ToString();
+	}
+	return Super::GetErrorMessage();
+}
+
 #endif	// WITH_EDITOR
 
diff --git a/Engine/Source/Runtime/AIModule/Private/BehaviorTree/Services/BTService_BlackboardBase.cpp b/Engine/Source/Runtime/AIModule/Private/BehaviorTree/Services/BTService_BlackboardBase.cpp
--- a/Engine/Source/Runtime/AIModule/Private/BehaviorTree/Services/BTService_BlackboardBase.cpp
+++ b/Engine/Source/Runtime/AIModule/Private/BehaviorTree/Services/BTService_BlackboardBase.cpp
@@ -14,10 +14,24 @@ UBTService_BlackboardBase::UBTService_BlackboardBase(const FObjectInitializer& O
 void UBTService_BlackboardBase::InitializeFromAsset(UBehaviorTree& Asset)
 {
 	Super::InitializeFromAsset(Asset);
-	UBlackboardData* BBAsset = GetBlackboardAsset();
-	if (ensure(BBAsset))
+
+	if (const UBlackboardData* BBAsset = GetBlackboardAsset())
 	{
 		BlackboardKey.ResolveSelectedKey(*BBAsset);
 	}
+	else
+	{
+		BlackboardKey.InvalidateResolvedKey();
+	}
 }
 
+#if WITH_EDITOR
+FString UBTService_BlackboardBase::GetErrorMessage() const
+{
+	if (GetBlackboardAsset() == nullptr)
+	{
+		return UE::BehaviorTree::Messages::BlackboardNotSet.ToString();
+	}
+	return Super::GetErrorMessage();
+}
+#endif // WITH_EDITOR
diff --git a/Engine/Source/Runtime/AIModule/Private/BehaviorTree/Tasks/BTTask_BlackboardBase.cpp b/Engine/Source/Runtime/AIModule/Private/BehaviorTree/Tasks/BTTask_BlackboardBase.cpp
--- a/Engine/Source/Runtime/AIModule/Private/BehaviorTree/Tasks/BTTask_BlackboardBase.cpp
+++ b/Engine/Source/Runtime/AIModule/Private/BehaviorTree/Tasks/BTTask_BlackboardBase.cpp
@@ -15,14 +15,24 @@ void UBTTask_BlackboardBase::InitializeFromAsset(UBehaviorTree& Asset)
 {
 	Super::InitializeFromAsset(Asset);
 
-	UBlackboardData* BBAsset = GetBlackboardAsset();
-	if (BBAsset)
+	if (const UBlackboardData* BBAsset = GetBlackboardAsset())
 	{
 		BlackboardKey.ResolveSelectedKey(*BBAsset);
 	}
 	else
 	{
-		UE_LOG(LogBehaviorTree, Warning, TEXT("Can't initialize task: %s, make sure that behavior tree specifies blackboard asset!"), *GetName());
+		BlackboardKey.InvalidateResolvedKey();
 	}
 }
 
+#if WITH_EDITOR
+FString UBTTask_BlackboardBase::GetErrorMessage() const
+{
+	if (GetBlackboardAsset() == nullptr)
+	{
+		return UE::BehaviorTree::Messages::BlackboardNotSet.ToString();
+	}
+	return Super::GetErrorMessage();
+}
+#endif // WITH_EDITOR
+
diff --git a/Engine/Source/Runtime/AIModule/Private/BehaviorTree/Tasks/BTTask_WaitBlackboardTime.cpp b/Engine/Source/Runtime/AIModule/Private/BehaviorTree/Tasks/BTTask_WaitBlackboardTime.cpp
--- a/Engine/Source/Runtime/AIModule/Private/BehaviorTree/Tasks/BTTask_WaitBlackboardTime.cpp
+++ b/Engine/Source/Runtime/AIModule/Private/BehaviorTree/Tasks/BTTask_WaitBlackboardTime.cpp
@@ -17,10 +17,25 @@ void UBTTask_WaitBlackboardTime::InitializeFromAsset(UBehaviorTree& Asset)
 {
 	Super::InitializeFromAsset(Asset);
 
-	UBlackboardData* BBAsset = GetBlackboardAsset();
-	if (ensure(BBAsset))
+	if (const UBlackboardData* BBAsset = GetBlackboardAsset())
 	{
 		BlackboardKey.ResolveSelectedKey(*BBAsset);
 		WaitTime.SetKey(BlackboardKey.SelectedKeyName);
 	}
-}
\ No newline at end of file
+	else
+	{
+		BlackboardKey.InvalidateResolvedKey();
+		WaitTime = {};
+	}
+}
+
+#if WITH_EDITOR
+FString UBTTask_WaitBlackboardTime::GetErrorMessage() const
+{
+	if (GetBlackboardAsset() == nullptr)
+	{
+		return UE::BehaviorTree::Messages::BlackboardNotSet.ToString();
+	}
+	return Super::GetErrorMessage();
+}
+#endif // WITH_EDITOR
\ No newline at end of file
