Gameplay debugger - EQS - Show the 5 most recent queries instead of the 5 oldest ones

diff --git a/Engine/Source/Runtime/AIModule/Private/GameplayDebugger/GameplayDebuggerCategory_EQS.cpp b/Engine/Source/Runtime/AIModule/Private/GameplayDebugger/GameplayDebuggerCategory_EQS.cpp
--- a/Engine/Source/Runtime/AIModule/Private/GameplayDebugger/GameplayDebuggerCategory_EQS.cpp
+++ b/Engine/Source/Runtime/AIModule/Private/GameplayDebugger/GameplayDebuggerCategory_EQS.cpp
@@ -77,7 +77,7 @@ void FGameplayDebuggerCategory_EQS::CollectData(APlayerController* OwnerPC, AAct
 	{
 		bool operator()(const FEQSDebugger::FEnvQueryInfo& A, const FEQSDebugger::FEnvQueryInfo& B) const
 		{
-			return (A.Timestamp < B.Timestamp);
+			return (A.Timestamp > B.Timestamp);
 		}
 	};
 
