PR #12583: Context ambiguous AIMoveTo strafe fix

diff --git a/Engine/Source/Runtime/AIModule/Private/Blueprint/AIBlueprintHelperLibrary.cpp b/Engine/Source/Runtime/AIModule/Private/Blueprint/AIBlueprintHelperLibrary.cpp
--- a/Engine/Source/Runtime/AIModule/Private/Blueprint/AIBlueprintHelperLibrary.cpp
+++ b/Engine/Source/Runtime/AIModule/Private/Blueprint/AIBlueprintHelperLibrary.cpp
@@ -156,7 +156,8 @@ UAIAsyncTaskBlueprintProxy* UAIBlueprintHelperLibrary::CreateMoveToProxyObject(U
 			MoveReq.SetGoalLocation(Destination);
 		}
 		MoveReq.SetNavigationFilter(AIController->GetDefaultNavigationFilterClass());
-		
+		MoveReq.SetCanStrafe(AIController->bAllowStrafe);
+
 		FPathFollowingRequestResult ResultData = AIController->MoveTo(MoveReq);
 		switch (ResultData.Code)
 		{
diff --git a/Engine/Source/Runtime/AIModule/Private/Tasks/AITask_MoveTo.cpp b/Engine/Source/Runtime/AIModule/Private/Tasks/AITask_MoveTo.cpp
--- a/Engine/Source/Runtime/AIModule/Private/Tasks/AITask_MoveTo.cpp
+++ b/Engine/Source/Runtime/AIModule/Private/Tasks/AITask_MoveTo.cpp
@@ -56,6 +56,7 @@ UAITask_MoveTo* UAITask_MoveTo::AIMoveTo(AAIController* Controller, FVector InGo
 		if (Controller)
 		{
 			MoveReq.SetNavigationFilter(Controller->GetDefaultNavigationFilterClass());
+			MoveReq.SetCanStrafe(Controller->bAllowStrafe);
 		}
 
 		MyTask->SetUp(Controller, MoveReq);
