PR #11803: Improve OnTargetPerceptionForgotten Callbacks in AIPerceptionComponent

diff --git a/Engine/Source/Runtime/AIModule/Classes/Perception/AIPerceptionComponent.h b/Engine/Source/Runtime/AIModule/Classes/Perception/AIPerceptionComponent.h
--- a/Engine/Source/Runtime/AIModule/Classes/Perception/AIPerceptionComponent.h
+++ b/Engine/Source/Runtime/AIModule/Classes/Perception/AIPerceptionComponent.h
@@ -397,11 +397,12 @@ public:
 
 	/**
 	 * Notifies all bound delegates that the perception info has been forgotten for a given target.
-	 * The notification get broadcast when all stimuli of a given target expire. Note that this
-	 * functionality requires the the actor forgetting must be enabled via AIPerceptionSystem.bForgetStaleActors.
+	 * The notification gets broadcast when all stimuli of a given target expire or the target was explicitly forgotten.
+	 * @note The expiration functionality requires the actor forgetting to be enabled via AIPerceptionSystem.bForgetStaleActors.
 	 *
 	 * @param	SourceActor	Actor associated to the stimulus (can not be null)
 	 * @param	Stimulus	Updated stimulus
+	 * @see ForgetActor, ForgetAll
 	 */
 	UPROPERTY(BlueprintAssignable)
 	FActorPerceptionForgetUpdatedDelegate OnTargetPerceptionForgotten;
diff --git a/Engine/Source/Runtime/AIModule/Private/Perception/AIPerceptionComponent.cpp b/Engine/Source/Runtime/AIModule/Private/Perception/AIPerceptionComponent.cpp
--- a/Engine/Source/Runtime/AIModule/Private/Perception/AIPerceptionComponent.cpp
+++ b/Engine/Source/Runtime/AIModule/Private/Perception/AIPerceptionComponent.cpp
@@ -614,15 +614,6 @@ void UAIPerceptionComponent::ProcessStimuli()
 		ForgetActor(ActorToForget);
 	}
 
-	// notify anyone interested
-	if (OnTargetPerceptionForgotten.IsBound())
-	{
-		for (AActor* ActorToForget : ActorsToForget)
-		{
-			OnTargetPerceptionForgotten.Broadcast(ActorToForget);
-		}
-	}
-
 	// remove perceptual info related to stale actors
 	for (const TObjectKey<AActor>& SourceKey : DataToRemove)
 	{
@@ -687,6 +678,12 @@ void UAIPerceptionComponent::ForgetActor(AActor* ActorToForget)
 		}
 
 		const int32 NumRemoved = PerceptualData.Remove(ActorToForget);
+
+		// notify anyone interested
+		if (NumRemoved > 0 && OnTargetPerceptionForgotten.IsBound())
+		{
+			OnTargetPerceptionForgotten.Broadcast(ActorToForget);
+		}
 	}
 }
 
@@ -700,6 +697,18 @@ void UAIPerceptionComponent::ForgetAll()
 			AIPerceptionSys->OnListenerForgetsAll(*this);
 		}
 
+		// notify anyone interested
+		if (OnTargetPerceptionForgotten.IsBound())
+		{
+			for (FActorPerceptionContainer::TConstIterator DataIt = GetPerceptualDataConstIterator(); DataIt; ++DataIt)
+			{
+				if (AActor* ActorToForget = DataIt.Value().Target.Get())
+				{
+					OnTargetPerceptionForgotten.Broadcast(ActorToForget);
+				}
+			}
+		}
+
 		PerceptualData.Reset();
 	}
 }
