EnvQueryManager - When saving assets, invalidate caches for all worlds and not just PIE
This makes queries update on the server world in PIE.

diff --git a/Engine/Source/Runtime/AIModule/Private/EnvironmentQuery/EnvQueryManager.cpp b/Engine/Source/Runtime/AIModule/Private/EnvironmentQuery/EnvQueryManager.cpp
--- a/Engine/Source/Runtime/AIModule/Private/EnvironmentQuery/EnvQueryManager.cpp
+++ b/Engine/Source/Runtime/AIModule/Private/EnvironmentQuery/EnvQueryManager.cpp
@@ -242,28 +242,30 @@ void UEnvQueryManager::NotifyAssetUpdate(UEnvQuery* Query)
 		return;
 	}
 
-	UWorld* World = GEditor->GetEditorWorldContext().World();
-	if (World)
+	for (const FWorldContext& WorldContext : GEditor->GetWorldContexts())
 	{
-		UEnvQueryManager* EQS = UEnvQueryManager::GetCurrent(World);
-		if (EQS)
+		if (UWorld* World = WorldContext.World())
 		{
-			EQS->InstanceCache.Reset();
-		}
-
-		// was as follows, but got broken with changes to actor iterator (FActorIteratorBase::SpawnedActorArray)
-		// for (TActorIterator<AEQSTestingPawn> It(World); It; ++It)
-		for (TActorIterator<AActor> It(World); It; ++It)
-		{
-			AEQSTestingPawn* EQSPawn = Cast<AEQSTestingPawn>(*It);
-			if (EQSPawn == NULL)
+			UEnvQueryManager* EQS = UEnvQueryManager::GetCurrent(World);
+			if (EQS)
 			{
-				continue;
+				EQS->InstanceCache.Reset();
 			}
 
-			if (EQSPawn->QueryTemplate == Query || Query == NULL)
+			// was as follows, but got broken with changes to actor iterator (FActorIteratorBase::SpawnedActorArray)
+			// for (TActorIterator<AEQSTestingPawn> It(World); It; ++It)
+			for (TActorIterator<AActor> It(World); It; ++It)
 			{
-				EQSPawn->RunEQSQuery();
+				AEQSTestingPawn* EQSPawn = Cast<AEQSTestingPawn>(*It);
+				if (EQSPawn == NULL)
+				{
+					continue;
+				}
+
+				if (EQSPawn->QueryTemplate == Query || Query == NULL)
+				{
+					EQSPawn->RunEQSQuery();
+				}
 			}
 		}
 	}
