[EQS] fixed ever growing list of queries when the owners are GC'ed
PR #12809: Fixing a memory leak in the EQS Debugger.

diff --git a/Engine/Source/Runtime/AIModule/Classes/EnvironmentQuery/EnvQueryManager.h b/Engine/Source/Runtime/AIModule/Classes/EnvironmentQuery/EnvQueryManager.h
--- a/Engine/Source/Runtime/AIModule/Classes/EnvironmentQuery/EnvQueryManager.h
+++ b/Engine/Source/Runtime/AIModule/Classes/EnvironmentQuery/EnvQueryManager.h
@@ -191,7 +191,7 @@ struct FEQSDebugger
 
 protected:
 	// maps owner to performed queries
-	TMap<const UObject*, TArray<FEnvQueryInfo> > StoredQueries;
+	TMap<FObjectKey, TArray<FEnvQueryInfo>> StoredQueries;
 };
 
 FArchive& operator<<(FArchive& Ar, FEQSDebugger::FStatsInfo& Data);
diff --git a/Engine/Source/Runtime/AIModule/Private/EnvironmentQuery/EnvQueryManager.cpp b/Engine/Source/Runtime/AIModule/Private/EnvironmentQuery/EnvQueryManager.cpp
--- a/Engine/Source/Runtime/AIModule/Private/EnvironmentQuery/EnvQueryManager.cpp
+++ b/Engine/Source/Runtime/AIModule/Private/EnvironmentQuery/EnvQueryManager.cpp
@@ -1197,7 +1197,14 @@ void FEQSDebugger::StoreQuery(const TSharedPtr<FEnvQueryInstance>& QueryInstance
 {
 	SCOPE_CYCLE_COUNTER(STAT_AI_EQS_Debug_StoreQuery);
 
-	StoredQueries.Remove(nullptr);
+	for (TMap<FObjectKey, TArray<FEnvQueryInfo>>::TIterator It(StoredQueries); It; ++It)
+	{
+		if (It.Key().ResolveObjectPtr() == nullptr)
+		{
+			It.RemoveCurrent();
+		}
+	}
+
 	if (!QueryInstance.IsValid() || QueryInstance->World == nullptr)
 	{
 		return;
