[AIController] added cleanup for gameplay tasks owned by the controller when handling OnUnPossess.

diff --git a/Engine/Source/Runtime/AIModule/Private/AIController.cpp b/Engine/Source/Runtime/AIModule/Private/AIController.cpp
--- a/Engine/Source/Runtime/AIModule/Private/AIController.cpp
+++ b/Engine/Source/Runtime/AIModule/Private/AIController.cpp
@@ -530,12 +530,22 @@ void AAIController::OnUnPossess()
 		PathFollowingComponent->Cleanup();
 	}
 
+	const bool bGameplayTaskCleanupRequired = (CachedGameplayTasksComponent && (CachedGameplayTasksComponent->GetOwner() == CurrentPawn));
+
+	// End all tasks owned by the controller before cleaning up the BrainComponent
+	// since it might need to update its resource lock (see OnGameplayTaskResourcesClaimed)
+	if (bGameplayTaskCleanupRequired)
+	{
+		CachedGameplayTasksComponent->EndAllResourceConsumingTasksOwnedBy(*this);
+	}
+
 	if (bStopAILogicOnUnposses)
 	{
 		CleanupBrainComponent();
 	}
 
-	if (CachedGameplayTasksComponent && (CachedGameplayTasksComponent->GetOwner() == CurrentPawn))
+	// Can now unregister callback and clear the cache
+	if (bGameplayTaskCleanupRequired)
 	{
 		CachedGameplayTasksComponent->OnClaimedResourcesChange.RemoveDynamic(this, &AAIController::OnGameplayTaskResourcesClaimed);
 		CachedGameplayTasksComponent = nullptr;
