Fix CIS issue #778525: missing header in non-unity build

diff --git a/Engine/Source/Runtime/AIModule/Private/GameplayDebugger/GameplayDebuggerCategory_NavLocalGrid.cpp b/Engine/Source/Runtime/AIModule/Private/GameplayDebugger/GameplayDebuggerCategory_NavLocalGrid.cpp
--- a/Engine/Source/Runtime/AIModule/Private/GameplayDebugger/GameplayDebuggerCategory_NavLocalGrid.cpp
+++ b/Engine/Source/Runtime/AIModule/Private/GameplayDebugger/GameplayDebuggerCategory_NavLocalGrid.cpp
@@ -2,6 +2,8 @@
 
 #include "GameplayDebugger/GameplayDebuggerCategory_NavLocalGrid.h"
 
+#include "DynamicMeshBuilder.h"
+
 #if WITH_GAMEPLAY_DEBUGGER_MENU
 
 #include "GameFramework/PlayerController.h"
