PR #12871: Adds #ifndef macro to USE_EQS_DEBUGGER in EnvQueryTypes.h to allow for the preprocessor macro to be modified in target.cs files via the GlobalDefinitions or Definitions variable.

diff --git a/Engine/Source/Runtime/AIModule/Classes/EnvironmentQuery/EnvQueryTypes.h b/Engine/Source/Runtime/AIModule/Classes/EnvironmentQuery/EnvQueryTypes.h
--- a/Engine/Source/Runtime/AIModule/Classes/EnvironmentQuery/EnvQueryTypes.h
+++ b/Engine/Source/Runtime/AIModule/Classes/EnvironmentQuery/EnvQueryTypes.h
@@ -32,7 +32,9 @@ struct FEnvQueryInstance;
 AIMODULE_API DECLARE_LOG_CATEGORY_EXTERN(LogEQS, Display, All);
 
 // If set, execution details will be processed by debugger
+#ifndef USE_EQS_DEBUGGER
 #define USE_EQS_DEBUGGER				(1 && !(UE_BUILD_SHIPPING || UE_BUILD_TEST))
+#endif
 
 // If set, execution stats will also gather EQS tick load data (16k memory for each query record)
 #define USE_EQS_TICKLOADDATA			(1 && USE_EQS_DEBUGGER && WITH_EDITOR)
