Fix decorator blackboard invalid operation type when changing key type

diff --git a/Engine/Source/Runtime/AIModule/Private/BehaviorTree/Decorators/BTDecorator_Blackboard.cpp b/Engine/Source/Runtime/AIModule/Private/BehaviorTree/Decorators/BTDecorator_Blackboard.cpp
--- a/Engine/Source/Runtime/AIModule/Private/BehaviorTree/Decorators/BTDecorator_Blackboard.cpp
+++ b/Engine/Source/Runtime/AIModule/Private/BehaviorTree/Decorators/BTDecorator_Blackboard.cpp
@@ -216,22 +216,24 @@ void UBTDecorator_Blackboard::PostEditChangeProperty(FPropertyChangedEvent& Prop
 		}
 	}
 
+	const bool bKeyTypeChanged = (ChangedPropName == GET_MEMBER_NAME_CHECKED(FBlackboardKeySelector, SelectedKeyType));
+
 	const UBlackboardKeyType* KeyCDO = BlackboardKey.SelectedKeyType ? BlackboardKey.SelectedKeyType->GetDefaultObject<UBlackboardKeyType>() : NULL;
-	if (ChangedPropName == GET_MEMBER_NAME_CHECKED(UBTDecorator_Blackboard, BasicOperation))
+	if (bKeyTypeChanged || ChangedPropName == GET_MEMBER_NAME_CHECKED(UBTDecorator_Blackboard, BasicOperation))
 	{
 		if (KeyCDO && KeyCDO->GetTestOperation() == EBlackboardKeyOperation::Basic)
 		{
 			OperationType = BasicOperation.GetIntValue();
 		}
 	}
-	else if (ChangedPropName == GET_MEMBER_NAME_CHECKED(UBTDecorator_Blackboard, ArithmeticOperation))
+	if (bKeyTypeChanged || ChangedPropName == GET_MEMBER_NAME_CHECKED(UBTDecorator_Blackboard, ArithmeticOperation))
 	{
 		if (KeyCDO && KeyCDO->GetTestOperation() == EBlackboardKeyOperation::Arithmetic)
 		{
 			OperationType = ArithmeticOperation.GetIntValue();
 		}
 	}
-	else if (ChangedPropName == GET_MEMBER_NAME_CHECKED(UBTDecorator_Blackboard, TextOperation))
+	if (bKeyTypeChanged || ChangedPropName == GET_MEMBER_NAME_CHECKED(UBTDecorator_Blackboard, TextOperation))
 	{
 		if (KeyCDO && KeyCDO->GetTestOperation() == EBlackboardKeyOperation::Text)
 		{
