[NavLink] Fix warning for invalid bounds if ANavLinkProxy does not have a simple link but is using a smart link.

diff --git a/Engine/Source/Runtime/AIModule/Private/Navigation/NavLinkProxy.cpp b/Engine/Source/Runtime/AIModule/Private/Navigation/NavLinkProxy.cpp
--- a/Engine/Source/Runtime/AIModule/Private/Navigation/NavLinkProxy.cpp
+++ b/Engine/Source/Runtime/AIModule/Private/Navigation/NavLinkProxy.cpp
@@ -212,12 +212,14 @@ void ANavLinkProxy::PostEditImport()
 
 void ANavLinkProxy::PostRegisterAllComponents()
 {
-	Super::PostRegisterAllComponents();
-
+	// Check if the smart link is enabled and valid for determining navigation relevancy
+	// A smart link with no simple links will result in warnings for invalid bounds if Super is called prior to setting the navigation relevancy of the smart link.
 	if (SmartLinkComp)
 	{
 		SmartLinkComp->SetNavigationRelevancy(bSmartLinkIsRelevant);
 	}
+
+	Super::PostRegisterAllComponents();
 }
 
 void ANavLinkProxy::PostLoad()
