[HierarchicalHashGrid] fixed comment and deprecated a header

diff --git a/Engine/Source/Runtime/AIModule/Public/HierarchicalHashGrid2D.h b/Engine/Source/Runtime/AIModule/Public/HierarchicalHashGrid2D.h
--- a/Engine/Source/Runtime/AIModule/Public/HierarchicalHashGrid2D.h
+++ b/Engine/Source/Runtime/AIModule/Public/HierarchicalHashGrid2D.h
@@ -2,7 +2,9 @@
 
 #pragma once
 
+#if UE_ENABLE_INCLUDE_ORDER_DEPRECATED_IN_5_6
 #include "CoreMinimal.h"
+#endif // UE_ENABLE_INCLUDE_ORDER_DEPRECATED_IN_5_6
 #include "Containers/StaticArray.h"
 
 /**
@@ -312,10 +314,10 @@ public:
 		return Move(ID, OldLocation, NewBounds);
 	}
 
-	/** Moves item based on previous bounding box and new bounding box.
+	/** Moves item based on previous cell location and new bounding box.
 	  * @param ID - External ID used to identify the item.
 	  * @param OldLocation - The same cell location the item was previously added or moved with.
-	  * @param NewBounds - New  bounds of the item
+	  * @param NewBounds - New bounds of the item
 	  * @return The cell location where the item was added to.
 	  */
 	FCellLocation Move(const ItemIDType ID, const FCellLocation& OldLocation, const FBox& NewBounds)
