Fix non-unity
#rnx

diff --git a/Engine/Source/Runtime/AIModule/Public/HierarchicalHashGrid2D.h b/Engine/Source/Runtime/AIModule/Public/HierarchicalHashGrid2D.h
--- a/Engine/Source/Runtime/AIModule/Public/HierarchicalHashGrid2D.h
+++ b/Engine/Source/Runtime/AIModule/Public/HierarchicalHashGrid2D.h
@@ -6,6 +6,7 @@
 #include "CoreMinimal.h"
 #endif // UE_ENABLE_INCLUDE_ORDER_DEPRECATED_IN_5_6
 #include "Containers/StaticArray.h"
+#include "Math/Box.h"
 
 /**
  * Hierarchical Hash Grid in 2D.
