PR #12960: Replaced the RefreshStimulus() method

Replaced the RefreshStimulus() method with a new method, ConditionallyStoreSuccessfulStimulus(), which returns a boolean result. This change allows for method overloading and enforces the invocation of OnTargetPerceptionUpdated when needed. This also enables the inclusion of custom logic, such as handling events for a changing Stimulus.Strength parameter when NotifyType == EAISenseNotifyType::OnPerceptionChange

diff --git a/Engine/Source/Runtime/AIModule/Classes/Perception/AIPerceptionComponent.h b/Engine/Source/Runtime/AIModule/Classes/Perception/AIPerceptionComponent.h
--- a/Engine/Source/Runtime/AIModule/Classes/Perception/AIPerceptionComponent.h
+++ b/Engine/Source/Runtime/AIModule/Classes/Perception/AIPerceptionComponent.h
@@ -445,7 +445,15 @@ protected:
 	AIMODULE_API void RemoveDeadData();
 
 	/** Updates the stimulus entry in StimulusStore, if NewStimulus is more recent or stronger */
-	AIMODULE_API virtual void RefreshStimulus(FAIStimulus& StimulusStore, const FAIStimulus& NewStimulus);
+	UE_DEPRECATED(5.6, "RefreshStimulus is deprecated. Use ConditionallyStoreSuccessfulStimulus() instead.")
+	AIMODULE_API virtual void RefreshStimulus(FAIStimulus& StimulusStore, const FAIStimulus& NewStimulus) final;
+
+	/**
+	 * Updates the stimulus entry in StimulusStore, if NewStimulus is more recent or stronger
+	 * @return whether the detected change in stimulus values should be considered a "perception change"
+	 *		(and, in consequence, emit appropriate notifications) - defaults to `false`
+	 */
+	AIMODULE_API virtual bool ConditionallyStoreSuccessfulStimulus(FAIStimulus& StimulusStore, const FAIStimulus& NewStimulus);
 
 	/** @note no need to call super implementation, it's there just for some validity checking */
 	AIMODULE_API virtual void HandleExpiredStimulus(FAIStimulus& StimulusStore);
diff --git a/Engine/Source/Runtime/AIModule/Private/Perception/AIPerceptionComponent.cpp b/Engine/Source/Runtime/AIModule/Private/Perception/AIPerceptionComponent.cpp
--- a/Engine/Source/Runtime/AIModule/Private/Perception/AIPerceptionComponent.cpp
+++ b/Engine/Source/Runtime/AIModule/Private/Perception/AIPerceptionComponent.cpp
@@ -546,9 +546,12 @@ void UAIPerceptionComponent::ProcessStimuli()
 		const bool bActorInfoUpdated = SourcedStimulus.Stimulus.WantsToNotifyOnlyOnPerceptionChange() == false 
 			|| SourcedStimulus.Stimulus.WasSuccessfullySensed() != StimulusStore.WasSuccessfullySensed();
 
+		bool bRequiresUpdate = bActorInfoUpdated;
+
 		if (SourcedStimulus.Stimulus.WasSuccessfullySensed())
 		{
-			RefreshStimulus(StimulusStore, SourcedStimulus.Stimulus);
+			bRequiresUpdate |= ConditionallyStoreSuccessfulStimulus(StimulusStore, SourcedStimulus.Stimulus);
+			
 		}
 		else if (StimulusStore.IsExpired() == false)
 		{	
@@ -574,7 +577,7 @@ void UAIPerceptionComponent::ProcessStimuli()
 		}
 
 		// if the new stimulus is "valid" or it's info that "no longer sensed" and it used to be sensed successfully
-		if (bActorInfoUpdated)
+		if (bRequiresUpdate)
 		{
 			// Source Actor is only resolved from SourceKey when required but might already have been resolved for new entry
 			SourceActor = (SourceActor == nullptr) ? CastChecked<AActor>(SourceKey.ResolveObjectPtr(), ECastCheckedType::NullAllowed) : SourceActor;
@@ -622,6 +625,11 @@ void UAIPerceptionComponent::ProcessStimuli()
 }
 
 void UAIPerceptionComponent::RefreshStimulus(FAIStimulus& StimulusStore, const FAIStimulus& NewStimulus)
+{
+	ConditionallyStoreSuccessfulStimulus(StimulusStore, NewStimulus);
+}
+
+bool UAIPerceptionComponent::ConditionallyStoreSuccessfulStimulus(FAIStimulus& StimulusStore, const FAIStimulus& NewStimulus)
 {
 	// if new stimulus is younger or stronger
 	// note that stimulus Age depends on PerceptionSystem::PerceptionAgingRate. It's possible that 
@@ -629,8 +637,9 @@ void UAIPerceptionComponent::RefreshStimulus(FAIStimulus& StimulusStore, const F
 	if (NewStimulus.GetAge() <= StimulusStore.GetAge() || StimulusStore.Strength < NewStimulus.Strength)
 	{
 		StimulusStore = NewStimulus;
-		// update stimulus 
 	}
+
+	return false;
 }
 
 void UAIPerceptionComponent::HandleExpiredStimulus(FAIStimulus& StimulusStore)
