PR #12972: Remove deprecated props for EQS Task

diff --git a/Engine/Source/Runtime/AIModule/Classes/BehaviorTree/Tasks/BTTask_RunEQSQuery.h b/Engine/Source/Runtime/AIModule/Classes/BehaviorTree/Tasks/BTTask_RunEQSQuery.h
--- a/Engine/Source/Runtime/AIModule/Classes/BehaviorTree/Tasks/BTTask_RunEQSQuery.h
+++ b/Engine/Source/Runtime/AIModule/Classes/BehaviorTree/Tasks/BTTask_RunEQSQuery.h
@@ -26,30 +26,6 @@ class UBTTask_RunEQSQuery : public UBTTask_BlackboardBase
 {
 	GENERATED_UCLASS_BODY()
 
-	/** query to run */
-	UE_DEPRECATED_FORGAME(5.0, "UBTTask_RunEQSQuery.QueryTemplate has been deprecated for a long while now. Will be removed in the next engine version.")
-	UPROPERTY()
-	TObjectPtr<UEnvQuery> QueryTemplate;
-
-	/** optional parameters for query */
-	UE_DEPRECATED_FORGAME(5.0, "UBTTask_RunEQSQuery.QueryParams has been deprecated for a long while now. Will be removed in the next engine version.")
-	UPROPERTY()
-	TArray<FEnvNamedValue> QueryParams;
-
-	UE_DEPRECATED_FORGAME(5.0, "UBTTask_RunEQSQuery.QueryConfig has been deprecated for a long while now. Will be removed in the next engine version.")
-	UPROPERTY()
-	TArray<FAIDynamicParam> QueryConfig;
-
-	/** determines which item will be stored (All = only first matching) */
-	UE_DEPRECATED_FORGAME(5.0, "UBTTask_RunEQSQuery.RunMode has been deprecated for a long while now. Will be removed in the next engine version.")
-	UPROPERTY()
-	TEnumAsByte<EEnvQueryRunMode::Type> RunMode;
-
-	/** blackboard key storing an EQS query template */
-	UE_DEPRECATED_FORGAME(5.0, "UBTTask_RunEQSQuery.EQSQueryBlackboardKey been deprecated for a long while now. Will be removed in the next engine version.")
-	UPROPERTY()
-	struct FBlackboardKeySelector EQSQueryBlackboardKey;
-
 	UPROPERTY(EditAnywhere, Category=Node, meta=(InlineEditConditionToggle))
 	bool bUseBBKey;
 
@@ -74,9 +50,6 @@ class UBTTask_RunEQSQuery : public UBTTask_BlackboardBase
 	/** finish task */
 	AIMODULE_API void OnQueryFinished(TSharedPtr<FEnvQueryResult> Result);
 
-	/** Convert QueryParams to QueryConfig */
-	AIMODULE_API virtual void PostLoad() override;
-
 #if WITH_EDITOR
 	/** prepare query params */
 	AIMODULE_API virtual void PostEditChangeProperty(struct FPropertyChangedEvent& PropertyChangedEvent) override;
diff --git a/Engine/Source/Runtime/AIModule/Private/BehaviorTree/Tasks/BTTask_RunEQSQuery.cpp b/Engine/Source/Runtime/AIModule/Private/BehaviorTree/Tasks/BTTask_RunEQSQuery.cpp
--- a/Engine/Source/Runtime/AIModule/Private/BehaviorTree/Tasks/BTTask_RunEQSQuery.cpp
+++ b/Engine/Source/Runtime/AIModule/Private/BehaviorTree/Tasks/BTTask_RunEQSQuery.cpp
@@ -29,9 +29,6 @@ UBTTask_RunEQSQuery::UBTTask_RunEQSQuery(const FObjectInitializer& ObjectInitial
 	QueryFinishedDelegate = FQueryFinishedSignature::CreateUObject(this, &UBTTask_RunEQSQuery::OnQueryFinished);
 
 	bUpdateBBOnFail = GET_AI_CONFIG_VAR(bClearBBEntryOnBTEQSFail);
-
-	// deprecated
-	EQSQueryBlackboardKey.AddObjectFilter(this, GET_MEMBER_NAME_CHECKED(UBTTask_RunEQSQuery, EQSQueryBlackboardKey), UEnvQuery::StaticClass());
 }
 
 void UBTTask_RunEQSQuery::InitializeFromAsset(UBehaviorTree& Asset)
@@ -167,27 +164,6 @@ void UBTTask_RunEQSQuery::CollectKeyFilters()
 // 	}
 }
 
-void UBTTask_RunEQSQuery::PostLoad()
-{
-	Super::PostLoad();
-
-	if (QueryParams.Num() > 0)
-	{
-		FAIDynamicParam::GenerateConfigurableParamsFromNamedValues(*this, QueryConfig, QueryParams);
-		QueryParams.Empty();
-	}
-
-	// patching part 2
-	if (EQSRequest.QueryTemplate == nullptr && EQSRequest.bUseBBKeyForQueryTemplate == false)
-	{
-		EQSRequest.QueryTemplate = QueryTemplate;
-		EQSRequest.QueryConfig = QueryConfig;
-		EQSRequest.RunMode = RunMode;
-		EQSRequest.EQSQueryBlackboardKey = EQSQueryBlackboardKey;
-		EQSRequest.bUseBBKeyForQueryTemplate = bUseBBKey;
-	}
-}
-
 #if WITH_EDITOR
 void UBTTask_RunEQSQuery::PostEditChangeProperty(FPropertyChangedEvent& PropertyChangedEvent)
 {
