Add IsKeyOfType helper using a KeyID in BlackBoardComponent

diff --git a/Engine/Source/Runtime/AIModule/Classes/BehaviorTree/BlackboardComponent.h b/Engine/Source/Runtime/AIModule/Classes/BehaviorTree/BlackboardComponent.h
--- a/Engine/Source/Runtime/AIModule/Classes/BehaviorTree/BlackboardComponent.h
+++ b/Engine/Source/Runtime/AIModule/Classes/BehaviorTree/BlackboardComponent.h
@@ -181,6 +181,9 @@ public:
 	/** Copy content from SourceKeyID to DestinationID and return true if it worked */
 	AIMODULE_API bool CopyKeyValue(FBlackboard::FKey SourceKeyID, FBlackboard::FKey DestinationID);
 
+	template<class TDataClass>
+	bool IsKeyOfType(const FName& KeyName) const;
+
 	template<class TDataClass>
 	bool IsKeyOfType(FBlackboard::FKey KeyID) const;
 
@@ -314,6 +317,13 @@ FORCEINLINE bool UBlackboardComponent::HasValidAsset() const
 	return BlackboardAsset && BlackboardAsset->IsValid();
 }
 
+template<class TDataClass>
+bool UBlackboardComponent::IsKeyOfType(const FName& KeyName) const
+{
+	const FBlackboard::FKey KeyID = GetKeyID(KeyName);
+	return IsKeyOfType<TDataClass>(KeyID);
+}
+
 template<class TDataClass>
 bool UBlackboardComponent::IsKeyOfType(FBlackboard::FKey KeyID) const
 {
