Deduplicated some header defined variables. Saves 75KB of program size on Win64.

diff --git a/Engine/Source/Runtime/AIModule/Classes/AITypes.h b/Engine/Source/Runtime/AIModule/Classes/AITypes.h
--- a/Engine/Source/Runtime/AIModule/Classes/AITypes.h
+++ b/Engine/Source/Runtime/AIModule/Classes/AITypes.h
@@ -25,13 +25,13 @@ DECLARE_CYCLE_STAT_EXTERN(TEXT("Overall AI Time"), STAT_AI_Overall, STATGROUP_AI
 
 namespace FAISystem
 {
-	inline static const FRotator InvalidRotation = FRotator(FLT_MAX);
-	inline static const FQuat InvalidOrientation = FQuat(FLT_MAX, FLT_MAX, FLT_MAX, FLT_MAX);
-	inline static const FVector InvalidLocation = FVector(FLT_MAX);
-	inline static const FVector InvalidDirection = FVector::ZeroVector; 
-	inline static const float InvalidRange = -1.f;
-	inline static const float InfiniteInterval = -FLT_MAX;
-	inline static const uint32 InvalidUnsignedID = uint32(INDEX_NONE);
+	extern AIMODULE_API const FRotator InvalidRotation;
+	extern AIMODULE_API const FQuat InvalidOrientation;
+	inline constexpr FVector InvalidLocation = FVector(FLT_MAX, UE::Math::TVectorConstInit{});
+	inline constexpr FVector InvalidDirection = FVector(0, UE::Math::TVectorConstInit{});
+	constexpr float InvalidRange = -1.f;
+	constexpr float InfiniteInterval = -FLT_MAX;
+	constexpr uint32 InvalidUnsignedID = uint32(INDEX_NONE);
 
 	FORCEINLINE bool IsValidLocation(const FVector& TestLocation)
 	{
@@ -105,9 +105,9 @@ namespace EAIForceParam
 
 namespace FAIMoveFlag
 {
-	static const bool StopOnOverlap = true;
-	static const bool UsePathfinding = true;
-	static const bool IgnorePathfinding = false;
+	constexpr bool StopOnOverlap = true;
+	constexpr bool UsePathfinding = true;
+	constexpr bool IgnorePathfinding = false;
 }
 
 namespace EAILogicResuming
@@ -141,7 +141,7 @@ namespace EAIRequestPriority
 
 namespace EAIRequestPriority
 {
-	static const int32 Lowest = EAIRequestPriority::Logic;
+	constexpr int32 Lowest = EAIRequestPriority::Logic;
 };
 
 UENUM()
@@ -300,9 +300,9 @@ typedef FAINamedID<FAIResCounter> FAIResourceID;
 //////////////////////////////////////////////////////////////////////////
 struct FAIResourcesSet
 {
-	static const uint32 NoResources = 0;
-	static const uint32 AllResources = uint32(-1);
-	static const uint8 MaxFlags = 32;
+	static constexpr uint32 NoResources = 0;
+	static constexpr uint32 AllResources = uint32(-1);
+	static constexpr uint8 MaxFlags = 32;
 private:
 	uint32 Flags;
 public:
@@ -404,8 +404,8 @@ struct FAIRequestID
 	GENERATED_USTRUCT_BODY()
 		
 private:
-	static const uint32 AnyRequestID = 0;
-	static const uint32 InvalidRequestID = uint32(-1);
+	static constexpr uint32 AnyRequestID = 0;
+	static constexpr uint32 InvalidRequestID = uint32(-1);
 
 	UPROPERTY()
 	uint32 RequestID;
diff --git a/Engine/Source/Runtime/AIModule/Private/AITypes.cpp b/Engine/Source/Runtime/AIModule/Private/AITypes.cpp
--- a/Engine/Source/Runtime/AIModule/Private/AITypes.cpp
+++ b/Engine/Source/Runtime/AIModule/Private/AITypes.cpp
@@ -6,6 +6,9 @@
 #include "Navigation/PathFollowingComponent.h"
 #include "NavFilters/NavigationQueryFilter.h"
 
+const FRotator FAISystem::InvalidRotation = FRotator(FLT_MAX);
+const FQuat FAISystem::InvalidOrientation = FQuat(FLT_MAX, FLT_MAX, FLT_MAX, FLT_MAX);
+
 //----------------------------------------------------------------------//
 // FAIResourceLock
 //----------------------------------------------------------------------//
