Env Query - OnCircle generator - Add support to offset arc direction
* also remove unused AngleRadians variable

diff --git a/Engine/Source/Runtime/AIModule/Classes/EnvironmentQuery/Generators/EnvQueryGenerator_OnCircle.h b/Engine/Source/Runtime/AIModule/Classes/EnvironmentQuery/Generators/EnvQueryGenerator_OnCircle.h
--- a/Engine/Source/Runtime/AIModule/Classes/EnvironmentQuery/Generators/EnvQueryGenerator_OnCircle.h
+++ b/Engine/Source/Runtime/AIModule/Classes/EnvironmentQuery/Generators/EnvQueryGenerator_OnCircle.h
@@ -46,12 +46,13 @@ class UEnvQueryGenerator_OnCircle : public UEnvQueryGenerator_ProjectedPoints
 	UPROPERTY(EditDefaultsOnly, Category=Generator, meta=(EditCondition="bDefineArc"))
 	FEnvDirection ArcDirection;
 
+	/** Offsets the ArcDirection and starts generating points this many degrees clock-wise around the vertical axis */
+	UPROPERTY(EditDefaultsOnly, Category=Generator)
+	FAIDataProviderFloatValue ArcDirectionOffsetDegrees;
+
 	/** If you generate items on a piece of circle you define angle of Arc cut here */
 	UPROPERTY(EditDefaultsOnly, Category=Generator)
 	FAIDataProviderFloatValue ArcAngle;
-	
-	UPROPERTY()
-	mutable float AngleRadians;
 
 	/** context */
 	UPROPERTY(EditAnywhere, Category=Generator)
diff --git a/Engine/Source/Runtime/AIModule/Private/EnvironmentQuery/Generators/EnvQueryGenerator_OnCircle.cpp b/Engine/Source/Runtime/AIModule/Private/EnvironmentQuery/Generators/EnvQueryGenerator_OnCircle.cpp
--- a/Engine/Source/Runtime/AIModule/Private/EnvironmentQuery/Generators/EnvQueryGenerator_OnCircle.cpp
+++ b/Engine/Source/Runtime/AIModule/Private/EnvironmentQuery/Generators/EnvQueryGenerator_OnCircle.cpp
@@ -29,7 +29,6 @@ UEnvQueryGenerator_OnCircle::UEnvQueryGenerator_OnCircle(const FObjectInitialize
 	ArcDirection.LineFrom = UEnvQueryContext_Querier::StaticClass();
 	ArcDirection.Rotation = UEnvQueryContext_Querier::StaticClass();
 	ArcAngle.DefaultValue = 360.f;
-	AngleRadians = FMath::DegreesToRadians(360.f);
 	bIgnoreAnyContextActorsWhenGeneratingCircle = false;
 
 	TraceData.bCanProjectDown = false;
@@ -49,7 +48,6 @@ void UEnvQueryGenerator_OnCircle::PostLoad()
 	}
 
 	TraceData.OnPostLoad();
-	AngleRadians = FMath::DegreesToRadians(ArcAngle.DefaultValue);
 }
 
 FVector UEnvQueryGenerator_OnCircle::CalcDirection(FEnvQueryInstance& QueryInstance) const
@@ -138,11 +136,13 @@ void UEnvQueryGenerator_OnCircle::GenerateItems(FEnvQueryInstance& QueryInstance
 	ArcAngle.BindData(QueryOwner, QueryInstance.QueryID);
 	CircleCenterZOffset.BindData(QueryOwner, QueryInstance.QueryID);
 	NumberOfPoints.BindData(QueryOwner, QueryInstance.QueryID);
+	ArcDirectionOffsetDegrees.BindData(QueryOwner, QueryInstance.QueryID);
 
-	float AngleDegree = ArcAngle.GetValue();
-	float RadiusValue = CircleRadius.GetValue();
-	float ItemSpace = SpaceBetween.GetValue();
-	int32 NumPoints = NumberOfPoints.GetValue();
+	const float AngleDegree = ArcAngle.GetValue();
+	const float RadiusValue = CircleRadius.GetValue();
+	const float ItemSpace = SpaceBetween.GetValue();
+	const int32 NumPoints = NumberOfPoints.GetValue();
+	const float DirectionOffsetDegrees = ArcDirectionOffsetDegrees.GetValue();
 
 	if ((AngleDegree <= 0.f) || 
 		(AngleDegree > 360.f) || 
@@ -153,8 +153,6 @@ void UEnvQueryGenerator_OnCircle::GenerateItems(FEnvQueryInstance& QueryInstance
 		return;
 	}
 
-	AngleRadians = FMath::DegreesToRadians(AngleDegree);
-
 	// first generate points on a circle
 	const float CircumferenceLength = 2.f * PI * RadiusValue;
 	const float ArcAnglePercentage = AngleDegree / 360.f;
@@ -188,7 +186,7 @@ void UEnvQueryGenerator_OnCircle::GenerateItems(FEnvQueryInstance& QueryInstance
 	const float AngleStep = AngleDegree / NumStepsBetween;
 
 	FVector StartDirection = CalcDirection(QueryInstance);
-	StartDirection = StartDirection.RotateAngleAxis(-AngleDegree/2, FVector::UpVector) * RadiusValue;
+	StartDirection = StartDirection.RotateAngleAxis(DirectionOffsetDegrees - AngleDegree/2, FVector::UpVector) * RadiusValue;
 	
 	// gather context with raw data, so it can be used by derived generators
 	FEnvQueryContextData ContextData;
@@ -351,7 +349,6 @@ void UEnvQueryGenerator_OnCircle::PostEditChangeProperty( FPropertyChangedEvent&
 		if (PropName == GET_MEMBER_NAME_CHECKED(UEnvQueryGenerator_OnCircle, ArcAngle))
 		{
 			ArcAngle.DefaultValue = FMath::Clamp(ArcAngle.DefaultValue, 0.0f, 360.f);
-			AngleRadians = FMath::DegreesToRadians(ArcAngle.DefaultValue);
 			bDefineArc = (ArcAngle.DefaultValue < 360.f) && (ArcAngle.DefaultValue > 0.f);
 		}
 		else if (PropName == GET_MEMBER_NAME_CHECKED(UEnvQueryGenerator_OnCircle, CircleRadius))
