AI Perception - Clarified code comments about the project setting requirement for perception expiration functionality.
#rnx

diff --git a/Engine/Source/Runtime/AIModule/Classes/Perception/AIPerceptionComponent.h b/Engine/Source/Runtime/AIModule/Classes/Perception/AIPerceptionComponent.h
--- a/Engine/Source/Runtime/AIModule/Classes/Perception/AIPerceptionComponent.h
+++ b/Engine/Source/Runtime/AIModule/Classes/Perception/AIPerceptionComponent.h
@@ -398,7 +398,7 @@ public:
 	/**
 	 * Notifies all bound delegates that the perception info has been forgotten for a given target.
 	 * The notification gets broadcast when all stimuli of a given target expire or the target was explicitly forgotten.
-	 * @note The expiration functionality requires the actor forgetting to be enabled via AIPerceptionSystem.bForgetStaleActors.
+	 * @note The target forgetting functionality must be enabled through the AIPerceptionSystem.bForgetStaleActors flag from Project Settings.
 	 *
 	 * @param	SourceActor	Actor associated to the stimulus (can not be null)
 	 * @param	Stimulus	Updated stimulus
diff --git a/Engine/Source/Runtime/AIModule/Classes/Perception/AISenseConfig.h b/Engine/Source/Runtime/AIModule/Classes/Perception/AISenseConfig.h
--- a/Engine/Source/Runtime/AIModule/Classes/Perception/AISenseConfig.h
+++ b/Engine/Source/Runtime/AIModule/Classes/Perception/AISenseConfig.h
@@ -23,7 +23,10 @@ protected:
 	UPROPERTY(EditAnywhere, BlueprintReadOnly, Category = Sense, AdvancedDisplay)
 	FColor DebugColor;
 
-	/** specifies age limit after stimuli generated by this sense become forgotten. 0 means "never" */
+	/** 
+	 *  Specifies age limit after stimuli generated by this sense become expired. 0 means "never"
+	 *  Whether stimuli expiring affects given listener's memory of the target actor and triggers the Forget event is controlled by AIPerceptionSystem.bForgetStaleActors flag in Project Settings.
+	 */
 	UPROPERTY(EditAnywhere, BlueprintReadOnly, Category = Sense, meta = (ClampMin = "0.0", UIMin = "0.0", Units="Seconds"))
 	float MaxAge;
 
