Adding loop until Success/Failure BT decorator

diff --git a/Engine/Source/Developer/AITestSuite/Classes/BehaviorTree/TestBTDecorator_LoopUntil.h b/Engine/Source/Developer/AITestSuite/Classes/BehaviorTree/TestBTDecorator_LoopUntil.h
new file mode 100644
--- /dev/null
+++ b/Engine/Source/Developer/AITestSuite/Classes/BehaviorTree/TestBTDecorator_LoopUntil.h
@@ -0,0 +1,19 @@
+// Copyright Epic Games, Inc. All Rights Reserved.
+
+#pragma once
+
+#include "BehaviorTree/Decorators/BTDecorator_LoopUntil.h"
+#include "TestBTDecorator_LoopUntil.generated.h"
+
+/**
+ * Loop until test version used for unit tests
+ */
+UCLASS(meta=(HiddenNode), MinimalAPI)
+class UTestBTDecorator_LoopUntil : public UBTDecorator_LoopUntil
+{
+	GENERATED_BODY()
+public:
+
+	UTestBTDecorator_LoopUntil(const FObjectInitializer& ObjectInitializer);
+	void SetRequiredResult(const EBTNodeResult::Type InRequiredResult);
+};
diff --git a/Engine/Source/Developer/AITestSuite/Private/BehaviorTree/TestBTDecorator_LoopUntil.cpp b/Engine/Source/Developer/AITestSuite/Private/BehaviorTree/TestBTDecorator_LoopUntil.cpp
new file mode 100644
--- /dev/null
+++ b/Engine/Source/Developer/AITestSuite/Private/BehaviorTree/TestBTDecorator_LoopUntil.cpp
@@ -0,0 +1,20 @@
+// Copyright Epic Games, Inc. All Rights Reserved.
+
+#include "BehaviorTree/TestBTDecorator_LoopUntil.h"
+
+#include UE_INLINE_GENERATED_CPP_BY_NAME(TestBTDecorator_LoopUntil)
+
+UTestBTDecorator_LoopUntil::UTestBTDecorator_LoopUntil(const FObjectInitializer& ObjectInitializer) 
+	: Super(ObjectInitializer)
+{
+	NodeName = TEXT("Test Loop Until");
+
+	bAllowAbortNone = false;
+	bAllowAbortLowerPri = false;
+	bAllowAbortChildNodes = false;
+}
+
+void UTestBTDecorator_LoopUntil::SetRequiredResult(const EBTNodeResult::Type InRequiredResult)
+{
+	RequiredResult = InRequiredResult;
+}
diff --git a/Engine/Source/Developer/AITestSuite/Private/Tests/BTTest.cpp b/Engine/Source/Developer/AITestSuite/Private/Tests/BTTest.cpp
--- a/Engine/Source/Developer/AITestSuite/Private/Tests/BTTest.cpp
+++ b/Engine/Source/Developer/AITestSuite/Private/Tests/BTTest.cpp
@@ -2510,6 +2510,41 @@ struct FAITest_BTLowPriObserverInLoop : public FAITest_SimpleBT
 };
 IMPLEMENT_AI_LATENT_TEST(FAITest_BTLowPriObserverInLoop, "System.AI.Behavior Trees.Other: low priority observer in looped branch")
 
+struct FAITest_BTLoopUntilSuccess : public FAITest_SimpleBT
+{
+	FAITest_BTLoopUntilSuccess()
+	{
+		UBTCompositeNode& CompNode = FBTBuilder::AddSelector(*BTAsset);
+		{
+			// We first set the bool key to false so task number 2 doesn't execute during the first loop
+			FBTBuilder::AddTaskFlagChange(CompNode, false, EBTNodeResult::Failed);
+
+			UBTCompositeNode& CompNode2 = FBTBuilder::AddSelector(CompNode);
+			{
+				// Loop the selector until it succeeds
+				FBTBuilder::WithDecoratorLoopUntil(CompNode, EBTNodeResult::Succeeded);
+
+				// Execute a failed task in the selector so we continue and can verify the looping
+				FBTBuilder::AddTask(CompNode2, 1, EBTNodeResult::Failed);
+
+				// Conditionally execute a succeeding task
+				FBTBuilder::AddTask(CompNode2, 2, EBTNodeResult::Succeeded);
+				{
+					FBTBuilder::WithDecoratorBlackboard(CompNode2, EBasicKeyOperation::Set, EBTFlowAbortMode::None, TEXT("Bool1"));
+				}
+
+				// Set the bool key to true so task number 2 will now execute and make the selector succeed
+				FBTBuilder::AddTaskFlagChange(CompNode2, true, EBTNodeResult::Failed);
+			}
+		}
+
+		ExpectedResult.Add(1);
+		ExpectedResult.Add(1);
+		ExpectedResult.Add(2);
+	}
+};
+IMPLEMENT_AI_LATENT_TEST(FAITest_BTLoopUntilSuccess, "System.AI.Behavior Trees.LoopUntil: Loop until success")
+
 struct FAITest_BTSubtreeSimple : public FAITest_SimpleBT
 {
 	FAITest_BTSubtreeSimple()
diff --git a/Engine/Source/Developer/AITestSuite/Public/BTBuilder.h b/Engine/Source/Developer/AITestSuite/Public/BTBuilder.h
--- a/Engine/Source/Developer/AITestSuite/Public/BTBuilder.h
+++ b/Engine/Source/Developer/AITestSuite/Public/BTBuilder.h
@@ -16,8 +16,9 @@
 #include "BehaviorTree/Tasks/BTTask_RunBehavior.h"
 
 #include "BehaviorTree/TestBTDecorator_Blackboard.h"
-#include "BehaviorTree/TestBTDecorator_DelayedAbort.h"
 #include "BehaviorTree/TestBTDecorator_Blueprint.h"
+#include "BehaviorTree/TestBTDecorator_DelayedAbort.h"
+#include "BehaviorTree/TestBTDecorator_LoopUntil.h"
 #include "BehaviorTree/TestBTDecorator_NodeMemoryValidator.h"
 #include "BehaviorTree/TestBTService_Log.h"
 #include "BehaviorTree/TestBTService_BTStopAction.h"
@@ -362,6 +363,12 @@ struct FBTBuilder
 		UBTDecorator_Loop& LoopDecorator = WithDecorator<UBTDecorator_Loop>(ParentNode);
 		LoopDecorator.NumLoops = NumLoops;
 	}
+
+	static void WithDecoratorLoopUntil(UBTCompositeNode& ParentNode, EBTNodeResult::Type RequiredResult)
+	{
+		UTestBTDecorator_LoopUntil& LoopUntilDecorator = WithDecorator<UTestBTDecorator_LoopUntil>(ParentNode);
+		LoopUntilDecorator.SetRequiredResult(RequiredResult);
+	}
 	
 	template<class T>
 	static T& WithService(UBTCompositeNode& ParentNode, UClass* ServiceClass = T::StaticClass())
diff --git a/Engine/Source/Runtime/AIModule/Classes/BehaviorTree/Decorators/BTDecorator_LoopUntil.h b/Engine/Source/Runtime/AIModule/Classes/BehaviorTree/Decorators/BTDecorator_LoopUntil.h
new file mode 100644
--- /dev/null
+++ b/Engine/Source/Runtime/AIModule/Classes/BehaviorTree/Decorators/BTDecorator_LoopUntil.h
@@ -0,0 +1,45 @@
+// Copyright Epic Games, Inc. All Rights Reserved.
+
+#pragma once
+
+#include "BehaviorTree/ValueOrBBKey.h"
+#include "BehaviorTree/BTDecorator.h"
+#include "BTDecorator_LoopUntil.generated.h"
+
+class UBlackboardComponent;
+
+struct FBTLoopUntilDecoratorMemory
+{
+	int32 SearchId = 0;
+};
+
+/**
+ * Loop until decorator node.
+ * A decorator node that loops execution until the child execution returns the required result
+ */
+UCLASS(HideCategories=(FlowControl, Condition), MinimalAPI)
+class UBTDecorator_LoopUntil : public UBTDecorator
+{
+	GENERATED_BODY()
+
+public:
+	AIMODULE_API UBTDecorator_LoopUntil(const FObjectInitializer& ObjectInitializer);
+
+private:
+	AIMODULE_API virtual void OnNodeDeactivation(FBehaviorTreeSearchData& SearchData, EBTNodeResult::Type NodeResult) override;
+
+	AIMODULE_API virtual uint16 GetInstanceMemorySize() const override;
+	AIMODULE_API virtual void InitializeMemory(UBehaviorTreeComponent& OwnerComp, uint8* NodeMemory, EBTMemoryInit::Type InitType) const override;
+	AIMODULE_API virtual void CleanupMemory(UBehaviorTreeComponent& OwnerComp, uint8* NodeMemory, EBTMemoryClear::Type CleanupType) const override;
+
+#if WITH_EDITOR
+	AIMODULE_API virtual FString GetErrorMessage() const;
+	AIMODULE_API virtual FName GetNodeIconName() const override;
+	AIMODULE_API virtual void PostEditChangeProperty(FPropertyChangedEvent& PropertyChangedEvent) override;
+#endif // WITH_EDITOR
+	void SetNodeName();
+
+protected:
+	UPROPERTY(EditAnywhere, Category=Decorator)
+	FValueOrBBKey_Enum RequiredResult = EBTNodeResult::Succeeded;
+};
diff --git a/Engine/Source/Runtime/AIModule/Private/BehaviorTree/Decorators/BTDecorator_LoopUntil.cpp b/Engine/Source/Runtime/AIModule/Private/BehaviorTree/Decorators/BTDecorator_LoopUntil.cpp
new file mode 100644
--- /dev/null
+++ b/Engine/Source/Runtime/AIModule/Private/BehaviorTree/Decorators/BTDecorator_LoopUntil.cpp
@@ -0,0 +1,114 @@
+// Copyright Epic Games, Inc. All Rights Reserved.
+
+#include "BehaviorTree/Decorators/BTDecorator_LoopUntil.h"
+#include "GameFramework/Actor.h"
+#include "VisualLogger/VisualLogger.h"
+#include "BehaviorTree/BTCompositeNode.h"
+
+#include UE_INLINE_GENERATED_CPP_BY_NAME(BTDecorator_LoopUntil)
+
+UBTDecorator_LoopUntil::UBTDecorator_LoopUntil(const FObjectInitializer& ObjectInitializer) : Super(ObjectInitializer)
+{
+	SetNodeName();
+	INIT_DECORATOR_NODE_NOTIFY_FLAGS();
+
+	bAllowAbortNone = false;
+	bAllowAbortLowerPri = false;
+	bAllowAbortChildNodes = false;
+}
+
+void UBTDecorator_LoopUntil::OnNodeDeactivation(FBehaviorTreeSearchData& SearchData, EBTNodeResult::Type NodeResult)
+{
+	FBTLoopUntilDecoratorMemory* DecoratorMemory = GetNodeMemory<FBTLoopUntilDecoratorMemory>(SearchData);
+
+	// protect from infinite loop within single search
+	if (DecoratorMemory->SearchId != SearchData.SearchId)
+	{
+		DecoratorMemory->SearchId = SearchData.SearchId;
+
+		// We only check aborted here since we can never be in progressed when deactivating
+		if (NodeResult != EBTNodeResult::Aborted)
+		{
+			const EBTNodeResult::Type ResolvedRequiredResult = static_cast<EBTNodeResult::Type>(RequiredResult.GetValue(SearchData.OwnerComp));
+			const bool bShouldLoop = NodeResult != ResolvedRequiredResult;
+			UE_VLOG(SearchData.OwnerComp.GetOwner(), LogBehaviorTree, Verbose, TEXT("Loop until %s: Node Result is %s -> %s"),
+				ResolvedRequiredResult == EBTNodeResult::Succeeded ? TEXT("Success") : TEXT("Failure"),
+				NodeResult == EBTNodeResult::Succeeded ? TEXT("Succeeded") : TEXT("Failed"),
+				bShouldLoop ? TEXT("Run Again!") : TEXT("Break"));
+
+			if (bShouldLoop)
+			{
+				GetParentNode()->SetChildOverride(SearchData, GetChildIndex());
+			}
+		}
+	}
+}
+
+uint16 UBTDecorator_LoopUntil::GetInstanceMemorySize() const
+{
+	return sizeof(FBTLoopUntilDecoratorMemory);
+}
+
+void UBTDecorator_LoopUntil::InitializeMemory(UBehaviorTreeComponent& OwnerComp, uint8* NodeMemory, EBTMemoryInit::Type InitType) const
+{
+	InitializeNodeMemory<FBTLoopUntilDecoratorMemory>(NodeMemory, InitType);
+}
+
+void UBTDecorator_LoopUntil::CleanupMemory(UBehaviorTreeComponent& OwnerComp, uint8* NodeMemory, EBTMemoryClear::Type CleanupType) const
+{
+	CleanupNodeMemory<FBTLoopUntilDecoratorMemory>(NodeMemory, CleanupType);
+}
+
+#if WITH_EDITOR
+
+FString UBTDecorator_LoopUntil::GetErrorMessage() const
+{
+	EBTNodeResult::Type DefaultResult = static_cast<EBTNodeResult::Type>(RequiredResult.GetValue(static_cast<UBlackboardComponent*>(nullptr)));
+	if (DefaultResult == EBTNodeResult::Aborted)
+	{
+		return TEXT("Can't use 'Aborted' as Required Result since the node won't be allowed to loop in that case");
+	}
+	else if (DefaultResult == EBTNodeResult::InProgress)
+	{
+		return TEXT("Can't use 'InProgress' as Required Result since the node can never be in that state when evaluating the condition to loop");
+	}
+
+	return Super::GetErrorMessage();
+}
+
+FName UBTDecorator_LoopUntil::GetNodeIconName() const
+{
+	return FName("BTEditor.Graph.BTNode.Decorator.Loop.Icon");
+}
+
+void UBTDecorator_LoopUntil::PostEditChangeProperty(FPropertyChangedEvent& PropertyChangedEvent)
+{
+	Super::PostEditChangeProperty(PropertyChangedEvent);
+	if (PropertyChangedEvent.Property && PropertyChangedEvent.Property->GetFName() == TEXT("DefaultValue"))
+	{
+		SetNodeName();
+	}
+	else if (PropertyChangedEvent.Property && PropertyChangedEvent.Property->GetFName() == TEXT("Key"))
+	{
+		SetNodeName();
+	}
+}
+
+#endif	// WITH_EDITOR
+
+void UBTDecorator_LoopUntil::SetNodeName()
+{
+	EBTNodeResult::Type DefaultResult = static_cast<EBTNodeResult::Type>(RequiredResult.GetValue(static_cast<UBlackboardComponent*>(nullptr)));
+	if (DefaultResult == EBTNodeResult::Succeeded)
+	{
+		NodeName = TEXT("Loop Until Success");
+	}
+	else if (DefaultResult == EBTNodeResult::Failed)
+	{
+		NodeName = TEXT("Loop Until Failure");
+	}
+	else
+	{
+		NodeName = TEXT("Loop Until");
+	}
+}
