PR #13065: AAIController::UseBlackboard - Fix to correctly returns the result of trying to initialize the Blackboard.

diff --git a/Engine/Source/Runtime/AIModule/Private/AIController.cpp b/Engine/Source/Runtime/AIModule/Private/AIController.cpp
--- a/Engine/Source/Runtime/AIModule/Private/AIController.cpp
+++ b/Engine/Source/Runtime/AIModule/Private/AIController.cpp
@@ -1110,21 +1110,20 @@ bool AAIController::UseBlackboard(UBlackboardData* BlackboardAsset, UBlackboardC
 		REDIRECT_OBJECT_TO_VLOG(Blackboard, this);
 		if (Blackboard != nullptr)
 		{
-			InitializeBlackboard(*Blackboard, *BlackboardAsset);
+			bSuccess = InitializeBlackboard(*Blackboard, *BlackboardAsset);
 			Blackboard->RegisterComponent();
 		}
-
 	}
 	else if (Blackboard->GetBlackboardAsset() == nullptr)
 	{
-		InitializeBlackboard(*Blackboard, *BlackboardAsset);
+		bSuccess = InitializeBlackboard(*Blackboard, *BlackboardAsset);
 	}
 	else if (Blackboard->GetBlackboardAsset() != BlackboardAsset)
 	{
 		// @todo this behavior should be opt-out-able.
 		UE_VLOG(this, LogBehaviorTree, Log, TEXT("UseBlackboard: requested blackboard %s while already has %s instantiated. Forcing new BB.")
 			, *GetNameSafe(BlackboardAsset), *GetNameSafe(Blackboard->GetBlackboardAsset()));
-		InitializeBlackboard(*Blackboard, *BlackboardAsset);
+		bSuccess = InitializeBlackboard(*Blackboard, *BlackboardAsset);
 	}
 
 	BlackboardComponent = Blackboard;
