PR #12942: [Mass Entity] Fix infinite loop caused by loop variable overflow

diff --git a/Engine/Plugins/AI/MassAI/Source/MassNavigation/Private/MassAvoidanceProcessors.cpp b/Engine/Plugins/AI/MassAI/Source/MassNavigation/Private/MassAvoidanceProcessors.cpp
--- a/Engine/Plugins/AI/MassAI/Source/MassNavigation/Private/MassAvoidanceProcessors.cpp
+++ b/Engine/Plugins/AI/MassAI/Source/MassNavigation/Private/MassAvoidanceProcessors.cpp
@@ -69,9 +69,11 @@ namespace UE::MassAvoidance
 		{
 			const FVector::FReal CellSize = AvoidanceObstacleGrid.GetCellSize(Level);
 			const FNavigationObstacleHashGrid2D::FCellRect Rect = AvoidanceObstacleGrid.CalcQueryBounds(QueryBox, Level);
-			for (int32 Y = Rect.MinY; Y <= Rect.MaxY; Y++)
+
+			// Use int64 to prevent overflow when MaxX or MaxY are equal to the maximum value of int32.
+			for (int64 Y = Rect.MinY; Y <= Rect.MaxY; Y++)
 			{
-				for (int32 X = Rect.MinX; X <= Rect.MaxX; X++)
+				for (int64 X = Rect.MinX; X <= Rect.MaxX; X++)
 				{
 					const FVector::FReal CenterX = (X + 0.5) * CellSize;
 					const FVector::FReal CenterY = (Y + 0.5) * CellSize;
diff --git a/Engine/Source/Runtime/AIModule/Public/HierarchicalHashGrid2D.h b/Engine/Source/Runtime/AIModule/Public/HierarchicalHashGrid2D.h
--- a/Engine/Source/Runtime/AIModule/Public/HierarchicalHashGrid2D.h
+++ b/Engine/Source/Runtime/AIModule/Public/HierarchicalHashGrid2D.h
@@ -111,8 +111,9 @@ public:
 	/** Iterator state over a rectangle */
 	struct FCellRectIterator
 	{
-		int32 X = 0;		/** Current X */
-		int32 Y = 0;		/** Current Y */
+		// Use int64 to prevent overflow when FCellRect::MaxX or FCellRect::MaxY are equal to the maximum value of int32.
+		int64 X = 0;		/** Current X */
+		int64 Y = 0;		/** Current Y */
 		int32 Level = 0;	/** Current level */
 		FCellRect Rect;		/** Area to iterate over */
 	};
@@ -352,11 +353,12 @@ public:
 			// Finest level query rect
 			FCellRect Rect = CalcQueryBounds(Bounds, Level);
 
-			for (int32 Y = Rect.MinY; Y <= Rect.MaxY; Y++)
+			// Use int64 to prevent overflow when FCellRect::MaxX or FCellRect::MaxY are equal to the maximum value of int32.
+			for (int64 Y = Rect.MinY; Y <= Rect.MaxY; Y++)
 			{
-				for (int32 X = Rect.MinX; X <= Rect.MaxX; X++)
+				for (int64 X = Rect.MinX; X <= Rect.MaxX; X++)
 				{
-					if (const FCell* Cell = FindCell(X, Y, Level))
+					if (const FCell* Cell = FindCell(IntCastChecked<int32>(X), IntCastChecked<int32>(Y), Level))
 					{
 						for (int32 Idx = Cell->First; Idx != INDEX_NONE; Idx = Items[Idx].Next)
 						{
@@ -420,7 +422,7 @@ public:
 				}
 			}
 
-			if (const FCell* Cell = FindCell(Iter.X, Iter.Y, Iter.Level))
+			if (const FCell* Cell = FindCell( IntCastChecked<int32>(Iter.X) , IntCastChecked<int32>(Iter.Y), Iter.Level))
 			{
 				// Collect items from this cell.
 				for (int32 Idx = Cell->First; Idx != INDEX_NONE; Idx = Items[Idx].Next)
@@ -435,7 +437,11 @@ public:
 					const int FinerLevel = Iter.Level - 1;
 					// The iteration rect is intersection between current coarse cell and finer level bounds (which is more accurate).
 					const FCellRect& FinerLevelRect = Rects[FinerLevel];
-					FCellRect CurrentRect(Iter.X * LevelRatio, Iter.Y * LevelRatio, Iter.X * LevelRatio + LevelRatio - 1, Iter.Y * LevelRatio + LevelRatio - 1);
+					const int32 MinX = ClampInt32(Iter.X * LevelRatio);
+					const int32 MinY = ClampInt32(Iter.Y * LevelRatio);
+					const int32 MaxX = ClampInt32(Iter.X * LevelRatio + LevelRatio - 1);
+					const int32 MaxY = ClampInt32(Iter.Y * LevelRatio + LevelRatio - 1);
+					FCellRect CurrentRect(MinX, MinY, MaxX, MaxY);
 					FCellRect NewIterRect = IntersectRect(CurrentRect, FinerLevelRect);
 					// Advance if rect is not empty.
 					if (NewIterRect.MaxX >= NewIterRect.MinX && NewIterRect.MaxY >= NewIterRect.MinY)
