Check that sense is instantiated

#tests PIE
#rnx

diff --git a/Engine/Source/Runtime/AIModule/Private/Perception/AIPerceptionSystem.cpp b/Engine/Source/Runtime/AIModule/Private/Perception/AIPerceptionSystem.cpp
--- a/Engine/Source/Runtime/AIModule/Private/Perception/AIPerceptionSystem.cpp
+++ b/Engine/Source/Runtime/AIModule/Private/Perception/AIPerceptionSystem.cpp
@@ -380,7 +380,7 @@ void UAIPerceptionSystem::UnregisterSource(AActor& SourceActor, const TSubclassO
 		{
 			// Unregister the source actor from a single sense
 			const FAISenseID SenseID = UAISense::GetSenseID(Sense);
-			if (Senses[SenseID] != nullptr && StimuliSource->RelevantSenses.ShouldRespondToChannel(Senses[SenseID]->GetSenseID()))
+			if (IsSenseInstantiated(SenseID) && StimuliSource->RelevantSenses.ShouldRespondToChannel(Senses[SenseID]->GetSenseID()))
 			{
 				Senses[SenseID]->UnregisterSource(SourceActor);
 				StimuliSource->RelevantSenses.FilterOutChannel(SenseID);
