Add option to normalize scores on query generator
Can now be Default, Unaltered and Normalized
Defaults to Default which keeps the current behavior
Unaltered doesn't normalize and Normalized normalizes

diff --git a/Engine/Source/Runtime/AIModule/Classes/EnvironmentQuery/EnvQueryGenerator.h b/Engine/Source/Runtime/AIModule/Classes/EnvironmentQuery/EnvQueryGenerator.h
--- a/Engine/Source/Runtime/AIModule/Classes/EnvironmentQuery/EnvQueryGenerator.h
+++ b/Engine/Source/Runtime/AIModule/Classes/EnvironmentQuery/EnvQueryGenerator.h
@@ -5,8 +5,9 @@
 #include "CoreMinimal.h"
 #include "UObject/ObjectMacros.h"
 #include "Templates/SubclassOf.h"
-#include "EnvironmentQuery/Items/EnvQueryItemType.h"
 #include "EnvironmentQuery/EnvQueryNode.h"
+#include "EnvironmentQuery/EnvQueryTypes.h"
+#include "EnvironmentQuery/Items/EnvQueryItemType.h"
 #include "EnvQueryGenerator.generated.h"
 
 namespace EnvQueryGeneratorVersion
@@ -36,6 +37,7 @@ class UEnvQueryGenerator : public UEnvQueryNode
 	virtual void GenerateItems(FEnvQueryInstance& QueryInstance) const { checkNoEntry(); }
 	virtual bool IsValidGenerator() const { return ItemType != nullptr; }
 	bool CanRunAsync() const { return bCanRunAsync; }
+	EEnvQueryResultNormalizationOption GetNormalizationOption() const;
 
 	AIMODULE_API virtual void PostLoad() override;
 	AIMODULE_API void UpdateNodeVersion() override;
@@ -44,4 +46,12 @@ protected:
 	/** To be overwritten by MassEnvQueryGenerators to indicate that they will run asynchronously. */
 	UPROPERTY(EditDefaultsOnly, Category = Option, AdvancedDisplay)
 	uint32 bCanRunAsync : 1;
+
+	/** Should the scores be normalized or not before returning the result. 
+	 *  If set to UEnvQueryResultNormalizationOption::Default, only runmode EEnvQueryRunMode::AllMatching is normalized and the other modes are not
+	 *  If set to UEnvQueryResultNormalizationOption::Normalized, all run modes will be normalized.
+	 *  If set to UEnvQueryResultNormalizationOption::Unaltered, all run modes will NOT be normalized.
+	 */
+	UPROPERTY(EditDefaultsOnly, Category = Option, AdvancedDisplay)
+	EEnvQueryResultNormalizationOption EnvQueryResultNormalizationOption = EEnvQueryResultNormalizationOption::Default;
 };
diff --git a/Engine/Source/Runtime/AIModule/Classes/EnvironmentQuery/EnvQueryTypes.h b/Engine/Source/Runtime/AIModule/Classes/EnvironmentQuery/EnvQueryTypes.h
--- a/Engine/Source/Runtime/AIModule/Classes/EnvironmentQuery/EnvQueryTypes.h
+++ b/Engine/Source/Runtime/AIModule/Classes/EnvironmentQuery/EnvQueryTypes.h
@@ -194,6 +194,14 @@ namespace EEnvQueryRunMode
 	};
 }
 
+UENUM()
+enum class EEnvQueryResultNormalizationOption : uint8
+{
+	Default,
+	Normalized,
+	Unaltered
+};
+
 UENUM()
 namespace EEnvQueryParam
 {
diff --git a/Engine/Source/Runtime/AIModule/Private/EnvironmentQuery/EnvQueryGenerator.cpp b/Engine/Source/Runtime/AIModule/Private/EnvironmentQuery/EnvQueryGenerator.cpp
--- a/Engine/Source/Runtime/AIModule/Private/EnvironmentQuery/EnvQueryGenerator.cpp
+++ b/Engine/Source/Runtime/AIModule/Private/EnvironmentQuery/EnvQueryGenerator.cpp
@@ -15,6 +15,11 @@ void UEnvQueryGenerator::UpdateNodeVersion()
 	VerNum = EnvQueryGeneratorVersion::Latest;
 }
 
+EEnvQueryResultNormalizationOption UEnvQueryGenerator::GetNormalizationOption() const
+{
+	return EnvQueryResultNormalizationOption;
+}
+
 void UEnvQueryGenerator::PostLoad()
 {
 	Super::PostLoad();
diff --git a/Engine/Source/Runtime/AIModule/Private/EnvironmentQuery/EnvQueryInstance.cpp b/Engine/Source/Runtime/AIModule/Private/EnvironmentQuery/EnvQueryInstance.cpp
--- a/Engine/Source/Runtime/AIModule/Private/EnvironmentQuery/EnvQueryInstance.cpp
+++ b/Engine/Source/Runtime/AIModule/Private/EnvironmentQuery/EnvQueryInstance.cpp
@@ -799,18 +799,46 @@ void FEnvQueryInstance::FinalizeQuery()
 {
 	if (NumValidItems > 0)
 	{
+		EEnvQueryResultNormalizationOption NormalizationOption = EEnvQueryResultNormalizationOption::Default;
+		if (const UEnvQueryGenerator* Generator = Options[OptionIndex].Generator)
+		{
+			NormalizationOption = Generator->GetNormalizationOption();
+		}
+
 		if (Mode == EEnvQueryRunMode::SingleResult)
 		{
 			// if last test was not pure condition: sort and pick one of best items
 			if (bFoundSingleResult == false && bPassOnSingleResult == false)
 			{
 				SortScores();
+
+				switch (NormalizationOption)
+				{
+				case EEnvQueryResultNormalizationOption::Default: // falling through on purpose
+				case EEnvQueryResultNormalizationOption::Unaltered:
+					break;
+				case EEnvQueryResultNormalizationOption::Normalized:
+					NormalizeScores();
+					break;
+				}
+
 				PickSingleItem(0);
 			}
 		}
 		else if (Mode == EEnvQueryRunMode::RandomBest5Pct || Mode == EEnvQueryRunMode::RandomBest25Pct)
 		{
 			SortScores();
+
+			switch (NormalizationOption)
+			{
+			case EEnvQueryResultNormalizationOption::Default: // falling through on purpose
+			case EEnvQueryResultNormalizationOption::Unaltered:
+				break;
+			case EEnvQueryResultNormalizationOption::Normalized:
+				NormalizeScores();
+				break;
+			}
+
 			const float ScoreRangePct = (Mode == EEnvQueryRunMode::RandomBest5Pct) ? 0.95f : 0.75f;
 			PickRandomItemOfScoreAtLeast(Items[0].Score * ScoreRangePct);
 		}
@@ -827,7 +855,15 @@ void FEnvQueryInstance::FinalizeQuery()
 
 			// normalizing after scoring and reducing number of elements to not 
 			// do anything for discarded items
-			NormalizeScores();
+			switch (NormalizationOption)
+			{
+			case EEnvQueryResultNormalizationOption::Unaltered:
+				break;
+			case EEnvQueryResultNormalizationOption::Default: // falling through on purpose
+			case EEnvQueryResultNormalizationOption::Normalized:
+				NormalizeScores();
+				break;
+			}
 		}
 
 		MarkAsFinishedWithoutIssues();
