Fixup API macros

#rnx

diff --git a/Engine/Source/Runtime/AIModule/Classes/BrainComponent.h b/Engine/Source/Runtime/AIModule/Classes/BrainComponent.h
--- a/Engine/Source/Runtime/AIModule/Classes/BrainComponent.h
+++ b/Engine/Source/Runtime/AIModule/Classes/BrainComponent.h
@@ -195,7 +195,7 @@ public:
 	UBlackboardComponent* GetBlackboardComponent();
 
 	/** @return blackboard used with this component */
-	AIMODULE_API const UBlackboardComponent* GetBlackboardComponent() const;
+	const UBlackboardComponent* GetBlackboardComponent() const;
 
 	inline AAIController* GetAIOwner() const { return AIOwner; }
 
