PR #13496: Updating the Description Details of UEnvQueryTest_Pathfinding to better reflect its behavior in EEnvTestPurpose::Score mode.

diff --git a/Engine/Source/Runtime/AIModule/Private/EnvironmentQuery/Tests/EnvQueryTest_Pathfinding.cpp b/Engine/Source/Runtime/AIModule/Private/EnvironmentQuery/Tests/EnvQueryTest_Pathfinding.cpp
--- a/Engine/Source/Runtime/AIModule/Private/EnvironmentQuery/Tests/EnvQueryTest_Pathfinding.cpp
+++ b/Engine/Source/Runtime/AIModule/Private/EnvironmentQuery/Tests/EnvQueryTest_Pathfinding.cpp
@@ -131,15 +131,19 @@ FText UEnvQueryTest_Pathfinding::GetDescriptionTitle() const
 
 FText UEnvQueryTest_Pathfinding::GetDescriptionDetails() const
 {
-	FText DiscardDesc = LOCTEXT("DiscardUnreachable", "discard unreachable");
+	// When in EEnvTestPurpose::Score mode, this test will not discard unreachable items, but rather score them as 0.
+	FText DiscardOrScoreDesc = (TestPurpose == EEnvTestPurpose::Score)
+		? LOCTEXT("ScoreUnreachable", "score unreachable as 0")
+		: LOCTEXT("DiscardUnreachable", "discard unreachable");
+	
 	FText Desc2;
 	if (SkipUnreachable.IsDynamic())
 	{
-		Desc2 = FText::Format(FText::FromString("{0}: {1}"), DiscardDesc, FText::FromString(SkipUnreachable.ToString()));
+		Desc2 = FText::Format(FText::FromString("{0}: {1}"), DiscardOrScoreDesc, FText::FromString(SkipUnreachable.ToString()));
 	}
 	else if (SkipUnreachable.DefaultValue)
 	{
-		Desc2 = DiscardDesc;
+		Desc2 = MoveTemp(DiscardOrScoreDesc);
 	}
 
 	FText TestParamDesc = GetWorkOnFloatValues() ? DescribeFloatTestParams() : DescribeBoolTestParams("existing path");
