PR #13525: Guard against zero DeltaTime in UPathFollowingComponent::FollowPathSegment to prevent division by zero

diff --git a/Engine/Source/Runtime/AIModule/Private/Navigation/PathFollowingComponent.cpp b/Engine/Source/Runtime/AIModule/Private/Navigation/PathFollowingComponent.cpp
--- a/Engine/Source/Runtime/AIModule/Private/Navigation/PathFollowingComponent.cpp
+++ b/Engine/Source/Runtime/AIModule/Private/Navigation/PathFollowingComponent.cpp
@@ -1147,7 +1147,7 @@ void UPathFollowingComponent::FollowPathSegment(float DeltaTime)
 		PostProcessMove.ExecuteIfBound(this, CurrentMoveInput);
 		NavMovementInterface->RequestPathMove(CurrentMoveInput);
 	}
-	else
+	else if (DeltaTime != 0.f)
 	{
 		FVector MoveVelocity = (CurrentTarget - CurrentLocation) / DeltaTime;
 
