[NavLinkProxy]
- Avoid triggering complex UpdateComponentToWorld from inside the object constructor

diff --git a/Engine/Source/Runtime/AIModule/Private/Navigation/NavLinkProxy.cpp b/Engine/Source/Runtime/AIModule/Private/Navigation/NavLinkProxy.cpp
--- a/Engine/Source/Runtime/AIModule/Private/Navigation/NavLinkProxy.cpp
+++ b/Engine/Source/Runtime/AIModule/Private/Navigation/NavLinkProxy.cpp
@@ -90,7 +90,7 @@ ANavLinkProxy::ANavLinkProxy(const FObjectInitializer& ObjectInitializer) : Supe
 		SpriteComponent->SpriteInfo.Category = ConstructorStatics.ID_Decals;
 		SpriteComponent->SpriteInfo.DisplayName = ConstructorStatics.NAME_Decals;
 		SpriteComponent->SetupAttachment(RootComponent);
-		SpriteComponent->SetAbsolute(false, false, true);
+		SpriteComponent->SetUsingAbsoluteScale(true);
 		SpriteComponent->bIsScreenSizeScaled = true;
 	}
 #endif
