[BehaviorTree] add multiple missing referenced objects from the BehaviorTree component to fix nodes reinstantiation at runtime

diff --git a/Engine/Source/Runtime/AIModule/Classes/BehaviorTree/BTNode.h b/Engine/Source/Runtime/AIModule/Classes/BehaviorTree/BTNode.h
--- a/Engine/Source/Runtime/AIModule/Classes/BehaviorTree/BTNode.h
+++ b/Engine/Source/Runtime/AIModule/Classes/BehaviorTree/BTNode.h
@@ -43,7 +43,7 @@ struct FBehaviorTreeNodeDebugContext
 	FString RuntimeDescription;
 };
 
-UCLASS(Abstract,config=Game, MinimalAPI)
+UCLASS(Abstract, config=Game, MinimalAPI)
 class UBTNode : public UObject, public IGameplayTaskOwnerInterface
 {
 	GENERATED_UCLASS_BODY()
@@ -233,27 +233,34 @@ private:
 #endif
 
 	/** depth first index (execution order) */
+	UPROPERTY(transient)
 	uint16 ExecutionIndex;
 
 	/** instance memory offset */
+	UPROPERTY(transient)
 	uint16 MemoryOffset;
 
 	/** depth in tree */
+	UPROPERTY(transient)
 	uint8 TreeDepth;
 
 	/** set automatically for node instances. Should never be set manually */
+	UPROPERTY(transient)
 	uint8 bIsInstanced : 1;
 
 	/** if set, node is injected by subtree. Should never be set manually */
+	UPROPERTY(transient)
 	uint8 bIsInjected : 1;
 
 protected:
 
 	/** if set, node will be instanced instead of using memory block and template shared with all other BT components */
+	UPROPERTY(transient)
 	uint8 bCreateNodeInstance : 1;
 
 	/** set to true if task owns any GameplayTasks. Note this requires tasks to be created via NewBTAITask
 	 *	Otherwise specific BT task node class is responsible for ending the gameplay tasks on node finish */
+	UPROPERTY(transient)
 	uint8 bOwnsGameplayTasks : 1;
 };
 
diff --git a/Engine/Source/Runtime/AIModule/Classes/BehaviorTree/BehaviorTreeComponent.h b/Engine/Source/Runtime/AIModule/Classes/BehaviorTree/BehaviorTreeComponent.h
--- a/Engine/Source/Runtime/AIModule/Classes/BehaviorTree/BehaviorTreeComponent.h
+++ b/Engine/Source/Runtime/AIModule/Classes/BehaviorTree/BehaviorTreeComponent.h
@@ -31,7 +31,7 @@ struct FBTNodeExecutionInfo
 	FBTNodeIndex SearchEnd;
 
 	/** node to be executed */
-	const UBTCompositeNode* ExecuteNode;
+	TObjectPtr<const UBTCompositeNode> ExecuteNode;
 
 	/** subtree index */
 	uint16 ExecuteInstanceIdx;
@@ -51,7 +51,7 @@ struct FBTNodeExecutionInfo
 struct FBTPendingExecutionInfo
 {
 	/** next task to execute */
-	UBTTaskNode* NextTask;
+	TObjectPtr<UBTTaskNode> NextTask;
 
 	/** if set, tree ran out of nodes */
 	uint32 bOutOfNodes : 1;
@@ -298,6 +298,8 @@ public:
 	void SetCSVTickStatName(const char* InCSVTickStatName) { CSVTickStatName = InCSVTickStatName; }
 #endif
 
+	static AIMODULE_API void AddReferencedObjects(UObject* InThis, FReferenceCollector& Collector);
+
 protected:
 	/** stack of behavior tree instances */
 	TArray<FBehaviorTreeInstance> InstanceStack;
@@ -332,7 +334,7 @@ protected:
 			, ContinueWithResult(InContinueWithResult)
 			, Action(InAction)
 		{}
-		const UBTNode* Node = nullptr;
+		TObjectPtr<const UBTNode> Node = nullptr;
 		EBTNodeResult::Type ContinueWithResult = EBTNodeResult::Succeeded;
 		EBTBranchAction Action;
 	};
diff --git a/Engine/Source/Runtime/AIModule/Classes/BehaviorTree/BehaviorTreeTypes.h b/Engine/Source/Runtime/AIModule/Classes/BehaviorTree/BehaviorTreeTypes.h
--- a/Engine/Source/Runtime/AIModule/Classes/BehaviorTree/BehaviorTreeTypes.h
+++ b/Engine/Source/Runtime/AIModule/Classes/BehaviorTree/BehaviorTreeTypes.h
@@ -190,7 +190,7 @@ namespace EBTNodeUpdateMode
 struct FBehaviorTreeParallelTask
 {
 	/** worker object */
-	const UBTTaskNode* TaskNode;
+	TObjectPtr<const UBTTaskNode> TaskNode;
 
 	/** additional mode data used for context switching */
 	EBTTaskStatus::Type Status;
@@ -290,10 +290,10 @@ struct FBehaviorTreeExecutionStep
 struct FBehaviorTreeInstanceId
 {
 	/** behavior tree asset */
-	UBehaviorTree* TreeAsset;
+	TObjectPtr<UBehaviorTree> TreeAsset;
 
 	/** root node in template for cleanup purposes */
-	UBTCompositeNode* RootNode;
+	TObjectPtr<UBTCompositeNode> RootNode;
 
 	/** execution index path from root */
 	TArray<uint16> Path;
@@ -319,13 +319,13 @@ DECLARE_DELEGATE_TwoParams(FBTInstanceDeactivation, UBehaviorTreeComponent&, EBT
 struct FBehaviorTreeInstance
 {
 	/** root node in template */
-	UBTCompositeNode* RootNode;
+	TObjectPtr<UBTCompositeNode> RootNode;
 
 	/** active node in template */
-	UBTNode* ActiveNode;
+	TObjectPtr<UBTNode> ActiveNode;
 
 	/** active auxiliary nodes */
-	TArray<UBTAuxiliaryNode*> ActiveAuxNodes;
+	TArray<TObjectPtr<UBTAuxiliaryNode>> ActiveAuxNodes;
 
 	/** active parallel tasks */
 	TArray<FBehaviorTreeParallelTask> ParallelTasks;
@@ -495,8 +495,8 @@ struct FBTNodeIndexRange
 /** node update data */
 struct FBehaviorTreeSearchUpdate
 {
-	UBTAuxiliaryNode* AuxNode;
-	UBTTaskNode* TaskNode;
+	TObjectPtr<UBTAuxiliaryNode> AuxNode;
+	TObjectPtr<UBTTaskNode> TaskNode;
 
 	uint16 InstanceIndex;
 
diff --git a/Engine/Source/Runtime/AIModule/Private/BehaviorTree/BehaviorTreeComponent.cpp b/Engine/Source/Runtime/AIModule/Private/BehaviorTree/BehaviorTreeComponent.cpp
--- a/Engine/Source/Runtime/AIModule/Private/BehaviorTree/BehaviorTreeComponent.cpp
+++ b/Engine/Source/Runtime/AIModule/Private/BehaviorTree/BehaviorTreeComponent.cpp
@@ -1069,8 +1069,8 @@ void UBehaviorTreeComponent::EvaluateBranch(EBTNodeResult::Type ContinueWithResu
 	if (InstanceStack.IsValidIndex(ActiveInstanceIdx))
 	{
 		const FBehaviorTreeInstance& ActiveInstance = InstanceStack[ActiveInstanceIdx];
-		UBTCompositeNode* ExecuteParent = (ActiveInstance.ActiveNode == NULL) ? ActiveInstance.RootNode :
-			(ActiveInstance.ActiveNodeType == EBTActiveNode::Composite) ? (UBTCompositeNode*)ActiveInstance.ActiveNode :
+		UBTCompositeNode* ExecuteParent = (ActiveInstance.ActiveNode == NULL) ? ActiveInstance.RootNode.Get() :
+			(ActiveInstance.ActiveNodeType == EBTActiveNode::Composite) ? (UBTCompositeNode*)ActiveInstance.ActiveNode.Get() :
 			ActiveInstance.ActiveNode->GetParentNode();
 
 		RequestExecution(ExecuteParent, InstanceStack.Num() - 1,
@@ -1327,8 +1327,8 @@ void UBehaviorTreeComponent::RequestExecution(const UBTCompositeNode* RequestedO
 		if (ExecutionRequest.ExecuteNode == NULL)
 		{
 			FBehaviorTreeInstance& ActiveInstance = InstanceStack[ActiveInstanceIdx];
-			CurrentNode = (ActiveInstance.ActiveNode == NULL) ? ActiveInstance.RootNode :
-				(ActiveInstance.ActiveNodeType == EBTActiveNode::Composite) ? (UBTCompositeNode*)ActiveInstance.ActiveNode :
+			CurrentNode = (ActiveInstance.ActiveNode == NULL) ? ActiveInstance.RootNode.Get() :
+				(ActiveInstance.ActiveNodeType == EBTActiveNode::Composite) ? (UBTCompositeNode*)ActiveInstance.ActiveNode.Get() :
 				ActiveInstance.ActiveNode->GetParentNode();
 
 			CurrentInstanceIdx = ActiveInstanceIdx;
@@ -3231,6 +3231,38 @@ void UBehaviorTreeComponent::DescribeSelfToVisLog(FVisualLogEntry* Snapshot) con
 
 #endif // ENABLE_VISUAL_LOG
 
+void UBehaviorTreeComponent::AddReferencedObjects(UObject* InThis, FReferenceCollector& Collector)
+{
+	UBehaviorTreeComponent* This = CastChecked<UBehaviorTreeComponent>(InThis);
+	for (FBehaviorTreeInstance& ActiveInstance : This->InstanceStack)
+	{
+		Collector.AddReferencedObject(ActiveInstance.RootNode);
+		Collector.AddReferencedObject(ActiveInstance.ActiveNode);
+		Collector.AddReferencedObjects(ActiveInstance.ActiveAuxNodes);
+		for (FBehaviorTreeParallelTask& ParallelTask : ActiveInstance.ParallelTasks)
+		{
+			Collector.AddReferencedObject(ParallelTask.TaskNode);
+		}
+	}
+
+	for (FBehaviorTreeSearchUpdate& PendingUpdate : This->SearchData.PendingUpdates)
+	{
+		Collector.AddReferencedObject(PendingUpdate.TaskNode);
+		Collector.AddReferencedObject(PendingUpdate.AuxNode);
+	}
+
+	Collector.AddReferencedObject(This->ExecutionRequest.ExecuteNode);
+	Collector.AddReferencedObject(This->PendingExecution.NextTask);
+
+	for (FBranchActionInfo& PendingBranchActionRequest : This->PendingBranchActionRequests)
+	{
+		Collector.AddReferencedObject(PendingBranchActionRequest.Node);
+	}
+
+
+	Super::AddReferencedObjects(InThis, Collector);
+}
+
 void UBehaviorTreeComponent::StoreDebuggerExecutionStep(EBTExecutionSnap::Type SnapType)
 {
 #if USE_BEHAVIORTREE_DEBUGGER
