Tentative fix for crash in UCrowdManager::PrepareAgentStep()

diff --git a/Engine/Source/Runtime/AIModule/Private/Navigation/CrowdFollowingComponent.cpp b/Engine/Source/Runtime/AIModule/Private/Navigation/CrowdFollowingComponent.cpp
--- a/Engine/Source/Runtime/AIModule/Private/Navigation/CrowdFollowingComponent.cpp
+++ b/Engine/Source/Runtime/AIModule/Private/Navigation/CrowdFollowingComponent.cpp
@@ -449,6 +449,7 @@ void UCrowdFollowingComponent::SetCrowdSimulationState(ECrowdSimulationState New
 		ICrowdAgentInterface* IAgent = Cast<ICrowdAgentInterface>(this);
 		if (bNeedRegistration)
 		{
+			check(IAgent != nullptr);			
 			Manager->RegisterAgent(IAgent);
 		}
 		else
@@ -1100,6 +1101,7 @@ void UCrowdFollowingComponent::RegisterCrowdAgent()
 	if (CrowdManager)
 	{
 		ICrowdAgentInterface* IAgent = Cast<ICrowdAgentInterface>(this);
+		check(IAgent != nullptr);
 		CrowdManager->RegisterAgent(IAgent);
 		bRegisteredWithCrowdSimulation = true;
 	}
diff --git a/Engine/Source/Runtime/AIModule/Private/Navigation/CrowdManager.cpp b/Engine/Source/Runtime/AIModule/Private/Navigation/CrowdManager.cpp
--- a/Engine/Source/Runtime/AIModule/Private/Navigation/CrowdManager.cpp
+++ b/Engine/Source/Runtime/AIModule/Private/Navigation/CrowdManager.cpp
@@ -338,6 +338,12 @@ void UCrowdManager::PostEditChangeProperty(FPropertyChangedEvent& PropertyChange
 
 void UCrowdManager::RegisterAgent(ICrowdAgentInterface* Agent)
 {
+	if (Agent == nullptr)
+	{
+		UE_LOG(LogCrowdFollowing, Warning, TEXT("%hs, Trying to add an invalid agent interface pointer."), __FUNCTION__);
+		return;
+	}
+	
 	UpdateNavData();
 
 	FCrowdAgentData AgentData;
