[Navigation] Addition of IsValid() check to narrow down crash source.

diff --git a/Engine/Source/Runtime/AIModule/Private/Navigation/PathFollowingComponent.cpp b/Engine/Source/Runtime/AIModule/Private/Navigation/PathFollowingComponent.cpp
--- a/Engine/Source/Runtime/AIModule/Private/Navigation/PathFollowingComponent.cpp
+++ b/Engine/Source/Runtime/AIModule/Private/Navigation/PathFollowingComponent.cpp
@@ -980,13 +980,14 @@ void UPathFollowingComponent::UpdatePathSegment()
 	DEBUG_bMovingDirectlyToGoal = false;
 #endif // !UE_BUILD_SHIPPING
 
-	if ((Path.IsValid() == false) || (NavMovementInterface == nullptr))
+	if ((Path.IsValid() == false) || !NavMovementInterface.IsValid())
 	{
 		UE_CVLOG(Path.IsValid() == false, this, LogPathFollowing, Log, TEXT("Aborting move due to not having a valid path object"));
 		OnPathFinished(EPathFollowingResult::Aborted, FPathFollowingResultFlags::InvalidPath);
 		return;
 	}
 
+	checkf(NavMovementInterface.IsValid() , TEXT("Temporary comment: if you hit this, check other threads to see what could stump the NavMovementInterface."));
 	const FVector CurrentLocation = NavMovementInterface->GetFeetLocation();
 	const bool bCanUpdateState = HasMovementAuthority();
 
