EnvQueryTest_GameplayTags: Added missing config related specifiers for UClass

#rnx

diff --git a/Engine/Source/Runtime/AIModule/Classes/EnvironmentQuery/Tests/EnvQueryTest_GameplayTags.h b/Engine/Source/Runtime/AIModule/Classes/EnvironmentQuery/Tests/EnvQueryTest_GameplayTags.h
--- a/Engine/Source/Runtime/AIModule/Classes/EnvironmentQuery/Tests/EnvQueryTest_GameplayTags.h
+++ b/Engine/Source/Runtime/AIModule/Classes/EnvironmentQuery/Tests/EnvQueryTest_GameplayTags.h
@@ -14,7 +14,7 @@ class IGameplayTagAssetInterface;
  * EnvQueryTest_GameplayTags attempts to cast items to IGameplayTagAssetInterface and test their tags with TagQueryToMatch.
  * The behavior of IGameplayTagAssetInterface-less items is configured by bRejectIncompatibleItems.
  */
-UCLASS(MinimalAPI)
+UCLASS(MinimalAPI, config = Game, defaultconfig)
 class UEnvQueryTest_GameplayTags : public UEnvQueryTest
 {
 	GENERATED_UCLASS_BODY()
