EnvQueryTest_GameplayTags: Searching for tags container in actor's components if actor itself doesn't implement IGameplayTagAssetInterface.

diff --git a/Engine/Source/Runtime/AIModule/Classes/EnvironmentQuery/Tests/EnvQueryTest_GameplayTags.h b/Engine/Source/Runtime/AIModule/Classes/EnvironmentQuery/Tests/EnvQueryTest_GameplayTags.h
--- a/Engine/Source/Runtime/AIModule/Classes/EnvironmentQuery/Tests/EnvQueryTest_GameplayTags.h
+++ b/Engine/Source/Runtime/AIModule/Classes/EnvironmentQuery/Tests/EnvQueryTest_GameplayTags.h
@@ -48,6 +48,11 @@ protected:
 	UPROPERTY(EditAnywhere, Category=GameplayTagCheck)
 	FGameplayTagQuery TagQueryToMatch;
 
+	// If the actor-item being tested doesn't implement IGameplayTagAssetInterface this flag controls whether the actor's components will be traversed in search of a one implementing IGameplayTagAssetInterface. 
+	// Note that only the first found component will be used. 
+	UPROPERTY(EditAnywhere, Config, Category=GameplayTagCheck)
+	bool bSearchComponentsForTags = false;
+
 	// This controls how to treat actors that do not implement IGameplayTagAssetInterface.
 	// When set to True, actors that do not implement the interface will be ignored, meaning
 	// they will not be scored and will not be considered when filtering.
diff --git a/Engine/Source/Runtime/AIModule/Private/EnvironmentQuery/Tests/EnvQueryTest_GameplayTags.cpp b/Engine/Source/Runtime/AIModule/Private/EnvironmentQuery/Tests/EnvQueryTest_GameplayTags.cpp
--- a/Engine/Source/Runtime/AIModule/Private/EnvironmentQuery/Tests/EnvQueryTest_GameplayTags.cpp
+++ b/Engine/Source/Runtime/AIModule/Private/EnvironmentQuery/Tests/EnvQueryTest_GameplayTags.cpp
@@ -103,7 +103,13 @@ void UEnvQueryTest_GameplayTags::RunTest(FEnvQueryInstance& QueryInstance) const
 	for (FEnvQueryInstance::ItemIterator It(this, QueryInstance); It; ++It)
 	{
 		const AActor* ItemActor = GetItemActor(QueryInstance, It.GetIndex());
-		if (const IGameplayTagAssetInterface* GameplayTagAssetInterface = Cast<const IGameplayTagAssetInterface>(ItemActor))
+		const IGameplayTagAssetInterface* GameplayTagAssetInterface = Cast<const IGameplayTagAssetInterface>(ItemActor);
+		if (bSearchComponentsForTags && GameplayTagAssetInterface == nullptr)
+		{
+			GameplayTagAssetInterface = ItemActor->FindComponentByInterface<const IGameplayTagAssetInterface>();
+		}
+
+		if (GameplayTagAssetInterface)
 		{
 			bool bSatisfiesTest = SatisfiesTest(GameplayTagAssetInterface);
 
