[EQS] fixed reinstancing crashes by properly exposing some properties to the reflection system

diff --git a/Engine/Source/Runtime/AIModule/Classes/EnvironmentQuery/EQSTestingPawn.h b/Engine/Source/Runtime/AIModule/Classes/EnvironmentQuery/EQSTestingPawn.h
--- a/Engine/Source/Runtime/AIModule/Classes/EnvironmentQuery/EQSTestingPawn.h
+++ b/Engine/Source/Runtime/AIModule/Classes/EnvironmentQuery/EQSTestingPawn.h
@@ -80,6 +80,7 @@ private:
 protected:
 	TSharedPtr<FEnvQueryInstance> QueryInstance;
 
+	UPROPERTY(Transient)
 	TArray<FEnvQueryInstance> StepResults;
 
 public:
@@ -94,6 +95,10 @@ public:
 	AIMODULE_API void OnPropertyChanged(const FName PropName);
 #endif // WITH_EDITOR
 
+	//~ Begin UObject Interface
+	static AIMODULE_API void AddReferencedObjects(UObject* InThis, FReferenceCollector& Collector);
+	//~ End UObject Interface
+
 	// IEQSQueryResultSourceInterface start
 	AIMODULE_API virtual const FEnvQueryResult* GetQueryResult() const override;
 	AIMODULE_API virtual const FEnvQueryInstance* GetQueryInstance() const  override;
diff --git a/Engine/Source/Runtime/AIModule/Classes/EnvironmentQuery/EnvQueryManager.h b/Engine/Source/Runtime/AIModule/Classes/EnvironmentQuery/EnvQueryManager.h
--- a/Engine/Source/Runtime/AIModule/Classes/EnvironmentQuery/EnvQueryManager.h
+++ b/Engine/Source/Runtime/AIModule/Classes/EnvironmentQuery/EnvQueryManager.h
@@ -136,6 +136,7 @@ struct FEnvQueryInstanceCache
 	TObjectPtr<UEnvQuery> Template = nullptr;
 
 	/** instance to duplicate */
+	UPROPERTY(Transient)
 	FEnvQueryInstance Instance;
 
 	/** the name of the source query */
diff --git a/Engine/Source/Runtime/AIModule/Classes/EnvironmentQuery/EnvQueryTypes.h b/Engine/Source/Runtime/AIModule/Classes/EnvironmentQuery/EnvQueryTypes.h
--- a/Engine/Source/Runtime/AIModule/Classes/EnvironmentQuery/EnvQueryTypes.h
+++ b/Engine/Source/Runtime/AIModule/Classes/EnvironmentQuery/EnvQueryTypes.h
@@ -637,9 +637,13 @@ struct FEnvQueryItemDetails
 	}
 };
 
+USTRUCT()
 struct FEnvQueryContextData
 {
+	GENERATED_BODY()
+
 	/** type of context values */
+	UPROPERTY(Transient)
 	TSubclassOf<UEnvQueryItemType> ValueType;
 
 	/** number of stored values */
@@ -655,24 +659,33 @@ struct FEnvQueryContextData
 	inline uint32 GetAllocatedSize() const { return sizeof(*this) + RawData.GetAllocatedSize(); }
 };
 
+USTRUCT()
 struct FEnvQueryOptionInstance
 {
-	/** generator object, raw pointer can be used safely because it will be always referenced by EnvQueryManager */
-	UEnvQueryGenerator* Generator;
+	GENERATED_BODY()
 
-	/** test objects, raw pointer can be used safely because it will be always referenced by EnvQueryManager */
-	TArray<UEnvQueryTest*> Tests;
+	/** generator object */
+	UPROPERTY(transient)
+	TObjectPtr<UEnvQueryGenerator> Generator;
+
+	/** test objects */
+	UPROPERTY(transient)
+	TArray<TObjectPtr<UEnvQueryTest>> Tests;
 
 	/** index of option in source asset */
 	int32 SourceOptionIndex;
 
 	/** type of generated items */
+	UPROPERTY(transient)
 	TSubclassOf<UEnvQueryItemType> ItemType;
 
 	/** if set, items will be shuffled after tests */
 	uint32 bHasNavLocations : 1;
 
-	inline uint32 GetAllocatedSize() const { return sizeof(*this) + Tests.GetAllocatedSize(); }
+	uint32 GetAllocatedSize() const
+	{
+		return sizeof(*this) + Tests.GetAllocatedSize();
+	}
 };
 
 #if NO_LOGGING
@@ -763,8 +776,23 @@ public:
 	static AIMODULE_API FText DescribeContext(TSubclassOf<UEnvQueryContext> ContextClass);
 };
 
+USTRUCT()
+struct FEnvQueryInstanceContextCacheItem
+{
+	GENERATED_BODY()
+
+	UPROPERTY(transient)
+	TObjectPtr<UClass> Class;
+
+	UPROPERTY(transient)
+	FEnvQueryContextData ContextData;
+};
+
+USTRUCT()
 struct FEnvQueryInstance : public FEnvQueryResult
 {
+	GENERATED_BODY()
+
 	typedef float FNamedParamValueType;
 
 	/** short name of query template - friendly name for debugging */
@@ -774,7 +802,8 @@ struct FEnvQueryInstance : public FEnvQueryResult
 	FName UniqueName;
 
 	/** world owning this query instance */
-	UWorld* World;
+	UPROPERTY(transient)
+	TObjectPtr<UWorld> World;
 
 	/** observer's delegate */
 	FQueryFinishedSignature FinishDelegate;
@@ -783,9 +812,11 @@ struct FEnvQueryInstance : public FEnvQueryResult
 	TMap<FName, FNamedParamValueType> NamedParams;
 
 	/** contexts in use */
-	TMap<UClass*, FEnvQueryContextData> ContextCache;
+	UPROPERTY(transient)
+	TArray<FEnvQueryInstanceContextCacheItem> ContextCache;
 
 	/** list of options */
+	UPROPERTY(transient)
 	TArray<FEnvQueryOptionInstance> Options;
 
 	/** currently processed test (-1 = generator) */
@@ -842,10 +873,12 @@ struct FEnvQueryInstance : public FEnvQueryResult
 	EEnvQueryRunMode::Type Mode;
 
 	/** item type's CDO for location tests */
-	UEnvQueryItemType_VectorBase* ItemTypeVectorCDO;
+	UPROPERTY(transient)
+	TObjectPtr<UEnvQueryItemType_VectorBase> ItemTypeVectorCDO;
 
 	/** item type's CDO for actor tests */
-	UEnvQueryItemType_ActorBase* ItemTypeActorCDO;
+	UPROPERTY(transient)
+	TObjectPtr<UEnvQueryItemType_ActorBase> ItemTypeActorCDO;
 
 	AIMODULE_API FEnvQueryInstance();
 	AIMODULE_API FEnvQueryInstance(const FEnvQueryInstance& Other);
diff --git a/Engine/Source/Runtime/AIModule/Private/EnvironmentQuery/EQSTestingPawn.cpp b/Engine/Source/Runtime/AIModule/Private/EnvironmentQuery/EQSTestingPawn.cpp
--- a/Engine/Source/Runtime/AIModule/Private/EnvironmentQuery/EQSTestingPawn.cpp
+++ b/Engine/Source/Runtime/AIModule/Private/EnvironmentQuery/EQSTestingPawn.cpp
@@ -221,6 +221,17 @@ void AEQSTestingPawn::PostLoad()
 	}
 }
 
+void AEQSTestingPawn::AddReferencedObjects(UObject* InThis, FReferenceCollector& Collector)
+{
+	const AEQSTestingPawn* This = CastChecked<AEQSTestingPawn>(InThis);
+	if (This->QueryInstance)
+	{
+		Collector.AddPropertyReferencesWithStructARO(FEnvQueryInstance::StaticStruct(), This->QueryInstance.Get(), This);
+	}
+
+	Super::AddReferencedObjects(InThis, Collector);
+}
+
 void AEQSTestingPawn::RunEQSQuery()
 {
 	if (HasAnyFlags(RF_ClassDefaultObject) == true)
diff --git a/Engine/Source/Runtime/AIModule/Private/EnvironmentQuery/EnvQueryInstance.cpp b/Engine/Source/Runtime/AIModule/Private/EnvironmentQuery/EnvQueryInstance.cpp
--- a/Engine/Source/Runtime/AIModule/Private/EnvironmentQuery/EnvQueryInstance.cpp
+++ b/Engine/Source/Runtime/AIModule/Private/EnvironmentQuery/EnvQueryInstance.cpp
@@ -109,15 +109,18 @@ bool FEnvQueryInstance::bDebuggingInfoEnabled = true;
 
 bool FEnvQueryInstance::PrepareContext(UClass* ContextClass, FEnvQueryContextData& ContextData)
 {
-	if (ContextClass == NULL)
+	if (ContextClass == nullptr)
 	{
 		return false;
 	}
 
 	if (ContextClass != UEnvQueryContext_Item::StaticClass())
 	{
-		FEnvQueryContextData* CachedData = ContextCache.Find(ContextClass);
-		if (CachedData == NULL)
+		const FEnvQueryInstanceContextCacheItem* CachedData = ContextCache.FindByPredicate([ContextClass](const FEnvQueryInstanceContextCacheItem& Item)
+			{
+				return Item.Class == ContextClass;
+			});
+		if (CachedData == nullptr)
 		{
 			UEnvQueryManager* QueryManager = UEnvQueryManager::GetCurrent(World);
 			const UEnvQueryContext* ContextOb = QueryManager->PrepareLocalContext(ContextClass);
@@ -131,13 +134,13 @@ bool FEnvQueryInstance::PrepareContext(UClass* ContextClass, FEnvQueryContextDat
 			DEC_MEMORY_STAT_BY(STAT_AI_EQS_InstanceMemory, AllocatedSize);
 #endif // STATS
 
-			ContextCache.Add(ContextClass, ContextData);
+			ContextCache.Add({ContextClass, ContextData});
 
 			INC_MEMORY_STAT_BY(STAT_AI_EQS_InstanceMemory, GetContextAllocatedSize());
 		}
 		else
 		{
-			ContextData = *CachedData;
+			ContextData = CachedData->ContextData;
 		}
 	}
 
@@ -953,9 +956,9 @@ uint32 FEnvQueryInstance::GetAllocatedSize() const
 uint32 FEnvQueryInstance::GetContextAllocatedSize() const
 {
 	SIZE_T MemSize = ContextCache.GetAllocatedSize();
-	for (TMap<UClass*, FEnvQueryContextData>::TConstIterator It(ContextCache); It; ++It)
+	for (const FEnvQueryInstanceContextCacheItem& CacheItem : ContextCache)
 	{
-		MemSize += It.Value().GetAllocatedSize();
+		MemSize += CacheItem.ContextData.GetAllocatedSize();
 	}
 
 	return IntCastChecked<uint32>(MemSize);
