[BehaviorTree] allowed types objects in BlackboardKeySelector are now created using RF_Public flag and in the transient package to fix reinstantiation issue in BP classes of BTNodes.

diff --git a/Engine/Source/Runtime/AIModule/Private/BehaviorTree/BehaviorTreeTypes.cpp b/Engine/Source/Runtime/AIModule/Private/BehaviorTree/BehaviorTreeTypes.cpp
--- a/Engine/Source/Runtime/AIModule/Private/BehaviorTree/BehaviorTreeTypes.cpp
+++ b/Engine/Source/Runtime/AIModule/Private/BehaviorTree/BehaviorTreeTypes.cpp
@@ -598,7 +598,7 @@ void FBlackboardKeySelector::InitSelection(const UBlackboardData& BlackboardAsse
 void FBlackboardKeySelector::AddObjectFilter(UObject* Owner, FName PropertyName, TSubclassOf<UObject> AllowedClass)
 {
 	const FName FilterName = MakeUniqueObjectName(Owner, UBlackboardKeyType_Object::StaticClass(), *FString::Printf(TEXT("%s_Object"), *PropertyName.ToString()));
-	UBlackboardKeyType_Object* FilterOb = NewObject<UBlackboardKeyType_Object>(Owner, FilterName);
+	UBlackboardKeyType_Object* FilterOb = NewObject<UBlackboardKeyType_Object>(GetTransientPackage(), FilterName, RF_Public);
 	FilterOb->BaseClass = AllowedClass;
 	AllowedTypes.Add(FilterOb);
 }
@@ -606,7 +606,7 @@ void FBlackboardKeySelector::AddObjectFilter(UObject* Owner, FName PropertyName,
 void FBlackboardKeySelector::AddClassFilter(UObject* Owner, FName PropertyName, TSubclassOf<UObject> AllowedClass)
 {
 	const FName FilterName = MakeUniqueObjectName(Owner, UBlackboardKeyType_Class::StaticClass(), *FString::Printf(TEXT("%s_Class"), *PropertyName.ToString()));
-	UBlackboardKeyType_Class* FilterOb = NewObject<UBlackboardKeyType_Class>(Owner, FilterName);
+	UBlackboardKeyType_Class* FilterOb = NewObject<UBlackboardKeyType_Class>(GetTransientPackage(), FilterName, RF_Public);
 	FilterOb->BaseClass = AllowedClass;
 	AllowedTypes.Add(FilterOb);
 }
@@ -614,7 +614,7 @@ void FBlackboardKeySelector::AddClassFilter(UObject* Owner, FName PropertyName,
 void FBlackboardKeySelector::AddEnumFilter(UObject* Owner, FName PropertyName, UEnum* AllowedEnum)
 {
 	const FName FilterName = MakeUniqueObjectName(Owner, UBlackboardKeyType_Enum::StaticClass(), *FString::Printf(TEXT("%s_Enum"), *PropertyName.ToString()));
-	UBlackboardKeyType_Enum* FilterOb = NewObject<UBlackboardKeyType_Enum>(Owner, FilterName);
+	UBlackboardKeyType_Enum* FilterOb = NewObject<UBlackboardKeyType_Enum>(GetTransientPackage(), FilterName, RF_Public);
 	FilterOb->EnumType = AllowedEnum;
 	AllowedTypes.Add(FilterOb);
 }
@@ -622,7 +622,7 @@ void FBlackboardKeySelector::AddEnumFilter(UObject* Owner, FName PropertyName, U
 void FBlackboardKeySelector::AddNativeEnumFilter(UObject* Owner, FName PropertyName, const FString& AllowedEnumName)
 {
 	const FName FilterName = MakeUniqueObjectName(Owner, UBlackboardKeyType_NativeEnum::StaticClass(), *FString::Printf(TEXT("%s_NativeEnum"), *PropertyName.ToString()));
-	UBlackboardKeyType_NativeEnum* FilterOb = NewObject<UBlackboardKeyType_NativeEnum>(Owner, FilterName);
+	UBlackboardKeyType_NativeEnum* FilterOb = NewObject<UBlackboardKeyType_NativeEnum>(GetTransientPackage(), FilterName, RF_Public);
 	FilterOb->EnumName = AllowedEnumName;
 	AllowedTypes.Add(FilterOb);
 }
@@ -630,49 +630,49 @@ void FBlackboardKeySelector::AddNativeEnumFilter(UObject* Owner, FName PropertyN
 void FBlackboardKeySelector::AddIntFilter(UObject* Owner, FName PropertyName)
 {
 	const FString FilterName = PropertyName.ToString() + TEXT("_Int");
-	AllowedTypes.Add(NewObject<UBlackboardKeyType_Int>(Owner, *FilterName));
+	AllowedTypes.Add(NewObject<UBlackboardKeyType_Int>(GetTransientPackage(), *FilterName, RF_Public));
 }
 
 void FBlackboardKeySelector::AddFloatFilter(UObject* Owner, FName PropertyName)
 {
 	const FString FilterName = PropertyName.ToString() + TEXT("_Float");
-	AllowedTypes.Add(NewObject<UBlackboardKeyType_Float>(Owner, *FilterName));
+	AllowedTypes.Add(NewObject<UBlackboardKeyType_Float>(GetTransientPackage(), *FilterName, RF_Public));
 }
 
 void FBlackboardKeySelector::AddBoolFilter(UObject* Owner, FName PropertyName)
 {
 	const FString FilterName = PropertyName.ToString() + TEXT("_Bool");
-	AllowedTypes.Add(NewObject<UBlackboardKeyType_Bool>(Owner, *FilterName));
+	AllowedTypes.Add(NewObject<UBlackboardKeyType_Bool>(GetTransientPackage(), *FilterName, RF_Public));
 }
 
 void FBlackboardKeySelector::AddVectorFilter(UObject* Owner, FName PropertyName)
 {
 	const FString FilterName = PropertyName.ToString() + TEXT("_Vector");
-	AllowedTypes.Add(NewObject<UBlackboardKeyType_Vector>(Owner, *FilterName));
+	AllowedTypes.Add(NewObject<UBlackboardKeyType_Vector>(GetTransientPackage(), *FilterName, RF_Public));
 }
 
 void FBlackboardKeySelector::AddRotatorFilter(UObject* Owner, FName PropertyName)
 {
 	const FString FilterName = PropertyName.ToString() + TEXT("_Rotator");
-	AllowedTypes.Add(NewObject<UBlackboardKeyType_Rotator>(Owner, *FilterName));
+	AllowedTypes.Add(NewObject<UBlackboardKeyType_Rotator>(GetTransientPackage(), *FilterName, RF_Public));
 }
 
 void FBlackboardKeySelector::AddStringFilter(UObject* Owner, FName PropertyName)
 {
 	const FString FilterName = PropertyName.ToString() + TEXT("_String");
-	AllowedTypes.Add(NewObject<UBlackboardKeyType_String>(Owner, *FilterName));
+	AllowedTypes.Add(NewObject<UBlackboardKeyType_String>(GetTransientPackage(), *FilterName, RF_Public));
 }
 
 void FBlackboardKeySelector::AddNameFilter(UObject* Owner, FName PropertyName)
 {
 	const FString FilterName = PropertyName.ToString() + TEXT("_Name");
-	AllowedTypes.Add(NewObject<UBlackboardKeyType_Name>(Owner, *FilterName));
+	AllowedTypes.Add(NewObject<UBlackboardKeyType_Name>(GetTransientPackage(), *FilterName, RF_Public));
 }
 
 void FBlackboardKeySelector::AddStructFilter(UObject* Owner, FName PropertyName, const UScriptStruct* AllowedStruct)
 {
 	const FString FilterName = PropertyName.ToString() + TEXT("_Struct_") + GetNameSafe(AllowedStruct);
-	UBlackboardKeyType_Struct* FilterOb = NewObject<UBlackboardKeyType_Struct>(Owner, *FilterName);
+	UBlackboardKeyType_Struct* FilterOb = NewObject<UBlackboardKeyType_Struct>(GetTransientPackage(), *FilterName, RF_Public);
 	FilterOb->DefaultValue.InitializeAs(AllowedStruct);
 	AllowedTypes.Add(FilterOb);
 }
