[BehaviorTree] allowed types objects in BlackboardKeySelector are now created using RF_Public flag and in the transient package to fix reinstantiation issue in BP classes of BTNodes.

#was-backout 45451796

diff --git a/Engine/Source/Runtime/AIModule/Private/BehaviorTree/BehaviorTreeTypes.cpp b/Engine/Source/Runtime/AIModule/Private/BehaviorTree/BehaviorTreeTypes.cpp
--- a/Engine/Source/Runtime/AIModule/Private/BehaviorTree/BehaviorTreeTypes.cpp
+++ b/Engine/Source/Runtime/AIModule/Private/BehaviorTree/BehaviorTreeTypes.cpp
@@ -595,84 +595,84 @@ void FBlackboardKeySelector::InitSelection(const UBlackboardData& BlackboardAsse
 	}
 }
 
-void FBlackboardKeySelector::AddObjectFilter(UObject* Owner, FName PropertyName, TSubclassOf<UObject> AllowedClass)
+namespace UE::BlackboardKeySelector
 {
-	const FName FilterName = MakeUniqueObjectName(Owner, UBlackboardKeyType_Object::StaticClass(), *FString::Printf(TEXT("%s_Object"), *PropertyName.ToString()));
-	UBlackboardKeyType_Object* FilterOb = NewObject<UBlackboardKeyType_Object>(Owner, FilterName);
+	template <typename T>
+	TNotNull<T*> CreateFilter(const FName PropertyName) requires TIsDerivedFrom<T, UBlackboardKeyType>::IsDerived
+	{
+		return static_cast<T*>(NewObject<T>(
+			GetTransientPackage()
+			, MakeUniqueObjectName(GetTransientPackage(), T::StaticClass(), *FString::Printf(TEXT("%s_%s"), *PropertyName.ToString(), *T::StaticClass()->GetName()))
+			, RF_Public));
+	}
+} // UE::BlackboardKeySelector
+
+void FBlackboardKeySelector::AddObjectFilter(UObject*, const FName PropertyName, TSubclassOf<UObject> AllowedClass)
+{
+	UBlackboardKeyType_Object* FilterOb = UE::BlackboardKeySelector::CreateFilter<UBlackboardKeyType_Object>(PropertyName);
 	FilterOb->BaseClass = AllowedClass;
 	AllowedTypes.Add(FilterOb);
 }
 
-void FBlackboardKeySelector::AddClassFilter(UObject* Owner, FName PropertyName, TSubclassOf<UObject> AllowedClass)
+void FBlackboardKeySelector::AddClassFilter(UObject*, const FName PropertyName, TSubclassOf<UObject> AllowedClass)
 {
-	const FName FilterName = MakeUniqueObjectName(Owner, UBlackboardKeyType_Class::StaticClass(), *FString::Printf(TEXT("%s_Class"), *PropertyName.ToString()));
-	UBlackboardKeyType_Class* FilterOb = NewObject<UBlackboardKeyType_Class>(Owner, FilterName);
+	UBlackboardKeyType_Class* FilterOb = UE::BlackboardKeySelector::CreateFilter<UBlackboardKeyType_Class>(PropertyName);
 	FilterOb->BaseClass = AllowedClass;
 	AllowedTypes.Add(FilterOb);
 }
 
-void FBlackboardKeySelector::AddEnumFilter(UObject* Owner, FName PropertyName, UEnum* AllowedEnum)
+void FBlackboardKeySelector::AddEnumFilter(UObject*, const FName PropertyName, UEnum* AllowedEnum)
 {
-	const FName FilterName = MakeUniqueObjectName(Owner, UBlackboardKeyType_Enum::StaticClass(), *FString::Printf(TEXT("%s_Enum"), *PropertyName.ToString()));
-	UBlackboardKeyType_Enum* FilterOb = NewObject<UBlackboardKeyType_Enum>(Owner, FilterName);
+	UBlackboardKeyType_Enum* FilterOb = UE::BlackboardKeySelector::CreateFilter<UBlackboardKeyType_Enum>(PropertyName);
 	FilterOb->EnumType = AllowedEnum;
 	AllowedTypes.Add(FilterOb);
 }
 
-void FBlackboardKeySelector::AddNativeEnumFilter(UObject* Owner, FName PropertyName, const FString& AllowedEnumName)
+void FBlackboardKeySelector::AddNativeEnumFilter(UObject*, const FName PropertyName, const FString& AllowedEnumName)
 {
-	const FName FilterName = MakeUniqueObjectName(Owner, UBlackboardKeyType_NativeEnum::StaticClass(), *FString::Printf(TEXT("%s_NativeEnum"), *PropertyName.ToString()));
-	UBlackboardKeyType_NativeEnum* FilterOb = NewObject<UBlackboardKeyType_NativeEnum>(Owner, FilterName);
+	UBlackboardKeyType_NativeEnum* FilterOb = UE::BlackboardKeySelector::CreateFilter<UBlackboardKeyType_NativeEnum>(PropertyName);
 	FilterOb->EnumName = AllowedEnumName;
 	AllowedTypes.Add(FilterOb);
 }
 
-void FBlackboardKeySelector::AddIntFilter(UObject* Owner, FName PropertyName)
+void FBlackboardKeySelector::AddIntFilter(UObject*, const FName PropertyName)
 {
-	const FString FilterName = PropertyName.ToString() + TEXT("_Int");
-	AllowedTypes.Add(NewObject<UBlackboardKeyType_Int>(Owner, *FilterName));
+	AllowedTypes.Add(UE::BlackboardKeySelector::CreateFilter<UBlackboardKeyType_Int>(PropertyName));
 }
 
-void FBlackboardKeySelector::AddFloatFilter(UObject* Owner, FName PropertyName)
+void FBlackboardKeySelector::AddFloatFilter(UObject*, const FName PropertyName)
 {
-	const FString FilterName = PropertyName.ToString() + TEXT("_Float");
-	AllowedTypes.Add(NewObject<UBlackboardKeyType_Float>(Owner, *FilterName));
+	AllowedTypes.Add(UE::BlackboardKeySelector::CreateFilter<UBlackboardKeyType_Float>(PropertyName));
 }
 
-void FBlackboardKeySelector::AddBoolFilter(UObject* Owner, FName PropertyName)
+void FBlackboardKeySelector::AddBoolFilter(UObject*, const FName PropertyName)
 {
-	const FString FilterName = PropertyName.ToString() + TEXT("_Bool");
-	AllowedTypes.Add(NewObject<UBlackboardKeyType_Bool>(Owner, *FilterName));
+	AllowedTypes.Add(UE::BlackboardKeySelector::CreateFilter<UBlackboardKeyType_Bool>(PropertyName));
 }
 
-void FBlackboardKeySelector::AddVectorFilter(UObject* Owner, FName PropertyName)
+void FBlackboardKeySelector::AddVectorFilter(UObject*, const FName PropertyName)
 {
-	const FString FilterName = PropertyName.ToString() + TEXT("_Vector");
-	AllowedTypes.Add(NewObject<UBlackboardKeyType_Vector>(Owner, *FilterName));
+	AllowedTypes.Add(UE::BlackboardKeySelector::CreateFilter<UBlackboardKeyType_Vector>(PropertyName));
 }
 
-void FBlackboardKeySelector::AddRotatorFilter(UObject* Owner, FName PropertyName)
+void FBlackboardKeySelector::AddRotatorFilter(UObject*, const FName PropertyName)
 {
-	const FString FilterName = PropertyName.ToString() + TEXT("_Rotator");
-	AllowedTypes.Add(NewObject<UBlackboardKeyType_Rotator>(Owner, *FilterName));
+	AllowedTypes.Add(UE::BlackboardKeySelector::CreateFilter<UBlackboardKeyType_Rotator>(PropertyName));
 }
 
-void FBlackboardKeySelector::AddStringFilter(UObject* Owner, FName PropertyName)
+void FBlackboardKeySelector::AddStringFilter(UObject*, const FName PropertyName)
 {
-	const FString FilterName = PropertyName.ToString() + TEXT("_String");
-	AllowedTypes.Add(NewObject<UBlackboardKeyType_String>(Owner, *FilterName));
+	AllowedTypes.Add(UE::BlackboardKeySelector::CreateFilter<UBlackboardKeyType_String>(PropertyName));
 }
 
-void FBlackboardKeySelector::AddNameFilter(UObject* Owner, FName PropertyName)
+void FBlackboardKeySelector::AddNameFilter(UObject*, const FName PropertyName)
 {
-	const FString FilterName = PropertyName.ToString() + TEXT("_Name");
-	AllowedTypes.Add(NewObject<UBlackboardKeyType_Name>(Owner, *FilterName));
+	AllowedTypes.Add(UE::BlackboardKeySelector::CreateFilter<UBlackboardKeyType_Name>(PropertyName));
 }
 
-void FBlackboardKeySelector::AddStructFilter(UObject* Owner, FName PropertyName, const UScriptStruct* AllowedStruct)
+void FBlackboardKeySelector::AddStructFilter(UObject*, const FName PropertyName, const UScriptStruct* AllowedStruct)
 {
-	const FString FilterName = PropertyName.ToString() + TEXT("_Struct_") + GetNameSafe(AllowedStruct);
-	UBlackboardKeyType_Struct* FilterOb = NewObject<UBlackboardKeyType_Struct>(Owner, *FilterName);
+	UBlackboardKeyType_Struct* FilterOb = UE::BlackboardKeySelector::CreateFilter<UBlackboardKeyType_Struct>(PropertyName);
 	FilterOb->DefaultValue.InitializeAs(AllowedStruct);
 	AllowedTypes.Add(FilterOb);
 }
