[AIBlueprint] tentative fix for CreateMoveToProxyObject processing an invalid UObject
#rnx

diff --git a/Engine/Source/Runtime/AIModule/Private/Blueprint/AIBlueprintHelperLibrary.cpp b/Engine/Source/Runtime/AIModule/Private/Blueprint/AIBlueprintHelperLibrary.cpp
--- a/Engine/Source/Runtime/AIModule/Private/Blueprint/AIBlueprintHelperLibrary.cpp
+++ b/Engine/Source/Runtime/AIModule/Private/Blueprint/AIBlueprintHelperLibrary.cpp
@@ -108,6 +108,16 @@ UAIBlueprintHelperLibrary::UAIBlueprintHelperLibrary(const FObjectInitializer& O
 
 UAIAsyncTaskBlueprintProxy* UAIBlueprintHelperLibrary::CreateMoveToProxyObject(UObject* WorldContextObject, APawn* Pawn, FVector Destination, AActor* TargetActor, float AcceptanceRadius, bool bStopOnOverlap)
 {
+	if (!ensureMsgf(WorldContextObject == nullptr || IsValid(WorldContextObject), TEXT("Unusable WorldContextObject pointer provided (i.e., garbage)")))
+	{
+		return nullptr;
+	}
+
+	if (!ensureMsgf(Pawn == nullptr || IsValid(Pawn), TEXT("Unusable Pawn pointer provided (i.e., garbage)")))
+	{
+		return nullptr;
+	}
+
 	if (WorldContextObject == nullptr)
 	{
 		if (Pawn != nullptr)
@@ -131,14 +141,16 @@ UAIAsyncTaskBlueprintProxy* UAIBlueprintHelperLibrary::CreateMoveToProxyObject(U
 		}
 	}
 
-	if (!Pawn)
+	if (Pawn == nullptr
+		|| !ensureMsgf(IsValid(Pawn), TEXT("Unusable Pawn pointer fetched from AIController (i.e., garbage)")))
 	{
-		return NULL;
+		return nullptr;
 	}
 
-	UAIAsyncTaskBlueprintProxy* MyObj = NULL;
+	UAIAsyncTaskBlueprintProxy* MyObj = nullptr;
 	AAIController* AIController = Cast<AAIController>(Pawn->GetController());
-	if (AIController)
+	if (AIController != nullptr
+		&& ensureMsgf(IsValid(AIController), TEXT("Unusable AIController pointer fetched from Pawn (i.e., garbage)")))
 	{
 		UWorld* World = GEngine->GetWorldFromContextObjectChecked(WorldContextObject);
 		MyObj = NewObject<UAIAsyncTaskBlueprintProxy>(World);
