[BehaviorTree] Change InstanceIdx to use a const int32 to prevent a need for CheckCast to prevent possible data truncation.

diff --git a/Engine/Source/Runtime/AIModule/Private/BehaviorTree/BTTaskNode.cpp b/Engine/Source/Runtime/AIModule/Private/BehaviorTree/BTTaskNode.cpp
--- a/Engine/Source/Runtime/AIModule/Private/BehaviorTree/BTTaskNode.cpp
+++ b/Engine/Source/Runtime/AIModule/Private/BehaviorTree/BTTaskNode.cpp
@@ -141,7 +141,7 @@ void UBTTaskNode::ReceivedMessage(UBrainComponent* BrainComp, const FAIMessage&
 	UBehaviorTreeComponent* OwnerComp = static_cast<UBehaviorTreeComponent*>(BrainComp);
 	check(OwnerComp);
 	
-	const uint16 InstanceIdx = IntCastChecked<uint16>(OwnerComp->FindInstanceContainingNode(this));
+	const int32 InstanceIdx = OwnerComp->FindInstanceContainingNode(this);
 	if (OwnerComp->InstanceStack.IsValidIndex(InstanceIdx))
 	{
 		uint8* NodeMemory = GetNodeMemory<uint8>(OwnerComp->InstanceStack[InstanceIdx]);
