EQS Manager: Fix crash when removing a test from the SortedTests invalidates MostExpensiveFilterIndex

diff --git a/Engine/Source/Runtime/AIModule/Private/EnvironmentQuery/EnvQueryManager.cpp b/Engine/Source/Runtime/AIModule/Private/EnvironmentQuery/EnvQueryManager.cpp
--- a/Engine/Source/Runtime/AIModule/Private/EnvironmentQuery/EnvQueryManager.cpp
+++ b/Engine/Source/Runtime/AIModule/Private/EnvironmentQuery/EnvQueryManager.cpp
@@ -899,6 +899,10 @@ TSharedPtr<FEnvQueryInstance> UEnvQueryManager::CreateQueryInstance(const UEnvQu
 						*GetNameSafe(LocalTemplate), *GetNameSafe(TestOb), OptionIndex, *MyOption->Generator->OptionName);
 
 					SortedTests.RemoveAt(TestIndex, EAllowShrinking::No);
+					if (MostExpensiveFilterIndex != INDEX_NONE)
+					{
+						--MostExpensiveFilterIndex;
+					}
 				}
 				else if (bOptionSingleResultSearch
 					&& TestOb->TestPurpose == EEnvTestPurpose::Filter
