Fixed StaticAnalysis error in EnvQueryGenerator_Donut (div 0).

#rnx

diff --git a/Engine/Source/Runtime/AIModule/Private/EnvironmentQuery/Generators/EnvQueryGenerator_Donut.cpp b/Engine/Source/Runtime/AIModule/Private/EnvironmentQuery/Generators/EnvQueryGenerator_Donut.cpp
--- a/Engine/Source/Runtime/AIModule/Private/EnvironmentQuery/Generators/EnvQueryGenerator_Donut.cpp
+++ b/Engine/Source/Runtime/AIModule/Private/EnvironmentQuery/Generators/EnvQueryGenerator_Donut.cpp
@@ -112,7 +112,7 @@ bool UEnvQueryGenerator_Donut::GenerateDonutNavPoints(const TArray<FVector>& Cen
 	const FVector::FReal ArcBisectDeg = GetArcBisectorAngle(QueryInstance);
 	const FVector::FReal ArcAngleDeg = FMath::Clamp(ArcAngleValue, 0., 360.);
 
-	const FVector::FReal RadiusDelta = (OuterRadiusValue - InnerRadiusValue) / (NumRings - 1);
+	const FVector::FReal RadiusDelta = NumRings > 1 ? (OuterRadiusValue - InnerRadiusValue) / (NumRings - 1) : 0.;
 	const FVector::FReal AngleDelta = 2. * UE_DOUBLE_PI / NumPoints;
 	FVector::FReal SectionAngle = FMath::DegreesToRadians(ArcBisectDeg);
 
