BehaviorTree: no FName definitions in header files

diff --git a/Engine/Source/Runtime/AIModule/Classes/BehaviorTree/Blackboard/BlackboardKey.h b/Engine/Source/Runtime/AIModule/Classes/BehaviorTree/Blackboard/BlackboardKey.h
--- a/Engine/Source/Runtime/AIModule/Classes/BehaviorTree/Blackboard/BlackboardKey.h
+++ b/Engine/Source/Runtime/AIModule/Classes/BehaviorTree/Blackboard/BlackboardKey.h
@@ -6,7 +6,7 @@
 
 namespace FBlackboard
 {
-	const FName KeySelf = TEXT("SelfActor");
+	extern AIMODULE_API const FName KeySelf;
 
 #ifdef AI_BLACKBOARD_KEY_SIZE_8
 	// this is the legacy BB key size. Add AI_BLACKBOARD_KEY_SIZE_8 to your *.target.cs to enable it.
@@ -34,4 +34,4 @@ namespace FBlackboard
 
 	inline uint32 GetTypeHash(const FKey& Key) { return ::GetTypeHash(Key.Key);}
 #endif
-}
\ No newline at end of file
+}
diff --git a/Engine/Source/Runtime/AIModule/Private/BehaviorTree/Blackboard/BlackboardKey.cpp b/Engine/Source/Runtime/AIModule/Private/BehaviorTree/Blackboard/BlackboardKey.cpp
new file mode 100644
--- /dev/null
+++ b/Engine/Source/Runtime/AIModule/Private/BehaviorTree/Blackboard/BlackboardKey.cpp
@@ -0,0 +1,5 @@
+// Copyright Epic Games, Inc. All Rights Reserved.
+
+#include "BehaviorTree/Blackboard/BlackboardKey.h"
+
+const FName FBlackboard::KeySelf = TEXT("SelfActor");
