[ChaosMover] NPC movement with ChaosMover change

diff --git a/Engine/Source/Runtime/AIModule/Classes/AIController.h b/Engine/Source/Runtime/AIModule/Classes/AIController.h
--- a/Engine/Source/Runtime/AIModule/Classes/AIController.h
+++ b/Engine/Source/Runtime/AIModule/Classes/AIController.h
@@ -126,6 +126,10 @@ public:
 	UPROPERTY(EditAnywhere, BlueprintReadWrite, Category = AI)
 	uint32 bSetControlRotationFromPawnOrientation:1;
 
+	/** If true, the RootComponent's transform will not be updated when SetControlRotation() is called. */
+	UPROPERTY(EditAnywhere, BlueprintReadWrite, Category = AI)
+	uint32 bSkipRootComponentUpdate:1;
+
 private:
 
 	/** Component used for moving along a path. */
@@ -260,6 +264,10 @@ public:
 	UFUNCTION(BlueprintCallable, Category = "AI")
 	AIMODULE_API virtual bool RunBehaviorTree(UBehaviorTree* BTAsset);
 
+	/** Set the control rotation. The RootComponent's rotation will also be updated to match
+	  * if bSkipRootComponentUpdate is false and RootComponent->bAbsoluteRotation is true. */
+	AIMODULE_API virtual void SetControlRotation(const FRotator& NewRotation) override;
+
 protected:
 	AIMODULE_API virtual void CleanupBrainComponent();
 
diff --git a/Engine/Source/Runtime/AIModule/Private/AIController.cpp b/Engine/Source/Runtime/AIModule/Private/AIController.cpp
--- a/Engine/Source/Runtime/AIModule/Private/AIController.cpp
+++ b/Engine/Source/Runtime/AIModule/Private/AIController.cpp
@@ -423,6 +423,25 @@ void AAIController::ActorsPerceptionUpdated(const TArray<AActor*>& UpdatedActors
 
 DEFINE_LOG_CATEGORY_STATIC(LogTestAI, All, All);
 
+void AAIController::SetControlRotation(const FRotator& NewRotation)
+{
+	if (!IsValidControlRotation(NewRotation))
+	{
+		logOrEnsureNanError(TEXT("AAIController::SetControlRotation attempted to apply NaN-containing or NaN-causing rotation! (%s)"), *NewRotation.ToString());
+		return;
+	}
+
+	if (!ControlRotation.Equals(NewRotation, 1e-3f))
+	{
+		ControlRotation = NewRotation;
+
+		if (!bSkipRootComponentUpdate && RootComponent && RootComponent->IsUsingAbsoluteRotation())
+		{
+			RootComponent->SetWorldRotation(GetControlRotation());
+		}
+	}
+}
+
 void AAIController::UpdateControlRotation(float DeltaTime, bool bUpdatePawn)
 {
 	APawn* const MyPawn = GetPawn();
