[AIController] we now end all tasks owned by the controller before calling OnUnPossess on the base class since those tasks might require access to the pawn's associated controller to exit properly.

diff --git a/Engine/Source/Runtime/AIModule/Private/AIController.cpp b/Engine/Source/Runtime/AIModule/Private/AIController.cpp
--- a/Engine/Source/Runtime/AIModule/Private/AIController.cpp
+++ b/Engine/Source/Runtime/AIModule/Private/AIController.cpp
@@ -540,22 +540,22 @@ void AAIController::OnPossess(APawn* InPawn)
 
 void AAIController::OnUnPossess()
 {
-	APawn* CurrentPawn = GetPawn();
+	const bool bGameplayTaskCleanupRequired = (CachedGameplayTasksComponent && (CachedGameplayTasksComponent->GetOwner() == GetPawn()));
 
-	Super::OnUnPossess();
-
-	if (PathFollowingComponent)
+	// End all tasks owned by the controller before calling OnUnPossess on the base class
+	// since those tasks might require access to the pawn's associated controller to exit properly.
+	// Also called before cleaning up the BrainComponent
+	// since it might need to update its resource lock (see OnGameplayTaskResourcesClaimed)
+	if (bGameplayTaskCleanupRequired)
 	{
-		PathFollowingComponent->Cleanup();
+		CachedGameplayTasksComponent->EndAllResourceConsumingTasksOwnedBy(*this);
 	}
 
-	const bool bGameplayTaskCleanupRequired = (CachedGameplayTasksComponent && (CachedGameplayTasksComponent->GetOwner() == CurrentPawn));
+	Super::OnUnPossess();
 
-	// End all tasks owned by the controller before cleaning up the BrainComponent
-	// since it might need to update its resource lock (see OnGameplayTaskResourcesClaimed)
-	if (bGameplayTaskCleanupRequired)
+	if (PathFollowingComponent)
 	{
-		CachedGameplayTasksComponent->EndAllResourceConsumingTasksOwnedBy(*this);
+		PathFollowingComponent->Cleanup();
 	}
 
 	if (bStopAILogicOnUnposses)
