Fix crash in ANavLinkProxy::GetComponentsBoundingBox

diff --git a/Engine/Source/Runtime/AIModule/Private/Navigation/NavLinkProxy.cpp b/Engine/Source/Runtime/AIModule/Private/Navigation/NavLinkProxy.cpp
--- a/Engine/Source/Runtime/AIModule/Private/Navigation/NavLinkProxy.cpp
+++ b/Engine/Source/Runtime/AIModule/Private/Navigation/NavLinkProxy.cpp
@@ -301,7 +301,10 @@ FBox ANavLinkProxy::GetComponentsBoundingBox(bool bNonColliding, bool bIncludeFr
 		LinksBB += SegmentLink.RightEnd;
 	}
 
-	LinksBB = LinksBB.TransformBy(RootComponent->GetComponentTransform());
+	if (RootComponent)
+	{
+		LinksBB = LinksBB.TransformBy(RootComponent->GetComponentTransform());
+	}
 
 	if (SmartLinkComp && SmartLinkComp->IsNavigationRelevant())
 	{
