Added a override for IsFollowingNavLink in Crowd following component. Also caching bTraversingLink, so that IsFollowingNavLink will use that.
Fixes issue with CrowdFollowing AI that have bSlideAlongNavMesh enabled getting stuck on ledges & not following their nav links.

#tests tested in test map and also tried running through a couple zones. Husks were dropping from cliffs and ledges that had nav links.

diff --git a/Engine/Source/Runtime/AIModule/Classes/Navigation/CrowdFollowingComponent.h b/Engine/Source/Runtime/AIModule/Classes/Navigation/CrowdFollowingComponent.h
--- a/Engine/Source/Runtime/AIModule/Classes/Navigation/CrowdFollowingComponent.h
+++ b/Engine/Source/Runtime/AIModule/Classes/Navigation/CrowdFollowingComponent.h
@@ -206,6 +206,8 @@ protected:
 
 	uint8 bEnableSimulationReplanOnResume : 1;
 
+	uint8 bIsTraversingLink: 1;
+
 	TEnumAsByte<ECrowdAvoidanceQuality::Type> AvoidanceQuality;
 	ECrowdSimulationState SimulationState;
 
@@ -231,6 +233,7 @@ protected:
 	AIMODULE_API virtual bool IsOnPath() const override;
 	AIMODULE_API virtual bool UpdateMovementComponent(bool bForce) override;
 	AIMODULE_API virtual void Reset() override;
+	AIMODULE_API virtual bool IsFollowingNavLink() const override;
 	// PathFollowingComponent END
 
 	AIMODULE_API void SwitchToNextPathPart();
diff --git a/Engine/Source/Runtime/AIModule/Private/Navigation/CrowdFollowingComponent.cpp b/Engine/Source/Runtime/AIModule/Private/Navigation/CrowdFollowingComponent.cpp
--- a/Engine/Source/Runtime/AIModule/Private/Navigation/CrowdFollowingComponent.cpp
+++ b/Engine/Source/Runtime/AIModule/Private/Navigation/CrowdFollowingComponent.cpp
@@ -27,6 +27,7 @@ UCrowdFollowingComponent::UCrowdFollowingComponent(const FObjectInitializer& Obj
 	bRotateToVelocity = true;
 	bSuspendCrowdSimulation = false;
 	bEnableSimulationReplanOnResume = true;
+	bIsTraversingLink = false;
 	bRegisteredWithCrowdSimulation = false;
 	bCanCheckMovingTooFar = true;
 	bCanUpdatePathPartInTick = true;
@@ -338,6 +339,9 @@ void UCrowdFollowingComponent::UpdateCachedDirections(const FVector& NewVelocity
 		MoveSegmentDirection = ToCorner.GetSafeNormal();
 	}
 
+	// Cache bTraversingLink
+	bIsTraversingLink = bTraversingLink;
+
 	// CrowdAgentMoveDirection either direction on path or aligned with current velocity
 	const bool bIsNotFalling = (NavMovementInterface == nullptr || !NavMovementInterface->IsFalling());
 	if (bIsNotFalling)
@@ -569,6 +573,19 @@ void UCrowdFollowingComponent::Reset()
 	bCanCheckMovingTooFar = true;
 	bCheckMovementAngle = false;
 	bUpdateDirectMoveVelocity = false;
+	bIsTraversingLink = false;
+}
+
+bool UCrowdFollowingComponent::IsFollowingNavLink() const
+{
+	if (IsCrowdSimulationEnabled())
+	{
+		return bIsTraversingLink;
+	}
+	else
+	{
+		return Super::IsFollowingNavLink();
+	}
 }
 
 bool UCrowdFollowingComponent::UpdateMovementComponent(bool bForce)
