GameplayDebugger BehaviorTree : Add paging (up/down) to be able to display all blackboard keys on the screen.
#tests PIE

#Synced-CL 48871016

diff --git a/Engine/Source/Runtime/AIModule/AIModule.Build.cs b/Engine/Source/Runtime/AIModule/AIModule.Build.cs
--- a/Engine/Source/Runtime/AIModule/AIModule.Build.cs
+++ b/Engine/Source/Runtime/AIModule/AIModule.Build.cs
@@ -22,7 +22,8 @@ namespace UnrealBuildTool.Rules
             PrivateDependencyModuleNames.AddRange(
                 new string[] {
                     "AutoRTFM",
-                    "RHI",
+					"InputCore",
+					"RHI",
                     "RenderCore",
                 }
                 );
diff --git a/Engine/Source/Runtime/AIModule/Private/GameplayDebugger/GameplayDebuggerCategory_BehaviorTree.cpp b/Engine/Source/Runtime/AIModule/Private/GameplayDebugger/GameplayDebuggerCategory_BehaviorTree.cpp
--- a/Engine/Source/Runtime/AIModule/Private/GameplayDebugger/GameplayDebuggerCategory_BehaviorTree.cpp
+++ b/Engine/Source/Runtime/AIModule/Private/GameplayDebugger/GameplayDebuggerCategory_BehaviorTree.cpp
@@ -8,10 +8,18 @@
 #include "BrainComponent.h"
 #include "AIController.h"
 #include "BehaviorTree/BlackboardComponent.h"
+#include "GameplayDebuggerTypes.h"
+#include "InputCoreTypes.h"
 
 FGameplayDebuggerCategory_BehaviorTree::FGameplayDebuggerCategory_BehaviorTree()
 {
 	SetDataPackReplication<FRepData>(&DataPack);
+
+	const FGameplayDebuggerInputHandlerConfig ScrollUpKeyConfig(TEXT("ScrollUp"), EKeys::Subtract.GetFName(), FGameplayDebuggerInputModifier::Shift);
+	BindKeyPress(ScrollUpKeyConfig, this, &FGameplayDebuggerCategory_BehaviorTree::ScrollUp);
+
+	const FGameplayDebuggerInputHandlerConfig ScrollDownKeyConfig(TEXT("ScrollDown"), EKeys::Add.GetFName(), FGameplayDebuggerInputModifier::Shift);
+	BindKeyPress(ScrollDownKeyConfig, this, &FGameplayDebuggerCategory_BehaviorTree::ScrollDown);
 }
 
 TSharedRef<FGameplayDebuggerCategory> FGameplayDebuggerCategory_BehaviorTree::MakeInstance()
@@ -48,6 +56,9 @@ void FGameplayDebuggerCategory_BehaviorTree::DrawData(APlayerController* OwnerPC
 {
 	if (!DataPack.CompName.IsEmpty())
 	{
+		// Display options
+		CanvasContext.Printf(TEXT("Display: Use {yellow}[Shift+]{white} and {yellow}[Shift-]{white} to scroll the blackboard.\n"));
+
 		CanvasContext.Printf(TEXT("Brain Component: {yellow}%s"), *DataPack.CompName);
 		CanvasContext.Print(DataPack.TreeDesc);
 
@@ -59,7 +70,7 @@ void FGameplayDebuggerCategory_BehaviorTree::DrawData(APlayerController* OwnerPC
 		CanvasContext.DefaultX = CanvasContext.CursorX = 600.0f;
 		CanvasContext.CursorY = CanvasContext.DefaultY;
 
-		for (int32 Idx = 0; Idx < BlackboardLines.Num(); Idx++)
+		for (int32 Idx = DisplayStartingIndex; Idx < BlackboardLines.Num(); Idx++)
 		{
 			int32 SeparatorIndex = INDEX_NONE;
 			BlackboardLines[Idx].FindChar(TEXT(':'), SeparatorIndex);
@@ -80,4 +91,20 @@ void FGameplayDebuggerCategory_BehaviorTree::DrawData(APlayerController* OwnerPC
 	}
 }
 
+void FGameplayDebuggerCategory_BehaviorTree::ScrollUp()
+{
+	DisplayStartingIndex = FMath::Max(DisplayStartingIndex - PagingOffset, 0);
+}
+
+void FGameplayDebuggerCategory_BehaviorTree::ScrollDown()
+{
+	if (!DataPack.CompName.IsEmpty())
+	{
+		TArray<FString> BlackboardLines;
+		DataPack.BlackboardDesc.ParseIntoArrayLines(BlackboardLines, /* Empty strings are not added to the array */ true);
+
+		DisplayStartingIndex = FMath::Min(DisplayStartingIndex + PagingOffset, BlackboardLines.Num());
+	}
+}
+
 #endif // WITH_GAMEPLAY_DEBUGGER_MENU
diff --git a/Engine/Source/Runtime/AIModule/Public/GameplayDebugger/GameplayDebuggerCategory_BehaviorTree.h b/Engine/Source/Runtime/AIModule/Public/GameplayDebugger/GameplayDebuggerCategory_BehaviorTree.h
--- a/Engine/Source/Runtime/AIModule/Public/GameplayDebugger/GameplayDebuggerCategory_BehaviorTree.h
+++ b/Engine/Source/Runtime/AIModule/Public/GameplayDebugger/GameplayDebuggerCategory_BehaviorTree.h
@@ -30,6 +30,20 @@ protected:
 		void Serialize(FArchive& Ar);
 	};
 	FRepData DataPack;
+
+private:
+
+	/** Bound on Shift- decrease the blackboard display starting index */
+	void ScrollUp();
+
+	/** Bound on Shift+ increase the blackboard display starting index */
+	void ScrollDown();
+
+	/** Current blackboard display starting index */
+	int32 DisplayStartingIndex = 0;
+
+	/** Offset used to move up or down the paging */
+	static constexpr int32 PagingOffset = 5;
 };
 
 #endif // WITH_GAMEPLAY_DEBUGGER_MENU
\ No newline at end of file
