Add templated getters for value or bbkeys
#Synced-CL 48913390

diff --git a/Engine/Source/Runtime/AIModule/Classes/BehaviorTree/ValueOrBBKey.h b/Engine/Source/Runtime/AIModule/Classes/BehaviorTree/ValueOrBBKey.h
--- a/Engine/Source/Runtime/AIModule/Classes/BehaviorTree/ValueOrBBKey.h
+++ b/Engine/Source/Runtime/AIModule/Classes/BehaviorTree/ValueOrBBKey.h
@@ -166,6 +166,24 @@ struct FValueOrBBKey_Class : public FValueOrBlackboardKeyBase
 		return TSubclassOf<T>(GetValue(BehaviorComp));
 	}
 
+	template <typename T>
+	TSubclassOf<T> GetValue(const UBehaviorTreeComponent* BehaviorComp) const
+	{
+		return TSubclassOf<T>(GetValue(BehaviorComp));
+	}
+
+	template <typename T>
+	TSubclassOf<T> GetValue(const UBlackboardComponent& BlackboardComp) const
+	{
+		return TSubclassOf<T>(GetValue(BlackboardComp));
+	}
+
+	template <typename T>
+	TSubclassOf<T> GetValue(const UBlackboardComponent* BlackboardComp) const
+	{
+		return TSubclassOf<T>(GetValue(BlackboardComp));
+	}
+
 	AIMODULE_API UClass* GetValue(const UBehaviorTreeComponent& BehaviorComp) const;
 	AIMODULE_API UClass* GetValue(const UBehaviorTreeComponent* BehaviorComp) const;
 	AIMODULE_API UClass* GetValue(const UBlackboardComponent& Blackboard) const;
@@ -211,6 +229,24 @@ struct FValueOrBBKey_Enum : public FValueOrBlackboardKeyBase
 		return static_cast<T>(GetValue(BehaviorComp));
 	}
 
+	template <typename T>
+	T GetValue(const UBehaviorTreeComponent* BehaviorComp) const
+	{
+		return static_cast<T>(GetValue(BehaviorComp));
+	}
+
+	template <typename T>
+	T GetValue(const UBlackboardComponent& BlackboardComp) const
+	{
+		return static_cast<T>(GetValue(BlackboardComp));
+	}
+
+	template <typename T>
+	T GetValue(const UBlackboardComponent* BlackboardComp) const
+	{
+		return static_cast<T>(GetValue(BlackboardComp));
+	}
+
 	AIMODULE_API uint8 GetValue(const UBehaviorTreeComponent& BehaviorComp) const;
 	AIMODULE_API uint8 GetValue(const UBehaviorTreeComponent* BehaviorComp) const;
 	AIMODULE_API uint8 GetValue(const UBlackboardComponent& Blackboard) const;
@@ -371,6 +407,24 @@ struct FValueOrBBKey_Object : public FValueOrBlackboardKeyBase
 		return Cast<T>(GetValue(BehaviorComp));
 	}
 
+	template <typename T>
+	T* GetValue(const UBehaviorTreeComponent* BehaviorComp) const
+	{
+		return Cast<T>(GetValue(BehaviorComp));
+	}
+
+	template <typename T>
+	T* GetValue(const UBlackboardComponent& BlackboardComp) const
+	{
+		return Cast<T>(GetValue(BlackboardComp));
+	}
+
+	template <typename T>
+	T* GetValue(const UBlackboardComponent* BlackboardComp) const
+	{
+		return Cast<T>(GetValue(BlackboardComp));
+	}
+
 	AIMODULE_API UObject* GetValue(const UBehaviorTreeComponent& BehaviorComp) const;
 	AIMODULE_API UObject* GetValue(const UBehaviorTreeComponent* BehaviorComp) const;
 	AIMODULE_API UObject* GetValue(const UBlackboardComponent& Blackboard) const;
@@ -479,6 +533,27 @@ public:
 		return GetValue(BehaviorComp).Get<const T>();
 	}
 
+	template <typename T>
+	const T& GetValue(const UBehaviorTreeComponent* BehaviorComp) const
+	{
+		check(DefaultValue.GetScriptStruct() == TBaseStructure<T>::Get());
+		return GetValue(BehaviorComp).Get<const T>();
+	}
+
+	template <typename T>
+	const T& GetValue(const UBlackboardComponent& BlackboardComp) const
+	{
+		check(DefaultValue.GetScriptStruct() == TBaseStructure<T>::Get());
+		return GetValue(BlackboardComp).Get<const T>();
+	}
+
+	template <typename T>
+	const T& GetValue(const UBlackboardComponent* BlackboardComp) const
+	{
+		check(DefaultValue.GetScriptStruct() == TBaseStructure<T>::Get());
+		return GetValue(BlackboardComp).Get<const T>();
+	}
+
 	bool SerializeFromMismatchedTag(const FPropertyTag& Tag, FStructuredArchive::FSlot Slot);
 
 	AIMODULE_API FConstStructView GetValue(const UBehaviorTreeComponent& BehaviorComp) const;
