[EQS] Sanitized generation parameters in UEnvQueryGenerator_SimpleGrid::GenerateItems and added warnings if the values are crazy-large.

#rnx

#Synced-CL 49801936

diff --git a/Engine/Source/Runtime/AIModule/Private/EnvironmentQuery/Generators/EnvQueryGenerator_SimpleGrid.cpp b/Engine/Source/Runtime/AIModule/Private/EnvironmentQuery/Generators/EnvQueryGenerator_SimpleGrid.cpp
--- a/Engine/Source/Runtime/AIModule/Private/EnvironmentQuery/Generators/EnvQueryGenerator_SimpleGrid.cpp
+++ b/Engine/Source/Runtime/AIModule/Private/EnvironmentQuery/Generators/EnvQueryGenerator_SimpleGrid.cpp
@@ -2,6 +2,7 @@
 
 #include "EnvironmentQuery/Generators/EnvQueryGenerator_SimpleGrid.h"
 #include "EnvironmentQuery/Contexts/EnvQueryContext_Querier.h"
+#include "VisualLogger/VisualLogger.h"
 
 #include UE_INLINE_GENERATED_CPP_BY_NAME(EnvQueryGenerator_SimpleGrid)
 
@@ -20,32 +21,58 @@ void UEnvQueryGenerator_SimpleGrid::GenerateItems(FEnvQueryInstance& QueryInstan
 	GridSize.BindData(BindOwner, QueryInstance.QueryID);
 	SpaceBetween.BindData(BindOwner, QueryInstance.QueryID);
 
-	float RadiusValue = GridSize.GetValue();
-	float DensityValue = SpaceBetween.GetValue();
+	const float DensityValue = SpaceBetween.GetValue();
+	const float RadiusValue = GridSize.GetValue();
 
-	const int32 ItemCount = FPlatformMath::TruncToInt((RadiusValue * 2.0f / DensityValue) + 1);
-	const int32 ItemCountHalf = ItemCount / 2;
+	if (DensityValue > 0.f && RadiusValue > 0.f)
+	{
+		const int32 ItemCount = FPlatformMath::TruncToInt((RadiusValue * 2.0f / DensityValue) + 1);
 
-	TArray<FVector> ContextLocations;
-	QueryInstance.PrepareContext(GenerateAround, ContextLocations);
+		TArray<FVector> ContextLocations;
+		QueryInstance.PrepareContext(GenerateAround, ContextLocations);
 
-	TArray<FNavLocation> GridPoints;
-	GridPoints.Reserve(ItemCount * ItemCount * ContextLocations.Num());
+		const uint64 TotalItemsCount = static_cast<uint64>(ItemCount) * static_cast<uint64>(ItemCount) * static_cast<uint64>(ContextLocations.Num());
 
-	for (int32 ContextIndex = 0; ContextIndex < ContextLocations.Num(); ContextIndex++)
-	{
-		for (int32 IndexX = 0; IndexX < ItemCount; ++IndexX)
+		if (TotalItemsCount <= static_cast<uint64>(TNumericLimits<int32>::Max()))
 		{
-			for (int32 IndexY = 0; IndexY < ItemCount; ++IndexY)
+			const int32 ItemCountHalf = ItemCount / 2;
+
+			TArray<FNavLocation> GridPoints;
+			GridPoints.Reserve(static_cast<int32>(TotalItemsCount));
+
+			for (FVector ContextLocation : ContextLocations)
 			{
-				const FNavLocation TestPoint = FNavLocation(ContextLocations[ContextIndex] - FVector(DensityValue * (IndexX - ItemCountHalf), DensityValue * (IndexY - ItemCountHalf), 0));
-				GridPoints.Add(TestPoint);
+				for (int32 IndexX = 0; IndexX < ItemCount; ++IndexX)
+				{
+					for (int32 IndexY = 0; IndexY < ItemCount; ++IndexY)
+					{
+						const FVector ItemRelativeLocation(DensityValue * (IndexX - ItemCountHalf), DensityValue * (IndexY - ItemCountHalf), 0);
+						FNavLocation TestPoint = FNavLocation(ContextLocation - ItemRelativeLocation);
+						GridPoints.Add(MoveTemp(TestPoint));
+					}
+				}
 			}
+
+			ProjectAndFilterNavPoints(GridPoints, QueryInstance);
+			StoreNavPoints(GridPoints, QueryInstance);
+		}
+		else
+		{
+			UE_VLOG_UELOG(QueryInstance.Owner.Get(), LogEQS, Error
+				, TEXT("%hs, query %s: attempting to generate %llu items. That's way more than expected or supported. Please review the query's config.")
+				, __FUNCTION__, *QueryInstance.QueryName, TotalItemsCount);
 		}
 	}
-
-	ProjectAndFilterNavPoints(GridPoints, QueryInstance);
-	StoreNavPoints(GridPoints, QueryInstance);
+	else
+	{
+		const UObject* Querier = QueryInstance.Owner.Get();
+		UE_CVLOG(DensityValue <= 0.f, Querier, LogEQS, Error
+			, TEXT("%hs, query %s: Zero or negative `SpaceBetween` values are not supported")
+			, __FUNCTION__, *QueryInstance.QueryName);
+		UE_CVLOG(RadiusValue <= 0.f, Querier, LogEQS, Error
+			, TEXT("%hs, query %s: Zero or negative `GridSize` values are not supported")
+			, __FUNCTION__, *QueryInstance.QueryName);
+	}
 }
 
 FText UEnvQueryGenerator_SimpleGrid::GetDescriptionTitle() const
