[AI][Blackboard] Fix crash in DestroyValues when blackboard failed to initialize

When a blackboard asset fails validation (e.g., duplicate keys in parent chain), InitializeBlackboard sets BlackboardAsset but skips memory allocation. On cleanup, DestroyValues iterates all keys and calls WrappedFree on a nullptr, crashing.

#Synced-CL 49927832

#Synced-CL 49927832

diff --git a/Engine/Source/Runtime/AIModule/Private/BehaviorTree/BlackboardComponent.cpp b/Engine/Source/Runtime/AIModule/Private/BehaviorTree/BlackboardComponent.cpp
--- a/Engine/Source/Runtime/AIModule/Private/BehaviorTree/BlackboardComponent.cpp
+++ b/Engine/Source/Runtime/AIModule/Private/BehaviorTree/BlackboardComponent.cpp
@@ -215,7 +215,16 @@ void UBlackboardComponent::DestroyValues()
 				check(FBlackboard::FKey(OffsetKey) != FBlackboard::InvalidKey);
 
 				uint8* KeyMemory = GetKeyRawData(FBlackboard::FKey(OffsetKey));
-				KeyType->WrappedFree(*this, KeyMemory);
+				if (KeyMemory)
+				{
+					KeyType->WrappedFree(*this, KeyMemory);
+				}
+				else
+				{
+					UE_LOG(LogBlackboard, Warning, 
+						   TEXT("DestroyValues skipping key '%s' - no memory allocated. Blackboard asset (%s) may have failed initialization."), 
+						   *It->Keys[KeyIndex].EntryName.ToString(), *GetNameSafe(BlackboardAsset));
+				}
 			}
 		}
 	}
