When invalid blackboard data is used when initializing a BlackboardComponent, it effectively silently fails until the BlackboardComponent gets destroyed - at which point it would crash by dereferencing nullptr (plus the size of the key).

This change makes it so that a failure is far more apparent - and it prevents a crash when PIE ends.

#okforgith

#Synced-CL 50066431

diff --git a/Engine/Source/Runtime/AIModule/Classes/BehaviorTree/BlackboardComponent.h b/Engine/Source/Runtime/AIModule/Classes/BehaviorTree/BlackboardComponent.h
--- a/Engine/Source/Runtime/AIModule/Classes/BehaviorTree/BlackboardComponent.h
+++ b/Engine/Source/Runtime/AIModule/Classes/BehaviorTree/BlackboardComponent.h
@@ -250,6 +250,9 @@ protected:
 	/** instanced keys with custom data allocations */
 	UPROPERTY(transient)
 	TArray<TObjectPtr<UBlackboardKeyType>> KeyInstances;
+	
+	/** tracks whether memory used to store values in the blackboard has been fully initialized */
+	bool bBlackboardMemoryInitialized = false;
 
 protected:
 	struct FOnBlackboardChangeNotificationInfo
diff --git a/Engine/Source/Runtime/AIModule/Private/BehaviorTree/BlackboardComponent.cpp b/Engine/Source/Runtime/AIModule/Private/BehaviorTree/BlackboardComponent.cpp
--- a/Engine/Source/Runtime/AIModule/Private/BehaviorTree/BlackboardComponent.cpp
+++ b/Engine/Source/Runtime/AIModule/Private/BehaviorTree/BlackboardComponent.cpp
@@ -111,7 +111,7 @@ void UBlackboardComponent::InitializeParentChain(UBlackboardData* NewAsset)
 bool UBlackboardComponent::InitializeBlackboard(UBlackboardData& NewAsset)
 {
 	// if we re-initialize with the same asset then there's no point
-	// in reseting, since we'd lose all the accumulated knowledge
+	// in resetting, since we'd lose all the accumulated knowledge
 	if (&NewAsset == BlackboardAsset)
 	{
 		return true;
@@ -132,10 +132,9 @@ bool UBlackboardComponent::InitializeBlackboard(UBlackboardData& NewAsset)
 	BlackboardAsset = &NewAsset;
 	ValueMemory.Reset();
 	ValueOffsets.Reset();
+	bBlackboardMemoryInitialized = false;
 	bSynchronizedKeyPopulated = false;
 
-	bool bSuccess = true;
-	
 	if (BlackboardAsset->IsValid())
 	{
 		InitializeParentChain(BlackboardAsset);
@@ -192,26 +191,32 @@ bool UBlackboardComponent::InitializeBlackboard(UBlackboardData& NewAsset)
 		{
 			PopulateSynchronizedKeys();
 		}
+		
+		bBlackboardMemoryInitialized = true;
 	}
 	else
 	{
-		bSuccess = false;
 		UE_LOG(LogBlackboard, Error, TEXT("Blackboard asset (%s) has errors and can't be used!"), *GetNameSafe(BlackboardAsset));
 	}
 
-	return bSuccess;
+	return bBlackboardMemoryInitialized;
 }
 
 void UBlackboardComponent::DestroyValues()
 {
+	if (!bBlackboardMemoryInitialized || 
+		!ensureMsgf(ValueMemory.Num() > 0 && ValueOffsets.Num() > 0, TEXT("The Blackboard is initialized but no memory was allocated to store values for Blackboard asset: %s"), *GetNameSafe(BlackboardAsset)))
+	{
+		return;
+	}
+	
 	for (UBlackboardData* It = BlackboardAsset; It; It = It->Parent)
 	{
 		for (int32 KeyIndex = 0; KeyIndex < It->Keys.Num(); KeyIndex++)
 		{
-			UBlackboardKeyType* KeyType = It->Keys[KeyIndex].KeyType;
-			if (KeyType)
+			if (UBlackboardKeyType* KeyType = It->Keys[KeyIndex].KeyType)
 			{
-				const int32 OffsetKey = KeyIndex + (int32)It->GetFirstKeyID();
+				const int32 OffsetKey = KeyIndex + static_cast<int32>(It->GetFirstKeyID());
 				check(FBlackboard::FKey(OffsetKey) != FBlackboard::InvalidKey);
 
 				uint8* KeyMemory = GetKeyRawData(FBlackboard::FKey(OffsetKey));
@@ -231,6 +236,8 @@ void UBlackboardComponent::DestroyValues()
 
 	ValueOffsets.Reset();
 	ValueMemory.Reset();
+
+	bBlackboardMemoryInitialized = false;
 }
 
 void UBlackboardComponent::PopulateSynchronizedKeys()
diff --git a/Engine/Source/Runtime/AIModule/Private/BehaviorTree/BlackboardData.cpp b/Engine/Source/Runtime/AIModule/Private/BehaviorTree/BlackboardData.cpp
--- a/Engine/Source/Runtime/AIModule/Private/BehaviorTree/BlackboardData.cpp
+++ b/Engine/Source/Runtime/AIModule/Private/BehaviorTree/BlackboardData.cpp
@@ -126,7 +126,7 @@ bool UBlackboardData::IsValid() const
 			const FBlackboard::FKey KeyID = Parent->InternalGetKeyID(Keys[KeyIndex].EntryName, CheckParentKeys);
 			if (KeyID != FBlackboard::InvalidKey)
 			{
-				UE_LOG(LogBehaviorTree, Warning, TEXT("Blackboard asset (%s) has duplicated key (%s) in parent chain!"),
+				UE_LOG(LogBehaviorTree, Error, TEXT("Blackboard asset (%s) has duplicated key (%s) in parent chain!"),
 					*GetName(), *Keys[KeyIndex].EntryName.ToString());
 
 				return false;
