[NavLinkProxy/ALT]
- Register to delegates in postload during loading to avoid races and being called too early.

#Synced-CL 50160399

diff --git a/Engine/Source/Runtime/AIModule/Private/Navigation/NavLinkProxy.cpp b/Engine/Source/Runtime/AIModule/Private/Navigation/NavLinkProxy.cpp
--- a/Engine/Source/Runtime/AIModule/Private/Navigation/NavLinkProxy.cpp
+++ b/Engine/Source/Runtime/AIModule/Private/Navigation/NavLinkProxy.cpp
@@ -114,7 +114,8 @@ void ANavLinkProxy::PostInitProperties()
 	Super::PostInitProperties();
 
 #if WITH_EDITOR
-	if (GIsEditor && !HasAnyFlags(RF_ClassDefaultObject))
+	// During loading, those will be registered to in postload instead for thread-safety
+	if (GIsEditor && !HasAnyFlags(RF_ClassDefaultObject) && !IsInAsyncLoadingThread())
 	{
 		OnNavAreaRegisteredDelegateHandle = UNavigationSystemBase::OnNavAreaRegisteredDelegate().AddUObject(this, &ANavLinkProxy::OnNavAreaRegistered);
 		OnNavAreaUnregisteredDelegateHandle = UNavigationSystemBase::OnNavAreaUnregisteredDelegate().AddUObject(this, &ANavLinkProxy::OnNavAreaUnregistered);
@@ -226,6 +227,14 @@ void ANavLinkProxy::PostLoad()
 {
 	Super::PostLoad();
 
+#if WITH_EDITOR
+	if (GIsEditor && !HasAnyFlags(RF_ClassDefaultObject))
+	{
+		OnNavAreaRegisteredDelegateHandle = UNavigationSystemBase::OnNavAreaRegisteredDelegate().AddUObject(this, &ANavLinkProxy::OnNavAreaRegistered);
+		OnNavAreaUnregisteredDelegateHandle = UNavigationSystemBase::OnNavAreaUnregisteredDelegate().AddUObject(this, &ANavLinkProxy::OnNavAreaUnregistered);
+	}
+#endif // WITH_EDITOR
+
 	for (FNavigationLink& Link : PointLinks)
 	{
 		Link.InitializeAreaClass();
diff --git a/Engine/Source/Runtime/NavigationSystem/Private/NavModifierVolume.cpp b/Engine/Source/Runtime/NavigationSystem/Private/NavModifierVolume.cpp
--- a/Engine/Source/Runtime/NavigationSystem/Private/NavModifierVolume.cpp
+++ b/Engine/Source/Runtime/NavigationSystem/Private/NavModifierVolume.cpp
@@ -49,6 +49,19 @@ void ANavModifierVolume::PostInitProperties()
 {
 	Super::PostInitProperties();
 
+#if WITH_EDITOR
+	if (GIsEditor && !HasAnyFlags(RF_ClassDefaultObject) && !IsInAsyncLoadingThread())
+	{
+		OnNavAreaRegisteredDelegateHandle = UNavigationSystemBase::OnNavAreaRegisteredDelegate().AddUObject(this, &ANavModifierVolume::OnNavAreaRegistered);
+		OnNavAreaUnregisteredDelegateHandle = UNavigationSystemBase::OnNavAreaUnregisteredDelegate().AddUObject(this, &ANavModifierVolume::OnNavAreaUnregistered);
+	}
+#endif // WITH_EDITOR
+}
+
+void ANavModifierVolume::PostLoad()
+{
+	Super::PostLoad();
+
 #if WITH_EDITOR
 	if (GIsEditor && !HasAnyFlags(RF_ClassDefaultObject))
 	{
diff --git a/Engine/Source/Runtime/NavigationSystem/Public/NavModifierVolume.h b/Engine/Source/Runtime/NavigationSystem/Public/NavModifierVolume.h
--- a/Engine/Source/Runtime/NavigationSystem/Public/NavModifierVolume.h
+++ b/Engine/Source/Runtime/NavigationSystem/Public/NavModifierVolume.h
@@ -69,6 +69,7 @@ public:
 #endif
 
 protected:
+	NAVIGATIONSYSTEM_API virtual void PostLoad() override;
 	NAVIGATIONSYSTEM_API virtual void PostInitProperties() override;
 	NAVIGATIONSYSTEM_API virtual void BeginDestroy() override;
 
