[AIBlueprintHelper] Updated CreateMoveToProxyObject to make sure the controller is still a valid Actor that can be stored in a weak pointer.
This early validation is a tentative fix for low occurence crashes where the AAIController pointer is considered valid by IsValid() (which validates for non-null ptr and without Garbage flag)
but then TWeakObjectPtr<>::Get would return null (potentially unreachable) when trying to access the ReceiveMoveCompleted delegate.

#Synced-CL 50152522

diff --git a/Engine/Source/Runtime/AIModule/Private/Blueprint/AIBlueprintHelperLibrary.cpp b/Engine/Source/Runtime/AIModule/Private/Blueprint/AIBlueprintHelperLibrary.cpp
--- a/Engine/Source/Runtime/AIModule/Private/Blueprint/AIBlueprintHelperLibrary.cpp
+++ b/Engine/Source/Runtime/AIModule/Private/Blueprint/AIBlueprintHelperLibrary.cpp
@@ -148,9 +148,16 @@ UAIAsyncTaskBlueprintProxy* UAIBlueprintHelperLibrary::CreateMoveToProxyObject(U
 	}
 
 	UAIAsyncTaskBlueprintProxy* MyObj = nullptr;
-	AAIController* AIController = Cast<AAIController>(Pawn->GetController());
-	if (AIController != nullptr
-		&& ensureMsgf(IsValid(AIController), TEXT("Unusable AIController pointer fetched from Pawn (i.e., garbage)")))
+
+	// Make sure the controller is still a valid Actor that can be stored in a weak pointer.
+	// This early validation is a tentative fix for low occurence crashes where the AAIController
+	// pointer is considered valid by IsValid() (which validates for non-null ptr and without Garbage flag)
+	// but then TWeakObjectPtr<>::Get would return null (potentially unreachable) when trying to access
+	// the ReceiveMoveCompleted delegate below.
+	// We don't want to perform the MoveTo at all in that case.
+	const TWeakObjectPtr WeakAIController = Cast<AAIController>(Pawn->GetController());
+	const TStrongObjectPtr<AAIController> AIController = WeakAIController.Pin();
+	if (AIController != nullptr)
 	{
 		UWorld* World = GEngine->GetWorldFromContextObjectChecked(WorldContextObject);
 		MyObj = NewObject<UAIAsyncTaskBlueprintProxy>(World);
@@ -170,11 +177,11 @@ UAIAsyncTaskBlueprintProxy* UAIBlueprintHelperLibrary::CreateMoveToProxyObject(U
 		MoveReq.SetNavigationFilter(AIController->GetDefaultNavigationFilterClass());
 		MoveReq.SetCanStrafe(AIController->bAllowStrafe);
 
-		FPathFollowingRequestResult ResultData = AIController->MoveTo(MoveReq);
+		const FPathFollowingRequestResult ResultData = AIController->MoveTo(MoveReq);
 		switch (ResultData.Code)
 		{
 		case EPathFollowingRequestResult::RequestSuccessful:
-			MyObj->AIController = AIController;
+			MyObj->AIController = WeakAIController;
 			MyObj->AIController->ReceiveMoveCompleted.AddDynamic(MyObj, &UAIAsyncTaskBlueprintProxy::OnMoveCompleted);
 			MyObj->MoveRequestId = ResultData.MoveId;
 			break;
