[EQS] nav data override in EQS:
- move override logic from ProjectedPoint to the utils method
- replace Pawn filtering by INavAgentInterface
- add NavDataOverrideContext in Pathfinding EQS test

#Synced-CL 51884340

diff --git a/Engine/Source/Runtime/AIModule/Classes/EnvironmentQuery/EnvQueryTypes.h b/Engine/Source/Runtime/AIModule/Classes/EnvironmentQuery/EnvQueryTypes.h
--- a/Engine/Source/Runtime/AIModule/Classes/EnvironmentQuery/EnvQueryTypes.h
+++ b/Engine/Source/Runtime/AIModule/Classes/EnvironmentQuery/EnvQueryTypes.h
@@ -1326,7 +1326,7 @@ public:
 
 namespace FEQSHelpers
 {
-	AIMODULE_API const ANavigationData* FindNavigationDataForQuery(FEnvQueryInstance& QueryInstance);
+	AIMODULE_API const ANavigationData* FindNavigationDataForQuery(FEnvQueryInstance& QueryInstance, TSubclassOf<UEnvQueryContext> NavDataOverrideContext = nullptr);
 }
 
 USTRUCT(BlueprintType)
diff --git a/Engine/Source/Runtime/AIModule/Classes/EnvironmentQuery/Tests/EnvQueryTest_Pathfinding.h b/Engine/Source/Runtime/AIModule/Classes/EnvironmentQuery/Tests/EnvQueryTest_Pathfinding.h
--- a/Engine/Source/Runtime/AIModule/Classes/EnvironmentQuery/Tests/EnvQueryTest_Pathfinding.h
+++ b/Engine/Source/Runtime/AIModule/Classes/EnvironmentQuery/Tests/EnvQueryTest_Pathfinding.h
@@ -51,6 +51,10 @@ class UEnvQueryTest_Pathfinding : public UEnvQueryTest
 	UPROPERTY(EditDefaultsOnly, Category=Pathfinding)
 	TSubclassOf<UNavigationQueryFilter> FilterClass;
 
+	/** Navigation Data Override */
+	UPROPERTY(EditDefaultsOnly, Category = "Generator")
+	TSubclassOf<UEnvQueryContext> NavDataOverrideContext;
+
 	AIMODULE_API virtual void RunTest(FEnvQueryInstance& QueryInstance) const override;
 
 	AIMODULE_API virtual FText GetDescriptionTitle() const override;
@@ -74,6 +78,7 @@ protected:
 	AIMODULE_API float FindPathLengthFrom(const FVector& ItemPos, const FVector& ContextPos, EPathFindingMode::Type Mode, const ANavigationData& NavData, UNavigationSystemV1& NavSys, FSharedConstNavQueryFilter NavFilter, const UObject* PathOwner) const;
 	AIMODULE_API float FindPathLengthTo(const FVector& ItemPos, const FVector& ContextPos, EPathFindingMode::Type Mode, const ANavigationData& NavData, UNavigationSystemV1& NavSys, FSharedConstNavQueryFilter NavFilter, const UObject* PathOwner) const;
 
+	UE_DEPRECATED(5.8, "Please use FEQSHelpers::FindNavigationDataForQuery instead")
 	AIMODULE_API ANavigationData* FindNavigationData(UNavigationSystemV1& NavSys, UObject* Owner) const;
 
 	AIMODULE_API virtual TSubclassOf<UNavigationQueryFilter> GetNavFilterClass(FEnvQueryInstance& QueryInstance) const;
diff --git a/Engine/Source/Runtime/AIModule/Private/EnvironmentQuery/EnvQueryTypes.cpp b/Engine/Source/Runtime/AIModule/Private/EnvironmentQuery/EnvQueryTypes.cpp
--- a/Engine/Source/Runtime/AIModule/Private/EnvironmentQuery/EnvQueryTypes.cpp
+++ b/Engine/Source/Runtime/AIModule/Private/EnvironmentQuery/EnvQueryTypes.cpp
@@ -288,7 +288,7 @@ void FEnvTraceData::OnPostLoad()
 //----------------------------------------------------------------------//
 namespace FEQSHelpers
 {
-	const ANavigationData* FindNavigationDataForQuery(FEnvQueryInstance& QueryInstance)
+	const ANavigationData* FindNavigationDataForQuery(FEnvQueryInstance& QueryInstance, TSubclassOf<UEnvQueryContext> NavDataOverrideContext)
 	{
 		const UNavigationSystemV1* NavSys = FNavigationSystem::GetCurrent<UNavigationSystemV1>(QueryInstance.World);
 		if (NavSys == nullptr)
@@ -296,6 +296,37 @@ namespace FEQSHelpers
 			return nullptr;
 		}
 
+		if (NavDataOverrideContext)
+		{
+			TArray<AActor*> NavDataCandidateActors;
+			if (QueryInstance.PrepareContext(NavDataOverrideContext, NavDataCandidateActors))
+			{
+				for (const AActor* NavDataCandidateActor : NavDataCandidateActors)
+				{
+					if (const ANavigationData* NavDataCandidate = Cast<ANavigationData>(NavDataCandidateActor))
+					{
+						// NavigationData found, we can stop here
+						return NavDataCandidate;
+					}
+					if (NavSys)
+					{
+						// Check if this actor use a NavData that could be used
+						if (const INavAgentInterface* NavAgent = Cast<INavAgentInterface>(QueryInstance.Owner.Get()))
+						{
+							const FNavAgentProperties& NavAgentProps = NavAgent->GetNavAgentPropertiesRef();
+							if (NavAgentProps.IsValid() || NavAgentProps.PreferredNavData.IsValid())	// test that to avoid GetNavDataForProps warning
+							{
+								if (const ANavigationData* AgentNavData = NavSys->GetNavDataForProps(NavAgent->GetNavAgentPropertiesRef()))
+								{
+									return AgentNavData;
+								}
+							}
+						}
+					}
+				}
+			}
+		}
+
 		// try to match navigation agent for querier
 		if (const INavAgentInterface* NavAgent = QueryInstance.Owner.IsValid() ? Cast<INavAgentInterface>(QueryInstance.Owner.Get()) : nullptr)
 		{
diff --git a/Engine/Source/Runtime/AIModule/Private/EnvironmentQuery/Generators/EnvQueryGenerator_ProjectedPoints.cpp b/Engine/Source/Runtime/AIModule/Private/EnvironmentQuery/Generators/EnvQueryGenerator_ProjectedPoints.cpp
--- a/Engine/Source/Runtime/AIModule/Private/EnvironmentQuery/Generators/EnvQueryGenerator_ProjectedPoints.cpp
+++ b/Engine/Source/Runtime/AIModule/Private/EnvironmentQuery/Generators/EnvQueryGenerator_ProjectedPoints.cpp
@@ -25,38 +25,7 @@ void UEnvQueryGenerator_ProjectedPoints::ProjectAndFilterNavPoints(TArray<FNavLo
 	const ANavigationData* NavData = nullptr;
 	if (ProjectionData.TraceMode != EEnvQueryTrace::None)
 	{
-		if (NavDataOverrideContext)
-		{
-			TArray<AActor*> NavDataCandidateActors;
-			if (QueryInstance.PrepareContext(NavDataOverrideContext, NavDataCandidateActors))
-			{
-				const UNavigationSystemV1* NavigationSystem = FNavigationSystem::GetCurrent<UNavigationSystemV1>(GetWorld());
-
-				for (const AActor* NavDataCandidateActor : NavDataCandidateActors)
-				{
-					if (const ANavigationData* NavDataCandidate = Cast<ANavigationData>(NavDataCandidateActor))
-					{
-						// NavigationData found, we can stop here
-						NavData = NavDataCandidate;
-						break;
-					}
-					else if (const APawn* PawnCandidate = Cast<APawn>(NavDataCandidateActor))
-					{
-						// Check if this actor use a NavData that could be used
-						if (NavigationSystem)
-						{
-							NavData = NavigationSystem->GetNavDataForProps(PawnCandidate->GetNavAgentPropertiesRef());
-							break;
-						}
-					}
-				}
-			}
-		}
-
-		if (!NavData)
-		{
-			NavData = FEQSHelpers::FindNavigationDataForQuery(QueryInstance);
-		}
+		NavData = FEQSHelpers::FindNavigationDataForQuery(QueryInstance, NavDataOverrideContext);
 	}
 
 	const UObject* Querier = QueryInstance.Owner.Get();
diff --git a/Engine/Source/Runtime/AIModule/Private/EnvironmentQuery/Tests/EnvQueryTest_Pathfinding.cpp b/Engine/Source/Runtime/AIModule/Private/EnvironmentQuery/Tests/EnvQueryTest_Pathfinding.cpp
--- a/Engine/Source/Runtime/AIModule/Private/EnvironmentQuery/Tests/EnvQueryTest_Pathfinding.cpp
+++ b/Engine/Source/Runtime/AIModule/Private/EnvironmentQuery/Tests/EnvQueryTest_Pathfinding.cpp
@@ -48,7 +48,7 @@ void UEnvQueryTest_Pathfinding::RunTest(FEnvQueryInstance& QueryInstance) const
 		return;
 	}
 
-	ANavigationData* NavData = FindNavigationData(*NavSys, QueryOwner);
+	const ANavigationData* NavData = FEQSHelpers::FindNavigationDataForQuery(QueryInstance, NavDataOverrideContext);
 	if (NavData == nullptr)
 	{
 		return;
@@ -231,6 +231,7 @@ float UEnvQueryTest_Pathfinding::FindPathLengthTo(const FVector& ItemPos, const
 	return (Result.IsSuccessful()) ? static_cast<float>(Result.Path->GetLength()) : BIG_NUMBER;
 }
 
+// Deprecated
 ANavigationData* UEnvQueryTest_Pathfinding::FindNavigationData(UNavigationSystemV1& NavSys, UObject* Owner) const
 {
 	INavAgentInterface* NavAgent = Cast<INavAgentInterface>(Owner);
diff --git a/Engine/Source/Runtime/AIModule/Private/EnvironmentQuery/Tests/EnvQueryTest_PathfindingBatch.cpp b/Engine/Source/Runtime/AIModule/Private/EnvironmentQuery/Tests/EnvQueryTest_PathfindingBatch.cpp
--- a/Engine/Source/Runtime/AIModule/Private/EnvironmentQuery/Tests/EnvQueryTest_PathfindingBatch.cpp
+++ b/Engine/Source/Runtime/AIModule/Private/EnvironmentQuery/Tests/EnvQueryTest_PathfindingBatch.cpp
@@ -79,8 +79,8 @@ void UEnvQueryTest_PathfindingBatch::RunTest(FEnvQueryInstance& QueryInstance) c
 		return;
 	}
 
-	ANavigationData* NavData = FindNavigationData(*NavSys, QueryOwner);
-	ARecastNavMesh* NavMeshData = Cast<ARecastNavMesh>(NavData);
+	const ANavigationData* NavData = FEQSHelpers::FindNavigationDataForQuery(QueryInstance, NavDataOverrideContext);
+	const ARecastNavMesh* NavMeshData = Cast<ARecastNavMesh>(NavData);
 	if (NavMeshData == nullptr)
 	{
 		return;
