[EQS] fixed some macros to fix some compilation when UE_DEBUG_RECORDING_USING_VLOG is set to 0

#Synced-CL 51945655

diff --git a/Engine/Source/Runtime/AIModule/Classes/EnvironmentQuery/EnvQueryDebugHelpers.h b/Engine/Source/Runtime/AIModule/Classes/EnvironmentQuery/EnvQueryDebugHelpers.h
--- a/Engine/Source/Runtime/AIModule/Classes/EnvironmentQuery/EnvQueryDebugHelpers.h
+++ b/Engine/Source/Runtime/AIModule/Classes/EnvironmentQuery/EnvQueryDebugHelpers.h
@@ -2,7 +2,9 @@
 
 #pragma once
 
+#if UE_ENABLE_INCLUDE_ORDER_DEPRECATED_IN_5_8
 #include "CoreMinimal.h"
+#endif // UE_ENABLE_INCLUDE_ORDER_DEPRECATED_IN_5_8
 #include "UObject/ObjectMacros.h"
 #include "UObject/Object.h"
 #include "EngineDefines.h"
@@ -15,7 +17,7 @@
 class FTestData;
 class UEnvQueryDebugHelpers;
 
-#if  USE_EQS_DEBUGGER || ENABLE_VISUAL_LOG
+#if USE_EQS_DEBUGGER || UE_DEBUG_RECORDING_USING_VLOG
 
 struct FLogCategoryBase;
 
@@ -146,13 +148,13 @@ FArchive& operator<<(FArchive& Ar, EQSDebug::FQueryData& Data)
 	return Ar;
 }
 
-#endif //USE_EQS_DEBUGGER || ENABLE_VISUAL_LOG
+#endif // USE_EQS_DEBUGGER || UE_DEBUG_RECORDING_USING_VLOG
 
-#if ENABLE_VISUAL_LOG && USE_EQS_DEBUGGER
+#if USE_EQS_DEBUGGER && UE_DEBUG_RECORDING_USING_VLOG
 #	define UE_VLOG_EQS(Query, Category, Verbosity)  UEnvQueryDebugHelpers::LogQuery(Query, Category, ELogVerbosity::Verbosity);
 #else
 #	define UE_VLOG_EQS(Query, CategoryName, Verbosity)
-#endif //ENABLE_VISUAL_LOG && USE_EQS_DEBUGGER
+#endif // USE_EQS_DEBUGGER && UE_DEBUG_RECORDING_USING_VLOG
 
 UCLASS(Abstract, MinimalAPI)
 class UEnvQueryDebugHelpers : public UObject
@@ -166,19 +168,19 @@ public:
 	static AIMODULE_API void BlobArrayToDebugData(const TArray<uint8>& BlobArray, EQSDebug::FQueryData& EQSLocalData, bool bUseCompression = false);
 #endif
 
-#if ENABLE_VISUAL_LOG && USE_EQS_DEBUGGER
+#if USE_EQS_DEBUGGER && UE_DEBUG_RECORDING_USING_VLOG
 	static void LogQuery(FEnvQueryInstance& Query, const FLogCategoryBase& Category, ELogVerbosity::Type Verbosity);
 	static void LogQuery(FEnvQueryInstance& Query, const FName& CategoryName, ELogVerbosity::Type Verbosity);
 
 private:
 	static void LogQueryInternal(FEnvQueryInstance& Query, const FName& CategoryName, ELogVerbosity::Type Verbosity, double TimeSeconds, FVisualLogEntry *CurrentEntry);
-#endif
+#endif // USE_EQS_DEBUGGER && UE_DEBUG_RECORDING_USING_VLOG
 };
 
-#if ENABLE_VISUAL_LOG && USE_EQS_DEBUGGER
+#if USE_EQS_DEBUGGER && UE_DEBUG_RECORDING_USING_VLOG
 inline void UEnvQueryDebugHelpers::LogQuery(FEnvQueryInstance& Query, const FLogCategoryBase& Category, ELogVerbosity::Type Verbosity)
 {
-	const FName CategoryName = Category.GetCategoryName(); 
+	const FName CategoryName = Category.GetCategoryName();
 	LogQuery(Query, CategoryName, Verbosity);
 }
 
@@ -193,5 +195,4 @@ inline void UEnvQueryDebugHelpers::LogQuery(FEnvQueryInstance& Query, const FNam
 
 	LogQueryInternal(Query, CategoryName, Verbosity, World->TimeSeconds, CurrentEntry);
 }
-#endif
-
+#endif // USE_EQS_DEBUGGER && UE_DEBUG_RECORDING_USING_VLOG
diff --git a/Engine/Source/Runtime/AIModule/Private/EnvironmentQuery/EnvQueryDebugHelpers.cpp b/Engine/Source/Runtime/AIModule/Private/EnvironmentQuery/EnvQueryDebugHelpers.cpp
--- a/Engine/Source/Runtime/AIModule/Private/EnvironmentQuery/EnvQueryDebugHelpers.cpp
+++ b/Engine/Source/Runtime/AIModule/Private/EnvironmentQuery/EnvQueryDebugHelpers.cpp
@@ -138,9 +138,9 @@ void  UEnvQueryDebugHelpers::BlobArrayToDebugData(const TArray<uint8>& BlobArray
 	}
 }
 
-#endif //ENABLE_VISUAL_LOG
+#endif // USE_EQS_DEBUGGER
 
-#if ENABLE_VISUAL_LOG && USE_EQS_DEBUGGER
+#if USE_EQS_DEBUGGER && UE_DEBUG_RECORDING_USING_VLOG
 
 #define PRINT_TABLE_ROW(__src__, __dst__, __len__) \
 	{ \
@@ -238,5 +238,4 @@ void UEnvQueryDebugHelpers::LogQueryInternal(FEnvQueryInstance& Query, const FNa
 }
 
 #undef PRINT_TABLE_ROW
-#endif //ENABLE_VISUAL_LOG && USE_EQS_DEBUGGER
-
+#endif // USE_EQS_DEBUGGER && UE_DEBUG_RECORDING_USING_VLOG
