Make BehaviorTreeComponent GetNodeMemory virtual
#Synced-CL 51943277

diff --git a/Engine/Source/Runtime/AIModule/Classes/BehaviorTree/BehaviorTreeComponent.h b/Engine/Source/Runtime/AIModule/Classes/BehaviorTree/BehaviorTreeComponent.h
--- a/Engine/Source/Runtime/AIModule/Classes/BehaviorTree/BehaviorTreeComponent.h
+++ b/Engine/Source/Runtime/AIModule/Classes/BehaviorTree/BehaviorTreeComponent.h
@@ -230,7 +230,7 @@ public:
 	uint16 GetActiveInstanceIdx() const;
 
 	/** @return node memory */
-	AIMODULE_API uint8* GetNodeMemory(const UBTNode* Node, int32 InstanceIdx) const;
+	AIMODULE_API virtual uint8* GetNodeMemory(const UBTNode* Node, int32 InstanceIdx) const;
 
 	/** @return true if ExecutionRequest is switching to higher priority node */
 	bool IsRestartPending() const;
