[AI][BT] Add AIModuleToolset plugin for BehaviorTree inspection

New editor-only plugin exposing BehaviorTree structure to AI agents via MCP ToolsetRegistry. Provides 7 tools for traversing BT nodes, depths, children, decorators, services, blackboard, and subtrees. Combined with ObjectTools.get_properties, enables full read-only BT inspection without opening the BT editor.

Changes:
- New AIModuleToolset plugin (editor-only, follows StateTreeToolset pattern)
- BehaviorTreeTools: 7 tools (list_nodes, get_node_depth/s, get_children, get_subtree, get_blackboard, get_root_decorators)
- C++: Added BlueprintReadOnly on BT properties across 5 AIModule headers
- 11 tests: 4 empty BT + 7 populated BT with value validation

#Synced-CL 51880525
#Synced-CL 51880525

diff --git a/Engine/Plugins/Experimental/AIModuleToolset/AIModuleToolset.uplugin b/Engine/Plugins/Experimental/AIModuleToolset/AIModuleToolset.uplugin
new file mode 100644
--- /dev/null
+++ b/Engine/Plugins/Experimental/AIModuleToolset/AIModuleToolset.uplugin
@@ -0,0 +1,34 @@
+{
+	"FileVersion": 3,
+	"Version": 1,
+	"VersionName": "1.0",
+	"FriendlyName": "AIModuleToolset",
+	"Description": "Toolset for AIModule Systems",
+	"Category": "AI",
+	"EditorOnly": true,
+	"CreatedBy": "Epic Games",
+	"CreatedByURL": "",
+	"DocsURL": "",
+	"MarketplaceURL": "",
+	"SupportURL": "",
+	"CanContainContent": true,
+	"IsBetaVersion": false,
+	"IsExperimentalVersion": true,
+	"NoRedist": true,
+	"EnabledByDefault": true,
+	"Installed": false,
+	"Modules": [
+		{
+			"Name": "AIModuleToolset",
+			"Type": "Editor",
+			"LoadingPhase": "Default",
+			"TargetAllowList": ["Editor"]
+		}
+	],
+	"Plugins": [
+		{
+			"Name": "ToolsetRegistry",
+			"Enabled": true
+		}
+	]
+}
diff --git a/Engine/Plugins/Experimental/AIModuleToolset/Source/AIModuleToolset/AIModuleToolset.Build.cs b/Engine/Plugins/Experimental/AIModuleToolset/Source/AIModuleToolset/AIModuleToolset.Build.cs
new file mode 100644
--- /dev/null
+++ b/Engine/Plugins/Experimental/AIModuleToolset/Source/AIModuleToolset/AIModuleToolset.Build.cs
@@ -0,0 +1,12 @@
+// Copyright Epic Games, Inc. All Rights Reserved.
+
+using UnrealBuildTool;
+
+public class AIModuleToolset : ModuleRules
+{
+	public AIModuleToolset(ReadOnlyTargetRules Target) : base(Target)
+	{
+		PCHUsage = PCHUsageMode.UseExplicitOrSharedPCHs;
+		PublicDependencyModuleNames.Add("Core");
+	}
+}
diff --git a/Engine/Plugins/Experimental/AIModuleToolset/Source/AIModuleToolset/Private/AIModuleToolset.cpp b/Engine/Plugins/Experimental/AIModuleToolset/Source/AIModuleToolset/Private/AIModuleToolset.cpp
new file mode 100644
--- /dev/null
+++ b/Engine/Plugins/Experimental/AIModuleToolset/Source/AIModuleToolset/Private/AIModuleToolset.cpp
@@ -0,0 +1,16 @@
+// Copyright Epic Games, Inc. All Rights Reserved.
+
+#include "AIModuleToolset.h"
+
+#include "Modules/ModuleManager.h"
+
+/*virtual*/ void FAIModuleToolsetModule::StartupModule()
+{
+}
+
+/*virtual*/ void FAIModuleToolsetModule::ShutdownModule()
+{
+}
+
+IMPLEMENT_MODULE(FAIModuleToolsetModule, AIModuleToolset);
+DEFINE_LOG_CATEGORY(LogAIModuleToolset);
diff --git a/Engine/Plugins/Experimental/AIModuleToolset/Source/AIModuleToolset/Public/AIModuleToolset.h b/Engine/Plugins/Experimental/AIModuleToolset/Source/AIModuleToolset/Public/AIModuleToolset.h
new file mode 100644
--- /dev/null
+++ b/Engine/Plugins/Experimental/AIModuleToolset/Source/AIModuleToolset/Public/AIModuleToolset.h
@@ -0,0 +1,15 @@
+// Copyright Epic Games, Inc. All Rights Reserved.
+
+#pragma once
+
+#include "Modules/ModuleInterface.h"
+#include "Logging/LogMacros.h"
+
+DECLARE_LOG_CATEGORY_EXTERN(LogAIModuleToolset, Log, Log);
+
+class FAIModuleToolsetModule : public IModuleInterface
+{
+public:
+	virtual void StartupModule() override;
+	virtual void ShutdownModule() override;
+};
diff --git a/Engine/Source/Runtime/AIModule/Classes/BehaviorTree/BTCompositeNode.h b/Engine/Source/Runtime/AIModule/Classes/BehaviorTree/BTCompositeNode.h
--- a/Engine/Source/Runtime/AIModule/Classes/BehaviorTree/BTCompositeNode.h
+++ b/Engine/Source/Runtime/AIModule/Classes/BehaviorTree/BTCompositeNode.h
@@ -47,7 +47,7 @@ namespace EBTDecoratorLogic
 	};
 }
 
-USTRUCT()
+USTRUCT(BlueprintType)
 struct FBTDecoratorLogic
 {
 	GENERATED_USTRUCT_BODY()
@@ -62,24 +62,24 @@ struct FBTDecoratorLogic
 	FBTDecoratorLogic(uint8 InOperation, uint16 InNumber) : Operation(InOperation), Number(InNumber) {}
 };
 
-USTRUCT()
+USTRUCT(BlueprintType)
 struct FBTCompositeChild
 {
 	GENERATED_USTRUCT_BODY()
 
 	/** child node */
-	UPROPERTY()
+	UPROPERTY(BlueprintReadOnly, Category = BehaviorTree)
 	TObjectPtr<UBTCompositeNode> ChildComposite = nullptr;
 
-	UPROPERTY()
+	UPROPERTY(BlueprintReadOnly, Category = BehaviorTree)
 	TObjectPtr<UBTTaskNode> ChildTask = nullptr;
 
 	/** execution decorators */
-	UPROPERTY()
+	UPROPERTY(BlueprintReadOnly, Category = BehaviorTree)
 	TArray<TObjectPtr<UBTDecorator>> Decorators;
 
 	/** logic operations for decorators */
-	UPROPERTY()
+	UPROPERTY(BlueprintReadOnly, Category = BehaviorTree)
 	TArray<FBTDecoratorLogic> DecoratorOps;
 };
 
@@ -89,11 +89,11 @@ class UBTCompositeNode : public UBTNode
 	GENERATED_UCLASS_BODY()
 
 	/** child nodes */
-	UPROPERTY()
+	UPROPERTY(BlueprintReadOnly, Category = BehaviorTree)
 	TArray<FBTCompositeChild> Children;
 
 	/** service nodes */
-	UPROPERTY()
+	UPROPERTY(BlueprintReadOnly, Category = BehaviorTree)
 	TArray<TObjectPtr<UBTService>> Services;
 
 	AIMODULE_API ~UBTCompositeNode();
diff --git a/Engine/Source/Runtime/AIModule/Classes/BehaviorTree/BTNode.h b/Engine/Source/Runtime/AIModule/Classes/BehaviorTree/BTNode.h
--- a/Engine/Source/Runtime/AIModule/Classes/BehaviorTree/BTNode.h
+++ b/Engine/Source/Runtime/AIModule/Classes/BehaviorTree/BTNode.h
@@ -214,7 +214,7 @@ class UBTNode : public UObject, public IGameplayTaskOwnerInterface
 	}
 
 	/** node name */
-	UPROPERTY(Category=Description, EditAnywhere)
+	UPROPERTY(Category=Description, EditAnywhere, BlueprintReadOnly)
 	FString NodeName;
 	
 private:
diff --git a/Engine/Source/Runtime/AIModule/Classes/BehaviorTree/BehaviorTree.h b/Engine/Source/Runtime/AIModule/Classes/BehaviorTree/BehaviorTree.h
--- a/Engine/Source/Runtime/AIModule/Classes/BehaviorTree/BehaviorTree.h
+++ b/Engine/Source/Runtime/AIModule/Classes/BehaviorTree/BehaviorTree.h
@@ -17,7 +17,7 @@ class UBehaviorTree : public UObject, public IBlackboardAssetProvider
 	GENERATED_UCLASS_BODY()
 
 	/** root node of loaded tree */
-	UPROPERTY()
+	UPROPERTY(BlueprintReadOnly, Category = BehaviorTree)
 	TObjectPtr<UBTCompositeNode> RootNode;
 
 #if WITH_EDITORONLY_DATA
@@ -38,11 +38,11 @@ class UBehaviorTree : public UObject, public IBlackboardAssetProvider
 	// END IBlackboardAssetProvider
 
 	/** blackboard asset for this tree */
-	UPROPERTY()
+	UPROPERTY(BlueprintReadOnly, Category = BehaviorTree)
 	TObjectPtr<UBlackboardData> BlackboardAsset;
 
 	/** root level decorators, used by subtrees */
-	UPROPERTY()
+	UPROPERTY(BlueprintReadOnly, Category = BehaviorTree)
 	TArray<TObjectPtr<UBTDecorator>> RootDecorators;
 
 	/** logic operators for root level decorators, used by subtrees  */
diff --git a/Engine/Source/Runtime/AIModule/Classes/BehaviorTree/Tasks/BTTask_RunBehavior.h b/Engine/Source/Runtime/AIModule/Classes/BehaviorTree/Tasks/BTTask_RunBehavior.h
--- a/Engine/Source/Runtime/AIModule/Classes/BehaviorTree/Tasks/BTTask_RunBehavior.h
+++ b/Engine/Source/Runtime/AIModule/Classes/BehaviorTree/Tasks/BTTask_RunBehavior.h
@@ -41,7 +41,7 @@ class UBTTask_RunBehavior : public UBTTaskNode
 protected:
 
 	/** behavior to run */
-	UPROPERTY(Category = Node, EditAnywhere)
+	UPROPERTY(Category = Node, EditAnywhere, BlueprintReadOnly)
 	TObjectPtr<UBehaviorTree> BehaviorAsset;
 
 	/** called when subtree is removed from active stack */
diff --git a/Engine/Source/Runtime/AIModule/Classes/BehaviorTree/Tasks/BTTask_RunBehaviorDynamic.h b/Engine/Source/Runtime/AIModule/Classes/BehaviorTree/Tasks/BTTask_RunBehaviorDynamic.h
--- a/Engine/Source/Runtime/AIModule/Classes/BehaviorTree/Tasks/BTTask_RunBehaviorDynamic.h
+++ b/Engine/Source/Runtime/AIModule/Classes/BehaviorTree/Tasks/BTTask_RunBehaviorDynamic.h
@@ -50,11 +50,11 @@ protected:
 	FGameplayTag InjectionTag;
 
 	/** default behavior to run */
-	UPROPERTY(Category=Node, EditAnywhere)
+	UPROPERTY(Category=Node, EditAnywhere, BlueprintReadOnly)
 	TObjectPtr<UBehaviorTree> DefaultBehaviorAsset;
 
 	/** current subtree */
-	UPROPERTY()
+	UPROPERTY(BlueprintReadOnly, Category = Node)
 	TObjectPtr<UBehaviorTree> BehaviorAsset;
 
 	/** called when subtree is removed from active stack */
