Removed virtual from GetNodeMemory and added a Aux version that gets called when GetNodeMemory fails to retrieve non nullptr NodeMemory

#Synced-CL 51943277

diff --git a/Engine/Source/Runtime/AIModule/Classes/BehaviorTree/BehaviorTreeComponent.h b/Engine/Source/Runtime/AIModule/Classes/BehaviorTree/BehaviorTreeComponent.h
--- a/Engine/Source/Runtime/AIModule/Classes/BehaviorTree/BehaviorTreeComponent.h
+++ b/Engine/Source/Runtime/AIModule/Classes/BehaviorTree/BehaviorTreeComponent.h
@@ -230,7 +230,10 @@ public:
 	uint16 GetActiveInstanceIdx() const;
 
 	/** @return node memory */
-	AIMODULE_API virtual uint8* GetNodeMemory(const UBTNode* Node, int32 InstanceIdx) const;
+	AIMODULE_API uint8* GetNodeMemory(const UBTNode* Node, int32 InstanceIdx) const;
+
+	/** @return node memory. Virtual override for when GetNodeMemory fails to find NodeMemory */
+	AIMODULE_API virtual uint8* GetNodeMemoryAux(const UBTNode* Node, int32 InstanceIdx) const;
 
 	/** @return true if ExecutionRequest is switching to higher priority node */
 	bool IsRestartPending() const;
diff --git a/Engine/Source/Runtime/AIModule/Private/BehaviorTree/BehaviorTreeComponent.cpp b/Engine/Source/Runtime/AIModule/Private/BehaviorTree/BehaviorTreeComponent.cpp
--- a/Engine/Source/Runtime/AIModule/Private/BehaviorTree/BehaviorTreeComponent.cpp
+++ b/Engine/Source/Runtime/AIModule/Private/BehaviorTree/BehaviorTreeComponent.cpp
@@ -2914,7 +2914,13 @@ UBTNode* UBehaviorTreeComponent::FindTemplateNode(const UBTNode* Node) const
 
 uint8* UBehaviorTreeComponent::GetNodeMemory(const UBTNode* Node, int32 InstanceIdx) const
 {
-	return InstanceStack.IsValidIndex(InstanceIdx) ? (uint8*)Node->GetNodeMemory<uint8>(InstanceStack[InstanceIdx]) : NULL;
+	uint8* NodeMemory = InstanceStack.IsValidIndex(InstanceIdx) ? (uint8*)Node->GetNodeMemory<uint8>(InstanceStack[InstanceIdx]) : nullptr;
+	return NodeMemory ? NodeMemory : GetNodeMemoryAux(Node, InstanceIdx);
+}
+
+uint8* UBehaviorTreeComponent::GetNodeMemoryAux(const UBTNode* Node, int32 InstanceIdx) const
+{
+	return nullptr;
 }
 
 void UBehaviorTreeComponent::RemoveAllInstances()
