[Mass] Add bStripFromClientBuilds option to UEnvQuery, mirroring UDataTable pattern. Allows EQS assets to be excluded from client builds since they are server-only AI queries.

#Synced-CL 52096882

diff --git a/Engine/Source/Runtime/AIModule/Classes/EnvironmentQuery/EnvQuery.h b/Engine/Source/Runtime/AIModule/Classes/EnvironmentQuery/EnvQuery.h
--- a/Engine/Source/Runtime/AIModule/Classes/EnvironmentQuery/EnvQuery.h
+++ b/Engine/Source/Runtime/AIModule/Classes/EnvironmentQuery/EnvQuery.h
@@ -36,6 +36,10 @@ protected:
 	TArray<TObjectPtr<UEnvQueryOption>> Options;
 
 public:
+	/** Set to true to not cook this query into client builds. Useful for server-only AI queries. */
+	UPROPERTY(EditAnywhere, Category=Query)
+	uint8 bStripFromClientBuilds : 1 = false;
+
 	/** Gather all required named params */
 	AIMODULE_API void CollectQueryParams(UObject& QueryOwner, TArray<FAIDynamicParam>& NamedValues) const;
 
@@ -48,6 +52,15 @@ public:
 	AIMODULE_API virtual void PostDuplicate(bool bDuplicateForPIE) override;
 #endif // WITH_EDITOR
 
+	virtual bool NeedsLoadForClient() const override
+	{
+		return !bStripFromClientBuilds && Super::NeedsLoadForClient();
+	}
+	virtual bool NeedsLoadForEditorGame() const override
+	{
+		return !bStripFromClientBuilds && Super::NeedsLoadForEditorGame();
+	}
+
 	FName GetQueryName() const { return QueryName; }
 
 	TArray<TObjectPtr<UEnvQueryOption>>& GetOptionsMutable() { return Options; }
