FValueOrBlackboardKeyBase::ToString is not tool only

#rnx

#Synced-CL 52204375

diff --git a/Engine/Plugins/Experimental/GameplayBehaviors/Source/GameplayBehaviorsModule/Public/AI/ValueOrBBKey_GameplayTag.h b/Engine/Plugins/Experimental/GameplayBehaviors/Source/GameplayBehaviorsModule/Public/AI/ValueOrBBKey_GameplayTag.h
--- a/Engine/Plugins/Experimental/GameplayBehaviors/Source/GameplayBehaviorsModule/Public/AI/ValueOrBBKey_GameplayTag.h
+++ b/Engine/Plugins/Experimental/GameplayBehaviors/Source/GameplayBehaviorsModule/Public/AI/ValueOrBBKey_GameplayTag.h
@@ -25,7 +25,7 @@ struct FValueOrBBKey_GameplayTagContainer : public FValueOrBlackboardKeyBase
 	GAMEPLAYBEHAVIORSMODULE_API virtual bool IsCompatibleType(const UBlackboardKeyType* KeyType) const override;
 #endif // WITH_EDITOR
 
-	GAMEPLAYBEHAVIORSMODULE_API FString ToString() const;
+	GAMEPLAYBEHAVIORSMODULE_API virtual FString ToString() const override;
 
 protected:
 	UPROPERTY(EditAnywhere, Category = "Value")
diff --git a/Engine/Source/Runtime/AIModule/Classes/BehaviorTree/ValueOrBBKey.h b/Engine/Source/Runtime/AIModule/Classes/BehaviorTree/ValueOrBBKey.h
--- a/Engine/Source/Runtime/AIModule/Classes/BehaviorTree/ValueOrBBKey.h
+++ b/Engine/Source/Runtime/AIModule/Classes/BehaviorTree/ValueOrBBKey.h
@@ -84,14 +84,15 @@ struct FValueOrBlackboardKeyBase
 	GENERATED_BODY()
 
 	friend FValueOrBBKeyDetails;
+	virtual ~FValueOrBlackboardKeyBase() = default;
 
 #if WITH_EDITOR
-	virtual ~FValueOrBlackboardKeyBase() = default;
 	virtual bool IsCompatibleType(const UBlackboardKeyType* KeyType) const { return false; };
-	virtual FString ToString() const { return FString(); }
 	AIMODULE_API void PreSave(const UObject* Outer, const UBlackboardData& Blackboard, FName PropertyName);
 #endif // WITH_EDITOR
 
+	virtual FString ToString() const { return FString(); }
+
 	const FName& GetKey() const { return Key; }
 	void SetKey(FName NewKey) { Key = NewKey; }
 
@@ -137,7 +138,7 @@ struct FValueOrBBKey_Bool : public FValueOrBlackboardKeyBase
 
 	bool SerializeFromMismatchedTag(const FPropertyTag& Tag, FStructuredArchive::FSlot Slot);
 
-	FString ToString() const { return ToStringInternal(DefaultValue); }
+	FString ToString() const override { return ToStringInternal(DefaultValue); }
 
 	UE_DEPRECATED_FORGAME(5.5, "Implicit conversion will be removed next version. Call GetValue instead")
 	operator bool() const { return DefaultValue; }
@@ -195,7 +196,7 @@ struct FValueOrBBKey_Class : public FValueOrBlackboardKeyBase
 	AIMODULE_API virtual bool IsCompatibleType(const UBlackboardKeyType* KeyType) const override;
 #endif // WITH_EDITOR
 
-	AIMODULE_API FString ToString() const;
+	AIMODULE_API FString ToString() const override;
 
 	UE_DEPRECATED_FORGAME(5.5, "Implicit conversion will be removed next version. Call GetValue instead")
 	operator UClass*() const { return DefaultValue; }
@@ -258,7 +259,7 @@ struct FValueOrBBKey_Enum : public FValueOrBlackboardKeyBase
 	AIMODULE_API virtual bool IsCompatibleType(const UBlackboardKeyType* KeyType) const override;
 #endif // WITH_EDITOR
 
-	AIMODULE_API FString ToString() const;
+	AIMODULE_API FString ToString() const override;
 
 	UE_DEPRECATED_FORGAME(5.5, "Implicit conversion will be removed next version. Call GetValue instead")
 	operator uint8() const { return DefaultValue; }
@@ -295,7 +296,7 @@ struct FValueOrBBKey_Float : public FValueOrBlackboardKeyBase
 	AIMODULE_API virtual bool IsCompatibleType(const UBlackboardKeyType* KeyType) const override;
 #endif // WITH_EDITOR
 
-	AIMODULE_API FString ToString() const;
+	AIMODULE_API FString ToString() const override;
 	AIMODULE_API bool IsBoundOrNonZero() const;
 
 	UE_DEPRECATED_FORGAME(5.5, "Implicit conversion will be removed next version. Call GetValue instead")
@@ -323,7 +324,7 @@ struct FValueOrBBKey_Int32 : public FValueOrBlackboardKeyBase
 	AIMODULE_API virtual bool IsCompatibleType(const UBlackboardKeyType* KeyType) const override;
 #endif // WITH_EDITOR
 
-	FString ToString() const { return ToStringInternal(DefaultValue); }
+	FString ToString() const override { return ToStringInternal(DefaultValue); }
 
 	UE_DEPRECATED_FORGAME(5.5, "Implicit conversion will be removed next version. Call GetValue instead")
 	operator int32() const { return DefaultValue; }
@@ -351,7 +352,7 @@ struct FValueOrBBKey_Name : public FValueOrBlackboardKeyBase
 	AIMODULE_API virtual bool IsCompatibleType(const UBlackboardKeyType* KeyType) const override;
 #endif // WITH_EDITOR
 
-	AIMODULE_API FString ToString() const;
+	AIMODULE_API FString ToString() const override;
 
 	UE_DEPRECATED_FORGAME(5.5, "Implicit conversion will be removed next version. Call GetValue instead")
 	operator FName() const { return DefaultValue; }
@@ -382,7 +383,7 @@ struct FValueOrBBKey_String : public FValueOrBlackboardKeyBase
 	UE_DEPRECATED_FORGAME(5.5, "Implicit conversion will be removed next version. Call GetValue instead")
 	operator FString() const { return DefaultValue; }
 
-	FString ToString() const { return ToStringInternal(DefaultValue); }
+	FString ToString() const override { return ToStringInternal(DefaultValue); }
 
 protected:
 	UPROPERTY(EditAnywhere, Category = "Value")
@@ -436,7 +437,7 @@ struct FValueOrBBKey_Object : public FValueOrBlackboardKeyBase
 	AIMODULE_API virtual bool IsCompatibleType(const UBlackboardKeyType* KeyType) const override;
 #endif // WITH_EDITOR
 
-	AIMODULE_API FString ToString() const;
+	AIMODULE_API FString ToString() const override;
 
 	UE_DEPRECATED_FORGAME(5.5, "Implicit conversion will be removed next version. Call GetValue instead")
 	operator UObject*() const { return DefaultValue; }
@@ -469,7 +470,7 @@ struct FValueOrBBKey_Rotator : public FValueOrBlackboardKeyBase
 	AIMODULE_API virtual bool IsCompatibleType(const UBlackboardKeyType* KeyType) const override;
 #endif // WITH_EDITOR
 
-	AIMODULE_API FString ToString() const;
+	AIMODULE_API FString ToString() const override;
 
 	UE_DEPRECATED_FORGAME(5.5, "Implicit conversion will be removed next version. Call GetValue instead")
 	operator FRotator() const { return DefaultValue; }
@@ -497,7 +498,7 @@ struct FValueOrBBKey_Vector : public FValueOrBlackboardKeyBase
 	AIMODULE_API virtual bool IsCompatibleType(const UBlackboardKeyType* KeyType) const override;
 #endif // WITH_EDITOR
 
-	AIMODULE_API FString ToString() const;
+	AIMODULE_API FString ToString() const override;
 
 	UE_DEPRECATED_FORGAME(5.5, "Implicit conversion will be removed next version. Call GetValue instead")
 	operator FVector() const { return DefaultValue; }
@@ -521,10 +522,10 @@ public:
 	FValueOrBBKey_Struct() = default;
 
 #if WITH_EDITOR
-	AIMODULE_API virtual bool IsCompatibleType(const UBlackboardKeyType* KeyType) const;
+	AIMODULE_API virtual bool IsCompatibleType(const UBlackboardKeyType* KeyType) const override;
 #endif // WITH_EDITOR
 
-	AIMODULE_API FString ToString() const;
+	AIMODULE_API FString ToString() const override;
 
 	template <typename T>
 	const T& GetValue(const UBehaviorTreeComponent& BehaviorComp) const
diff --git a/Engine/Source/Runtime/AIModule/Private/BehaviorTree/BlueprintNodeHelpers.cpp b/Engine/Source/Runtime/AIModule/Private/BehaviorTree/BlueprintNodeHelpers.cpp
--- a/Engine/Source/Runtime/AIModule/Private/BehaviorTree/BlueprintNodeHelpers.cpp
+++ b/Engine/Source/Runtime/AIModule/Private/BehaviorTree/BlueprintNodeHelpers.cpp
@@ -113,12 +113,10 @@ namespace BlueprintNodeHelpers
 			}
 #endif
 		}
-#if WITH_EDITOR
 		else if (StructProp && StructProp->Struct->IsChildOf(FValueOrBlackboardKeyBase::StaticStruct()))
 		{
 			ExportedStringValue = reinterpret_cast<const FValueOrBlackboardKeyBase*>(PropertyAddr)->ToString();
 		}
-#endif // WITH_EDITOR
 		else if (FloatProp)
 		{
 			// special case for floats to remove unnecessary zeros
