[EQS] share spacing method between OnCircle and Donut generator

#tests tested in PIE and editor (EQS pawn test) with donut and circle generator. Both spacing methods with and without arcs

#Synced-CL 52630979

diff --git a/Engine/Config/BaseEngine.ini b/Engine/Config/BaseEngine.ini
--- a/Engine/Config/BaseEngine.ini
+++ b/Engine/Config/BaseEngine.ini
@@ -1431,6 +1431,7 @@ ClassName=/Script/OnlineSubsystemUtils.OnlineEngineInterfaceImpl
 +FunctionRedirects=(OldName="/Script/GameplayTags.GameplayTagAssetInterface.GetOwnedGameplayTags", NewName="BP_GetOwnedGameplayTags")
 +PropertyRedirects=(OldName="GameplayTagAssetInterface.BP_GetOwnedGameplayTags.TagContainer", NewName="ReturnValue")
 +PropertyRedirects=(OldName="/Script/UnrealEd.LevelEditorViewportSettings.TransformWidgetSizeAdjustment",NewName="/Script/UnrealEd.LevelEditorViewportSettings.TransformWidgetSizeAdjustment_DEPRECATED")
++EnumRedirects=(OldName="/Script/AIModule.EPointOnCircleSpacingMethod",NewName="/Script/AIModule.EEnvQueryPointSpacingMethod")
 
 [CoreUObject.Metadata]
 ; Note: UnrealHeaderTool should be rerun after making changes to MetadataRedirects to catch any keys specified in class headers
diff --git a/Engine/Source/Runtime/AIModule/Classes/EnvironmentQuery/Generators/EnvQueryGenerator_Donut.h b/Engine/Source/Runtime/AIModule/Classes/EnvironmentQuery/Generators/EnvQueryGenerator_Donut.h
--- a/Engine/Source/Runtime/AIModule/Classes/EnvironmentQuery/Generators/EnvQueryGenerator_Donut.h
+++ b/Engine/Source/Runtime/AIModule/Classes/EnvironmentQuery/Generators/EnvQueryGenerator_Donut.h
@@ -27,8 +27,16 @@ class UEnvQueryGenerator_Donut : public UEnvQueryGenerator_ProjectedPoints
 	UPROPERTY(EditDefaultsOnly, Category = Generator)
 	FAIDataProviderIntValue NumberOfRings;
 
-	/** number of items to generate for each ring */
+	/** how we are choosing where the points are in the ring */
 	UPROPERTY(EditDefaultsOnly, Category = Generator)
+	EEnvQueryPointSpacingMethod PointOnRingSpacingMethod;
+
+	/** items will be generated on a ring this much apart */
+	UPROPERTY(EditDefaultsOnly, Category = Generator, meta = (EditCondition = "PointOnRingSpacingMethod == EEnvQueryPointSpacingMethod::BySpaceBetween", EditConditionHides))
+	FAIDataProviderFloatValue SpaceBetweenPoints;
+
+	/** number of items to generate for each ring */
+	UPROPERTY(EditDefaultsOnly, Category = Generator, meta = (EditCondition = "PointOnRingSpacingMethod == EEnvQueryPointSpacingMethod::ByNumberOfPoints", EditConditionHides))
 	FAIDataProviderIntValue PointsPerRing;
 
 	/** If you generate items on a piece of circle you define direction of Arc cut here */
diff --git a/Engine/Source/Runtime/AIModule/Classes/EnvironmentQuery/Generators/EnvQueryGenerator_OnCircle.h b/Engine/Source/Runtime/AIModule/Classes/EnvironmentQuery/Generators/EnvQueryGenerator_OnCircle.h
--- a/Engine/Source/Runtime/AIModule/Classes/EnvironmentQuery/Generators/EnvQueryGenerator_OnCircle.h
+++ b/Engine/Source/Runtime/AIModule/Classes/EnvironmentQuery/Generators/EnvQueryGenerator_OnCircle.h
@@ -12,8 +12,7 @@
 
 class AActor;
 
-UENUM()
-enum class EPointOnCircleSpacingMethod :uint8
+enum class UE_DEPRECATED(5.8, "Use EEnvQueryPointSpacingMethod instead") EPointOnCircleSpacingMethod :uint8
 {
 	// Use the SpaceBetween value to determine how far apart points should be.
 	BySpaceBetween,
@@ -32,14 +31,14 @@ class UEnvQueryGenerator_OnCircle : public UEnvQueryGenerator_ProjectedPoints
 
 	/** how we are choosing where the points are in the circle */
 	UPROPERTY(EditDefaultsOnly, Category = Generator)
-	EPointOnCircleSpacingMethod PointOnCircleSpacingMethod;
+	EEnvQueryPointSpacingMethod PointOnCircleSpacingMethod;
 
 	/** items will be generated on a circle this much apart */
-	UPROPERTY(EditDefaultsOnly, Category = Generator, meta = (EditCondition = "PointOnCircleSpacingMethod == EPointOnCircleSpacingMethod::BySpaceBetween", EditConditionHides))
+	UPROPERTY(EditDefaultsOnly, Category = Generator, meta = (EditCondition = "PointOnCircleSpacingMethod == EEnvQueryPointSpacingMethod::BySpaceBetween", EditConditionHides))
 	FAIDataProviderFloatValue SpaceBetween;
 
 	/** this many items will be generated on a circle */
-	UPROPERTY(EditDefaultsOnly, Category = Generator, meta = (EditCondition = "PointOnCircleSpacingMethod == EPointOnCircleSpacingMethod::ByNumberOfPoints", EditConditionHides))
+	UPROPERTY(EditDefaultsOnly, Category = Generator, meta = (EditCondition = "PointOnCircleSpacingMethod == EEnvQueryPointSpacingMethod::ByNumberOfPoints", EditConditionHides))
 	FAIDataProviderIntValue NumberOfPoints;
 
 	/** If you generate items on a piece of circle you define direction of Arc cut here */
diff --git a/Engine/Source/Runtime/AIModule/Classes/EnvironmentQuery/Generators/EnvQueryGenerator_ProjectedPoints.h b/Engine/Source/Runtime/AIModule/Classes/EnvironmentQuery/Generators/EnvQueryGenerator_ProjectedPoints.h
--- a/Engine/Source/Runtime/AIModule/Classes/EnvironmentQuery/Generators/EnvQueryGenerator_ProjectedPoints.h
+++ b/Engine/Source/Runtime/AIModule/Classes/EnvironmentQuery/Generators/EnvQueryGenerator_ProjectedPoints.h
@@ -13,6 +13,15 @@
 
 class EnvQueryContext;
 
+UENUM()
+enum class EEnvQueryPointSpacingMethod :uint8
+{
+	// Use the SpaceBetween value to determine how far apart points should be.
+	BySpaceBetween,
+	// Use a fixed number of points
+	ByNumberOfPoints
+};
+
 UCLASS(Abstract, MinimalAPI)
 class UEnvQueryGenerator_ProjectedPoints : public UEnvQueryGenerator
 {
diff --git a/Engine/Source/Runtime/AIModule/Private/EnvironmentQuery/Generators/EnvQueryGenerator_Donut.cpp b/Engine/Source/Runtime/AIModule/Private/EnvironmentQuery/Generators/EnvQueryGenerator_Donut.cpp
--- a/Engine/Source/Runtime/AIModule/Private/EnvironmentQuery/Generators/EnvQueryGenerator_Donut.cpp
+++ b/Engine/Source/Runtime/AIModule/Private/EnvironmentQuery/Generators/EnvQueryGenerator_Donut.cpp
@@ -12,6 +12,8 @@ UEnvQueryGenerator_Donut::UEnvQueryGenerator_Donut(const FObjectInitializer& Obj
 	Center = UEnvQueryContext_Querier::StaticClass();
 	InnerRadius.DefaultValue = 300.0f;
 	OuterRadius.DefaultValue = 1000.0f;
+	PointOnRingSpacingMethod = EEnvQueryPointSpacingMethod::ByNumberOfPoints;
+	SpaceBetweenPoints.DefaultValue = 50.f;
 	NumberOfRings.DefaultValue = 3;
 	PointsPerRing.DefaultValue = 8;
 	ArcDirection.DirMode = EEnvDirection::TwoPoints;
@@ -94,72 +96,126 @@ bool UEnvQueryGenerator_Donut::GenerateDonutNavPoints(const TArray<FVector>& Cen
 	OuterRadius.BindData(BindOwner, QueryInstance.QueryID);
 	NumberOfRings.BindData(BindOwner, QueryInstance.QueryID);
 	PointsPerRing.BindData(BindOwner, QueryInstance.QueryID);
+	SpaceBetweenPoints.BindData(BindOwner, QueryInstance.QueryID);
 	ArcAngle.BindData(BindOwner, QueryInstance.QueryID);
 
-	FVector::FReal ArcAngleValue = ArcAngle.GetValue();
-	FVector::FReal InnerRadiusValue = InnerRadius.GetValue();
-	FVector::FReal OuterRadiusValue = OuterRadius.GetValue();
-	int32 NumRings = NumberOfRings.GetValue();
-	int32 NumPoints = PointsPerRing.GetValue();
+	const FVector::FReal ArcAngleValue = ArcAngle.GetValue();
+	const FVector::FReal InnerRadiusValue = InnerRadius.GetValue();
+	const FVector::FReal OuterRadiusValue = OuterRadius.GetValue();
+	const int32 NumRings = NumberOfRings.GetValue();
 
 	if ((InnerRadiusValue < 0.) || (OuterRadiusValue <= 0.) ||
-		(InnerRadiusValue > OuterRadiusValue) ||
-		(NumRings < 1) || (NumPoints < 1))
+		(InnerRadiusValue > OuterRadiusValue) || (NumRings < 1))
 	{
 		return false;
 	}
 
 	const FVector::FReal ArcBisectDeg = GetArcBisectorAngle(QueryInstance);
 	const FVector::FReal ArcAngleDeg = FMath::Clamp(ArcAngleValue, 0., 360.);
-
 	const FVector::FReal RadiusDelta = NumRings > 1 ? (OuterRadiusValue - InnerRadiusValue) / (NumRings - 1) : 0.;
-	const FVector::FReal AngleDelta = 2. * UE_DOUBLE_PI / NumPoints;
-	FVector::FReal SectionAngle = FMath::DegreesToRadians(ArcBisectDeg);
-
-	Points.Reserve(NumPoints * NumRings * CenterPoints.Num());
 
-	if (!bUseSpiralPattern)
+	// Handle the two spacing methods differently
+	if (PointOnRingSpacingMethod == EEnvQueryPointSpacingMethod::ByNumberOfPoints)
 	{
-		for (int32 SectionIdx = 0; SectionIdx < NumPoints; SectionIdx++, SectionAngle += AngleDelta)
+		// ByNumberOfPoints: All rings have the same point count
+		const int32 NumPoints = PointsPerRing.GetValue();
+		if (NumPoints < 1)
+		{
+			return false;
+		}
+
+		const FVector::FReal AngleDelta = 2. * UE_DOUBLE_PI / NumPoints;
+		FVector::FReal SectionAngle = FMath::DegreesToRadians(ArcBisectDeg);
+
+		Points.Reserve(NumPoints * NumRings * CenterPoints.Num());
+
+		if (!bUseSpiralPattern)
 		{
-			if (IsAngleAllowed(SectionAngle, ArcBisectDeg, ArcAngleDeg, bDefineArc))
+			for (int32 SectionIdx = 0; SectionIdx < NumPoints; SectionIdx++, SectionAngle += AngleDelta)
 			{
-				const FVector::FReal SinValue = FMath::Sin(SectionAngle);
-				const FVector::FReal CosValue = FMath::Cos(SectionAngle);
+				if (IsAngleAllowed(SectionAngle, ArcBisectDeg, ArcAngleDeg, bDefineArc))
+				{
+					const FVector::FReal SinValue = FMath::Sin(SectionAngle);
+					const FVector::FReal CosValue = FMath::Cos(SectionAngle);
 
-				FVector::FReal RingRadius = InnerRadiusValue;
-				for (int32 RingIdx = 0; RingIdx < NumRings; RingIdx++, RingRadius += RadiusDelta)
+					FVector::FReal RingRadius = InnerRadiusValue;
+					for (int32 RingIdx = 0; RingIdx < NumRings; RingIdx++, RingRadius += RadiusDelta)
+					{
+						const FVector RingPos(RingRadius * CosValue, RingRadius * SinValue, 0.);
+						for (int32 ContextIdx = 0; ContextIdx < CenterPoints.Num(); ContextIdx++)
+						{
+							const FNavLocation PointPos = FNavLocation(CenterPoints[ContextIdx] + RingPos);
+							Points.Add(PointPos);
+						}
+					}
+				}
+			}
+		}
+		else
+		{	// In order to spiral, we need to change the angle for each ring as well as each spoke.  By changing it as a
+			// fraction of the SectionAngle, we guarantee that the offset won't cross to the starting point of the next
+			// spoke.
+			const FVector::FReal RingAngleDelta = AngleDelta / NumRings;
+			FVector::FReal RingRadius = InnerRadiusValue;
+			FVector::FReal CurrentRingAngleDelta = 0.f;
+
+			// So, start with ring angle and then add section angle.
+			for (int32 RingIdx = 0; RingIdx < NumRings; RingIdx++, RingRadius += RadiusDelta, CurrentRingAngleDelta += RingAngleDelta)
+			{
+				for (int32 SectionIdx = 0; SectionIdx < NumPoints; SectionIdx++, SectionAngle += AngleDelta)
 				{
-					const FVector RingPos(RingRadius * CosValue, RingRadius * SinValue, 0.);
-					for (int32 ContextIdx = 0; ContextIdx < CenterPoints.Num(); ContextIdx++)
+					FVector::FReal RingSectionAngle = CurrentRingAngleDelta + SectionAngle;
+					if (IsAngleAllowed(RingSectionAngle, ArcBisectDeg, ArcAngleDeg, bDefineArc))
 					{
-						const FNavLocation PointPos = FNavLocation(CenterPoints[ContextIdx] + RingPos);
-						Points.Add(PointPos);
+						const FVector::FReal SinValue = FMath::Sin(RingSectionAngle);
+						const FVector::FReal CosValue = FMath::Cos(RingSectionAngle);
+
+						const FVector RingPos(RingRadius * CosValue, RingRadius * SinValue, 0.);
+						for (int32 ContextIdx = 0; ContextIdx < CenterPoints.Num(); ContextIdx++)
+						{
+							const FNavLocation PointPos = FNavLocation(CenterPoints[ContextIdx] + RingPos);
+							Points.Add(PointPos);
+						}
 					}
 				}
 			}
 		}
 	}
-	else
-	{	// In order to spiral, we need to change the angle for each ring as well as each spoke.  By changing it as a
-		// fraction of the SectionAngle, we guarantee that the offset won't cross to the starting point of the next
-		// spoke.
-		const FVector::FReal RingAngleDelta = AngleDelta / NumRings;
+	else // BySpaceBetween
+	{
+		// BySpaceBetween: Each ring has different point count based on its circumference
+		const FVector::FReal SpaceBetweenPointsValue = SpaceBetweenPoints.GetValue();
+		if (SpaceBetweenPointsValue <= 0.0)
+		{
+			return false;
+		}
+
 		FVector::FReal RingRadius = InnerRadiusValue;
-		FVector::FReal CurrentRingAngleDelta = 0.f;
 
-		// So, start with ring angle and then add section angle.
-		for (int32 RingIdx = 0; RingIdx < NumRings; RingIdx++, RingRadius += RadiusDelta, CurrentRingAngleDelta += RingAngleDelta)
+		// Reserve approximate space (use outer ring's count as upper bound)
+		const FVector::FReal OuterArcLength = 2.0 * UE_DOUBLE_PI * OuterRadiusValue * ArcAngleDeg / 360.0;
+		const int32 ApproxMaxPoints = FMath::CeilToInt(OuterArcLength / SpaceBetweenPointsValue) + 1;
+		Points.Reserve(ApproxMaxPoints * NumRings * CenterPoints.Num());
+
+		FVector::FReal SectionAngle = FMath::DegreesToRadians(ArcBisectDeg);
+		// Generate points for each ring independently
+		for (int32 RingIdx = 0; RingIdx < NumRings; RingIdx++, RingRadius += RadiusDelta)
 		{
-			for (int32 SectionIdx = 0; SectionIdx < NumPoints; SectionIdx++, SectionAngle += AngleDelta)
+			// Calculate point count for this specific ring
+			const FVector::FReal RingLength = 2.0 * UE_DOUBLE_PI * RingRadius;
+			const int32 RingPointCount = FMath::CeilToInt(RingLength / SpaceBetweenPointsValue) + 1;
+
+			// Generate points evenly distributed around this ring
+			const FVector::FReal RingAngleDelta = 2.0 * UE_DOUBLE_PI / RingPointCount;
+
+			for (int32 PointIdx = 0; PointIdx < RingPointCount; PointIdx++, SectionAngle += RingAngleDelta)
 			{
-				FVector::FReal RingSectionAngle = CurrentRingAngleDelta + SectionAngle;
-				if (IsAngleAllowed(RingSectionAngle, ArcBisectDeg, ArcAngleDeg, bDefineArc))
+				if (IsAngleAllowed(SectionAngle, ArcBisectDeg, ArcAngleDeg, bDefineArc))
 				{
-					const FVector::FReal SinValue = FMath::Sin(RingSectionAngle);
-					const FVector::FReal CosValue = FMath::Cos(RingSectionAngle);
-
+					const FVector::FReal SinValue = FMath::Sin(SectionAngle);
+					const FVector::FReal CosValue = FMath::Cos(SectionAngle);
 					const FVector RingPos(RingRadius * CosValue, RingRadius * SinValue, 0.);
+
 					for (int32 ContextIdx = 0; ContextIdx < CenterPoints.Num(); ContextIdx++)
 					{
 						const FNavLocation PointPos = FNavLocation(CenterPoints[ContextIdx] + RingPos);
@@ -167,6 +223,9 @@ bool UEnvQueryGenerator_Donut::GenerateDonutNavPoints(const TArray<FVector>& Cen
 					}
 				}
 			}
+
+			// for spiral pattern we keep the section angle from one ring to the other. Else we just reset it
+			SectionAngle = bUseSpiralPattern ? (SectionAngle + RingAngleDelta) : FMath::DegreesToRadians(ArcBisectDeg);
 		}
 	}
 
@@ -188,9 +247,18 @@ FText UEnvQueryGenerator_Donut::GetDescriptionDetails() const
 	Args.Add(TEXT("InnerRadius"), FText::FromString(InnerRadius.ToString()));
 	Args.Add(TEXT("OuterRadius"), FText::FromString(OuterRadius.ToString()));
 	Args.Add(TEXT("NumRings"), FText::FromString(NumberOfRings.ToString()));
-	Args.Add(TEXT("NumPerRing"), FText::FromString(PointsPerRing.ToString()));
 
-	FText Desc = FText::Format(LOCTEXT("DonutDescription", "radius: {InnerRadius} to {OuterRadius}\nrings: {NumRings}, points per ring: {NumPerRing}"), Args);
+	FText Desc;
+	if (PointOnRingSpacingMethod == EEnvQueryPointSpacingMethod::ByNumberOfPoints)
+	{
+		Args.Add(TEXT("NumPerRing"), FText::FromString(PointsPerRing.ToString()));
+		Desc = FText::Format(LOCTEXT("DonutDescriptionByCount", "radius: {InnerRadius} to {OuterRadius}\nrings: {NumRings}, points per ring: {NumPerRing}"), Args);
+	}
+	else
+	{
+		Args.Add(TEXT("Spacing"), FText::FromString(SpaceBetweenPoints.ToString()));
+		Desc = FText::Format(LOCTEXT("DonutDescriptionBySpacing", "radius: {InnerRadius} to {OuterRadius}\nrings: {NumRings}, spacing: {Spacing} (points vary per ring)"), Args);
+	}
 
 	if (bDefineArc)
 	{
@@ -233,6 +301,10 @@ void UEnvQueryGenerator_Donut::PostEditChangeProperty(FPropertyChangedEvent& Pro
 		{
 			PointsPerRing.DefaultValue = FMath::Max(1, PointsPerRing.DefaultValue);
 		}
+		else if (PropName == GET_MEMBER_NAME_CHECKED(UEnvQueryGenerator_Donut, SpaceBetweenPoints))
+		{
+			SpaceBetweenPoints.DefaultValue = FMath::Max(1.0f, SpaceBetweenPoints.DefaultValue);
+		}
 	}
 
 	Super::PostEditChangeProperty(PropertyChangedEvent);
diff --git a/Engine/Source/Runtime/AIModule/Private/EnvironmentQuery/Generators/EnvQueryGenerator_OnCircle.cpp b/Engine/Source/Runtime/AIModule/Private/EnvironmentQuery/Generators/EnvQueryGenerator_OnCircle.cpp
--- a/Engine/Source/Runtime/AIModule/Private/EnvironmentQuery/Generators/EnvQueryGenerator_OnCircle.cpp
+++ b/Engine/Source/Runtime/AIModule/Private/EnvironmentQuery/Generators/EnvQueryGenerator_OnCircle.cpp
@@ -24,7 +24,7 @@ UEnvQueryGenerator_OnCircle::UEnvQueryGenerator_OnCircle(const FObjectInitialize
 	CircleRadius.DefaultValue = 1000.0f;
 	SpaceBetween.DefaultValue = 50.0f;
 	NumberOfPoints.DefaultValue = 8;
-	PointOnCircleSpacingMethod = EPointOnCircleSpacingMethod::BySpaceBetween;
+	PointOnCircleSpacingMethod = EEnvQueryPointSpacingMethod::BySpaceBetween;
 	ArcDirection.DirMode = EEnvDirection::TwoPoints;
 	ArcDirection.LineFrom = UEnvQueryContext_Querier::StaticClass();
 	ArcDirection.Rotation = UEnvQueryContext_Querier::StaticClass();
@@ -147,8 +147,8 @@ void UEnvQueryGenerator_OnCircle::GenerateItems(FEnvQueryInstance& QueryInstance
 	if ((AngleDegree <= 0.f) || 
 		(AngleDegree > 360.f) || 
 		(RadiusValue <= 0.f) || 
-		((PointOnCircleSpacingMethod == EPointOnCircleSpacingMethod::BySpaceBetween) && (ItemSpace <= 0.f)) ||
-		((PointOnCircleSpacingMethod == EPointOnCircleSpacingMethod::ByNumberOfPoints) && (NumPoints <= 0)))
+		((PointOnCircleSpacingMethod == EEnvQueryPointSpacingMethod::BySpaceBetween) && (ItemSpace <= 0.f)) ||
+		((PointOnCircleSpacingMethod == EEnvQueryPointSpacingMethod::ByNumberOfPoints) && (NumPoints <= 0)))
 	{
 		return;
 	}
@@ -162,14 +162,14 @@ void UEnvQueryGenerator_OnCircle::GenerateItems(FEnvQueryInstance& QueryInstance
 	int32 NumStepsBetween;
 	switch (PointOnCircleSpacingMethod)
 	{
-		case EPointOnCircleSpacingMethod::BySpaceBetween:
+		case EEnvQueryPointSpacingMethod::BySpaceBetween:
 		{
 			StepsCount = FMath::CeilToInt(ArcLength / ItemSpace) + 1;
 			NumStepsBetween = StepsCount - 1;
 			break;
 		}
 
-		case EPointOnCircleSpacingMethod::ByNumberOfPoints:
+		case EEnvQueryPointSpacingMethod::ByNumberOfPoints:
 		{
 			StepsCount = NumPoints;
 			NumStepsBetween = StepsCount;
@@ -314,7 +314,7 @@ FText UEnvQueryGenerator_OnCircle::GetDescriptionDetails() const
 	Args.Add(TEXT("NumberOfPoints"), FText::FromString(NumberOfPoints.ToString()));
 	Args.Add(TEXT("TraceData"), TraceData.ToText(FEnvTraceData::Detailed));
 
-	FText Desc = PointOnCircleSpacingMethod == EPointOnCircleSpacingMethod::BySpaceBetween 
+	FText Desc = PointOnCircleSpacingMethod == EEnvQueryPointSpacingMethod::BySpaceBetween 
 		? FText::Format(LOCTEXT("OnCircleDescriptionSpan", "radius: {Radius}, item span: {ItemSpacing}\n{TraceData}"), Args)
 		: FText::Format(LOCTEXT("OnCircleDescriptionNumItems", "radius: {Radius}, number of items to generate: {NumberOfPoints}\n{TraceData}"), Args);
 
